RigVM: Localized registry in cooked environments - first steps.

diff --git a/Engine/Plugins/Runtime/RigVM/Source/RigVM/Public/RigVMCore/RigVMRegistry.h b/Engine/Plugins/Runtime/RigVM/Source/RigVM/Public/RigVMCore/RigVMRegistry.h
--- a/Engine/Plugins/Runtime/RigVM/Source/RigVM/Public/RigVMCore/RigVMRegistry.h
+++ b/Engine/Plugins/Runtime/RigVM/Source/RigVM/Public/RigVMCore/RigVMRegistry.h
@@ -7,6 +7,7 @@
 #include "Containers/Map.h"
 #include "Containers/UnrealString.h"
 #include "CoreMinimal.h"
+#include "RigVMRegistryHandle.h"
 #include "RigVMCore/RigVMTraits.h"
 #include "RigVMDispatchFactory.h"
 #include "RigVMFunction.h"
@@ -31,6 +32,8 @@ struct FRigVMDispatchFactory;
 struct FRigVMRegistry_RWLock;
 typedef FRigVMRegistry_RWLock FRigVMRegistry;
 
+extern UE_API TAutoConsoleVariable<bool> CVarRigVMEnableLocalizedRegistry;
+
 /**
  * The FRigVMRegistry is used to manage all known function pointers
  * for use in the RigVM. The Register method is called automatically
@@ -39,7 +42,7 @@ typedef FRigVMRegistry_RWLock FRigVMRegistry;
  * 
  * Inheriting from FGCObject to ensure that all type objects cannot be GCed
  */
-struct FRigVMRegistry_NoLock : public FGCObject
+struct FRigVMRegistry_NoLock : public FGCObject, public TSharedFromThis<FRigVMRegistry_NoLock>
 {
 public:
 
@@ -50,10 +53,15 @@ public:
 		LockType_Invalid
 	};
 
-	// Returns the singleton registry
-	static UE_API FRigVMRegistry_NoLock& Get(ELockType InLockType);
-	static const FRigVMRegistry_NoLock& GetForRead() { return Get(LockType_Read); }
-	static FRigVMRegistry_NoLock& GetForWrite() { return Get(LockType_Write); }
+	// returns true for the singleton global registry
+	UE_API bool IsGlobalRegistry() const { return bIsGlobalRegistry; }
+
+	// returns the handle referring to this hierarchy
+	UE_API const FRigVMRegistryHandle& GetHandle_NoLock() const;
+	UE_API FRigVMRegistryHandle& GetHandle_NoLock();
+
+	// Constructor for a localized, independent registry
+	UE_API static TSharedRef<FRigVMRegistry_NoLock> CreateLocalizedRegistry(); 
 
 	DECLARE_MULTICAST_DELEGATE(FOnRigVMRegistryChanged);
 
@@ -69,7 +77,7 @@ public:
 	UE_API virtual void Register_NoLock(const TCHAR* InName, FRigVMFunctionPtr InFunctionPtr, UScriptStruct* InStruct = nullptr, const TArray<FRigVMFunctionArgument>& InArguments = TArray<FRigVMFunctionArgument>());
 
 	// Registers a dispatch factory given its struct.
-	UE_API virtual const FRigVMDispatchFactory* RegisterFactory_NoLock(UScriptStruct* InFactoryStruct);
+	UE_API virtual const FRigVMDispatchFactory* RegisterFactory_NoLock(UScriptStruct* InFactoryStruct, const TArray<FRigVMTemplateArgumentInfo>& InArgumentInfos = {});
 
 	// Register a predicate contained in the input struct
 	UE_API virtual void RegisterPredicate_NoLock(UScriptStruct* InStruct, const TCHAR* InName, const TArray<FRigVMFunctionArgument>& InArguments);
@@ -120,6 +128,9 @@ public:
 	// Clear the registry
 	UE_API virtual void Reset_NoLock();
 
+	// Serializes the registry to an archive. This throws for the global registry.
+	UE_API void Serialize_NoLock(FArchive& Ar);
+
 	// Adds a type if it doesn't exist yet and returns its index.
 	// This function is thread-safe
 	UE_API virtual TRigVMTypeIndex FindOrAddType_NoLock(const FRigVMTemplateArgumentType& InType, bool bForce = false);
@@ -320,7 +331,7 @@ public:
 protected:
 
 
-	UE_API FRigVMRegistry_NoLock();
+	UE_API FRigVMRegistry_NoLock(bool InIsGlobalRegistry);
 
 	// disable copy constructor
 	FRigVMRegistry_NoLock(const FRigVMRegistry_NoLock&) = delete;
@@ -343,12 +354,13 @@ protected:
 		TRigVMTypeIndex ArrayTypeIndex;
 		bool bIsArray;
 		bool bIsExecute;
-		uint32 Hash;
+		mutable uint32 Hash;
+
+		void Serialize_NoLock(FArchive& Ar);
 	};
 
 	// Initialize the base types
-	UE_API void Initialize_NoLock();
-	virtual void Initialize(bool bLockRegistry) = 0;
+	UE_API void InitializeBaseTypes_NoLock();
 
 	static EObjectFlags DisallowedFlags()
 	{
@@ -370,6 +382,12 @@ protected:
 	UE_API void PropagateTypeAddedToCategory_NoLock(const FRigVMTemplateArgument::ETypeCategory InCategory, const TRigVMTypeIndex InTypeIndex);
 	UE_API void RemoveTypeInCategory_NoLock(FRigVMTemplateArgument::ETypeCategory InCategory, TRigVMTypeIndex InTypeIndex);
 
+	// the context storing this registry
+	mutable FRigVMRegistryHandle Context;
+	
+	// true if this registry is the globally shared registry 
+	bool bIsGlobalRegistry;
+
 	// memory for all (known) types
 	TArray<FTypeInfo> Types;
 	TMap<FRigVMTemplateArgumentType, TRigVMTypeIndex> TypeToIndex;
@@ -445,13 +463,13 @@ struct FRigVMRegistry_RWLock final : public FRigVMRegistry_NoLock
 	
 protected:
 
-	class FConditionalScopeLock
+	class FConditionalScopeLock : public FRigVMRegistryHandleLock
 	{
 	public:
 		UE_NODISCARD_CTOR explicit FConditionalScopeLock(const FRigVMRegistry_RWLock& InRegistry, ELockType InLockType, bool bInLockEnabled = true)
-		: Registry(const_cast<FRigVMRegistry_RWLock*>(&InRegistry))
+		: FRigVMRegistryHandleLock(&InRegistry, bInLockEnabled)
+		, Registry(const_cast<FRigVMRegistry_RWLock*>(&InRegistry))
 		, DesiredLockType(InLockType)
-		, bLockEnabled(bInLockEnabled)
 		{
 			if(bLockEnabled)
 			{
@@ -460,17 +478,13 @@ protected:
 				{
 					Registry->Lock.ReadLock();
 
-					// fetch_add returns the value preceding the modification
-					// so we have to add one manually to get the current value
-					CurrentLockCount = Registry->LockCount.fetch_add(1) + 1;
+					CurrentLockCount = ++Registry->LockCount;
 				}
 				else if(DesiredLockType == LockType_Write)
 				{
 					Registry->Lock.WriteLock();
 
-					// fetch_add returns the value preceding the modification
-					// so we have to add one manually to get the current value
-					CurrentLockCount = Registry->LockCount.fetch_add(1) + 1;
+					CurrentLockCount = ++Registry->LockCount;
 					ensure(CurrentLockCount == 1);
 				}
 
@@ -487,9 +501,7 @@ protected:
 			{
 				if(DesiredLockType == LockType_Read)
 				{
-					// fetch_sub returns the value preceding the modification
-					// so we have to subtract one manually to get the current value
-					const int32 CurrentLockCount = Registry->LockCount.fetch_sub(1) - 1;
+					const int32 CurrentLockCount = --Registry->LockCount;
 					ensure(CurrentLockCount >= 0);
 					if(CurrentLockCount == 0)
 					{
@@ -499,9 +511,7 @@ protected:
 				}
 				else if(DesiredLockType == LockType_Write)
 				{
-					// fetch_sub returns the value preceding the modification
-					// so we have to subtract one manually to get the current value
-					const int32 CurrentLockCount = Registry->LockCount.fetch_sub(1) - 1;
+					const int32 CurrentLockCount = --Registry->LockCount;
 					ensure(CurrentLockCount == 0);
 					Registry->LockType.store(LockType_Invalid);
 					Registry->Lock.WriteUnlock();
@@ -522,7 +532,6 @@ protected:
 	private:
 		FRigVMRegistry_RWLock* Registry;
 		ELockType DesiredLockType;
-		const bool bLockEnabled;
 
 		UE_NONCOPYABLE(FConditionalScopeLock);
 	};
@@ -569,35 +578,35 @@ public:
 	// for example "FMyStruct::MyVirtualMethod"
 	void Register(const TCHAR* InName, FRigVMFunctionPtr InFunctionPtr, UScriptStruct* InStruct = nullptr, const TArray<FRigVMFunctionArgument>& InArguments = TArray<FRigVMFunctionArgument>(), bool bLockRegistry = true)
 	{
-		FConditionalWriteScopeLock _(*this, bLockRegistry);
+		FConditionalWriteScopeLock _(*this);
 		Super::Register_NoLock(InName, InFunctionPtr, InStruct, InArguments);
 	}
 
 	// Registers a dispatch factory given its struct.
-	const FRigVMDispatchFactory* RegisterFactory(UScriptStruct* InFactoryStruct, bool bLockRegistry = true)
+	const FRigVMDispatchFactory* RegisterFactory(UScriptStruct* InFactoryStruct)
 	{
-		FConditionalWriteScopeLock _(*this, bLockRegistry);
+		FConditionalWriteScopeLock _(*this);
 		return Super::RegisterFactory_NoLock(InFactoryStruct);
 	}
 
 	// Register a predicate contained in the input struct
-	void RegisterPredicate(UScriptStruct* InStruct, const TCHAR* InName, const TArray<FRigVMFunctionArgument>& InArguments, bool bLockRegistry = true)
+	void RegisterPredicate(UScriptStruct* InStruct, const TCHAR* InName, const TArray<FRigVMFunctionArgument>& InArguments)
 	{
-		FConditionalWriteScopeLock _(*this, bLockRegistry);
+		FConditionalWriteScopeLock _(*this);
 		return Super::RegisterPredicate_NoLock(InStruct, InName, InArguments);
 	}
 
 	// Register a set of allowed object types
-	void RegisterObjectTypes(TConstArrayView<TPair<UClass*, ERegisterObjectOperation>> InClasses, bool bLockRegistry = true)
+	void RegisterObjectTypes(TConstArrayView<TPair<UClass*, ERegisterObjectOperation>> InClasses)
 	{
-		FConditionalWriteScopeLock _(*this, bLockRegistry);
+		FConditionalWriteScopeLock _(*this);
 		Super::RegisterObjectTypes_NoLock(InClasses);
 	}
 
 	// Register a set of allowed struct types
-	void RegisterStructTypes(TConstArrayView<UScriptStruct*> InStructs, bool bLockRegistry = true)
+	void RegisterStructTypes(TConstArrayView<UScriptStruct*> InStructs)
 	{
-		FConditionalWriteScopeLock _(*this, bLockRegistry);
+		FConditionalWriteScopeLock _(*this);
 		Super::RegisterStructTypes_NoLock(InStructs);
 	}
 
@@ -611,9 +620,9 @@ public:
 
 	// Refreshes the list and finds the function pointers
 	// based on the names.
-	void RefreshEngineTypesIfRequired(bool bLockRegistry = true)
+	void RefreshEngineTypesIfRequired()
 	{
-		FConditionalWriteScopeLock _(*this, bLockRegistry);
+		FConditionalWriteScopeLock _(*this);
 		Super::RefreshEngineTypesIfRequired_NoLock();
 	}
 	
@@ -648,103 +657,103 @@ public:
 	FOnRigVMRegistryChanged& OnRigVMRegistryChanged() { return OnRigVMRegistryChangedDelegate; }
 	
 	// Clear the registry
-	void Reset(bool bLockRegistry = true)
+	void Reset()
 	{
-		FConditionalWriteScopeLock _(*this, bLockRegistry);
+		FConditionalWriteScopeLock _(*this);
 		Super::Reset_NoLock();
 	}
 
 	// Adds a type if it doesn't exist yet and returns its index.
 	// This function is thread-safe
-	TRigVMTypeIndex FindOrAddType(const FRigVMTemplateArgumentType& InType, bool bForce = false, bool bLockRegistry = true)
+	TRigVMTypeIndex FindOrAddType(const FRigVMTemplateArgumentType& InType, bool bForce = false)
 	{
-		FConditionalWriteScopeLock _(*this, bLockRegistry);
+		FConditionalWriteScopeLock _(*this);
 		return Super::FindOrAddType_NoLock(InType, bForce);
 	}
 
 	// Removes a type from the registry, and updates all dependent templates
 	// which also creates invalid permutations in templates that we should ignore
-	bool RemoveType(const FSoftObjectPath& InObjectPath, const UClass* InObjectClass, bool bLockRegistry = true)
+	bool RemoveType(const FSoftObjectPath& InObjectPath, const UClass* InObjectClass)
 	{
-		FConditionalWriteScopeLock _(*this, bLockRegistry);
+		FConditionalWriteScopeLock _(*this);
 		return Super::RemoveType_NoLock(InObjectPath, InObjectClass);
 	}
 
 	// Returns the type index given a type
-	TRigVMTypeIndex GetTypeIndex(const FRigVMTemplateArgumentType& InType, bool bLockRegistry = true) const
+	TRigVMTypeIndex GetTypeIndex(const FRigVMTemplateArgumentType& InType) const
 	{
-		FConditionalReadScopeLock _(*this, bLockRegistry);
+		FConditionalReadScopeLock _(*this);
 		return Super::GetTypeIndex_NoLock(InType);
 	}
 
 	// Returns the type index given a cpp type and a type object
-	TRigVMTypeIndex GetTypeIndex(const FName& InCPPType, UObject* InCPPTypeObject, bool bLockRegistry = true) const
+	TRigVMTypeIndex GetTypeIndex(const FName& InCPPType, UObject* InCPPTypeObject) const
 	{
-		FConditionalReadScopeLock _(*this, bLockRegistry);
+		FConditionalReadScopeLock _(*this);
 		return GetTypeIndex(FRigVMTemplateArgumentType(InCPPType, InCPPTypeObject));
 	}
 
 	// Returns the type index given an enum, struct, or object
 	template <typename T>
-	TRigVMTypeIndex GetTypeIndex(bool bAsArray = false, bool bLockRegistry = true) const
+	TRigVMTypeIndex GetTypeIndex(bool bAsArray = false) const
 	{
-		FConditionalReadScopeLock _(*this, bLockRegistry);
+		FConditionalReadScopeLock _(*this);
 		return Super::GetTypeIndex_NoLock<T>(bAsArray);
 	}
 
 	// Returns the type given its index
-	const FRigVMTemplateArgumentType& GetType(TRigVMTypeIndex InTypeIndex, bool bLockRegistry = true) const
+	const FRigVMTemplateArgumentType& GetType(TRigVMTypeIndex InTypeIndex) const
 	{
-		FConditionalReadScopeLock _(*this, bLockRegistry);
+		FConditionalReadScopeLock _(*this);
 		return Super::GetType_NoLock(InTypeIndex);
 	}
 
 	// Returns the number of types
-	int32 NumTypes(bool bLockRegistry = true) const
+	int32 NumTypes() const
 	{
-		FConditionalReadScopeLock _(*this, bLockRegistry);
+		FConditionalReadScopeLock _(*this);
 		return Super::NumTypes_NoLock();
 	}
 
 	// Returns the type given only its cpp type
-	const FRigVMTemplateArgumentType& FindTypeFromCPPType(const FString& InCPPType, bool bLockRegistry = true) const
+	const FRigVMTemplateArgumentType& FindTypeFromCPPType(const FString& InCPPType) const
 	{
-		FConditionalReadScopeLock _(*this, bLockRegistry);
+		FConditionalReadScopeLock _(*this);
 		return Super::FindTypeFromCPPType_NoLock(InCPPType);
 	}
 
 	// Returns the type index given only its cpp type
-	TRigVMTypeIndex GetTypeIndexFromCPPType(const FString& InCPPType, bool bLockRegistry = true) const
+	TRigVMTypeIndex GetTypeIndexFromCPPType(const FString& InCPPType) const
 	{
-		FConditionalReadScopeLock _(*this, bLockRegistry);
+		FConditionalReadScopeLock _(*this);
 		return Super::GetTypeIndexFromCPPType_NoLock(InCPPType);
 	}
 
 	// Returns true if the type is an array
-	bool IsArrayType(TRigVMTypeIndex InTypeIndex, bool bLockRegistry = true) const
+	bool IsArrayType(TRigVMTypeIndex InTypeIndex) const
 	{
-		FConditionalReadScopeLock _(*this, bLockRegistry);
+		FConditionalReadScopeLock _(*this);
 		return Super::IsArrayType_NoLock(InTypeIndex);
 	}
 
 	// Returns true if the type is an execute type
-	bool IsExecuteType(TRigVMTypeIndex InTypeIndex, bool bLockRegistry = true) const
+	bool IsExecuteType(TRigVMTypeIndex InTypeIndex) const
 	{
-		FConditionalReadScopeLock _(*this, bLockRegistry);
+		FConditionalReadScopeLock _(*this);
 		return Super::IsExecuteType_NoLock(InTypeIndex);
 	} 
 
 	// Converts the given execute context type to the base execute context type
-	bool ConvertExecuteContextToBaseType(TRigVMTypeIndex& InOutTypeIndex, bool bLockRegistry = true) const
+	bool ConvertExecuteContextToBaseType(TRigVMTypeIndex& InOutTypeIndex) const
 	{
-		FConditionalReadScopeLock _(*this, bLockRegistry);
+		FConditionalReadScopeLock _(*this);
 		return Super::ConvertExecuteContextToBaseType_NoLock(InOutTypeIndex);
 	} 
 
 	// Returns the dimensions of the array 
-	int32 GetArrayDimensionsForType(TRigVMTypeIndex InTypeIndex, bool bLockRegistry = true) const
+	int32 GetArrayDimensionsForType(TRigVMTypeIndex InTypeIndex) const
 	{
-		FConditionalReadScopeLock _(*this, bLockRegistry);
+		FConditionalReadScopeLock _(*this);
 		return Super::GetArrayDimensionsForType_NoLock(InTypeIndex);
 	} 
 
@@ -756,9 +765,9 @@ public:
 	} 
 
 	// Returns true if the types can be matched.
-	bool CanMatchTypes(TRigVMTypeIndex InTypeIndexA, TRigVMTypeIndex InTypeIndexB, bool bAllowFloatingPointCasts, bool bLockRegistry = true) const
+	bool CanMatchTypes(TRigVMTypeIndex InTypeIndexA, TRigVMTypeIndex InTypeIndexB, bool bAllowFloatingPointCasts) const
 	{
-		FConditionalReadScopeLock _(*this, bLockRegistry);
+		FConditionalReadScopeLock _(*this);
 		return Super::CanMatchTypes_NoLock(InTypeIndexA, InTypeIndexB, bAllowFloatingPointCasts);
 	}
 
@@ -770,58 +779,58 @@ public:
 	}
 
 	// Returns all compatible types given a category
-	const TArray<TRigVMTypeIndex>& GetTypesForCategory(FRigVMTemplateArgument::ETypeCategory InCategory, bool bLockRegistry = true) const
+	const TArray<TRigVMTypeIndex>& GetTypesForCategory(FRigVMTemplateArgument::ETypeCategory InCategory) const
 	{
-		FConditionalReadScopeLock _(*this, bLockRegistry);
+		FConditionalReadScopeLock _(*this);
 		return Super::GetTypesForCategory_NoLock(InCategory);
 	}
 
 	// Returns the type index of the array matching the given element type index
-	TRigVMTypeIndex GetArrayTypeFromBaseTypeIndex(TRigVMTypeIndex InTypeIndex, bool bLockRegistry = true) const
+	TRigVMTypeIndex GetArrayTypeFromBaseTypeIndex(TRigVMTypeIndex InTypeIndex) const
 	{
-		FConditionalReadScopeLock _(*this, bLockRegistry);
+		FConditionalReadScopeLock _(*this);
 		return Super::GetArrayTypeFromBaseTypeIndex_NoLock(InTypeIndex);
 	}
 
 	// Returns the type index of the element matching the given array type index
-	TRigVMTypeIndex GetBaseTypeFromArrayTypeIndex(TRigVMTypeIndex InTypeIndex, bool bLockRegistry = true) const
+	TRigVMTypeIndex GetBaseTypeFromArrayTypeIndex(TRigVMTypeIndex InTypeIndex) const
 	{
-		FConditionalReadScopeLock _(*this, bLockRegistry);
+		FConditionalReadScopeLock _(*this);
 		return Super::GetBaseTypeFromArrayTypeIndex_NoLock(InTypeIndex);
 	}
 
 	// Returns the function given its name (or nullptr)
-	const FRigVMFunction* FindFunction(const TCHAR* InName, const FRigVMUserDefinedTypeResolver& InTypeResolver = FRigVMUserDefinedTypeResolver(), bool bLockRegistry = true) const
+	const FRigVMFunction* FindFunction(const TCHAR* InName, const FRigVMUserDefinedTypeResolver& InTypeResolver = FRigVMUserDefinedTypeResolver()) const
 	{
-		FConditionalWriteScopeLock _(*this, bLockRegistry);
+		FConditionalWriteScopeLock _(*this);
 		return Super::FindFunction_NoLock(InName, InTypeResolver);
 	}
 
 	// Returns the function given its backing up struct and method name
-	const FRigVMFunction* FindFunction(UScriptStruct* InStruct, const TCHAR* InName, const FRigVMUserDefinedTypeResolver& InResolvalInfo = FRigVMUserDefinedTypeResolver(), bool bLockRegistry = true) const
+	const FRigVMFunction* FindFunction(UScriptStruct* InStruct, const TCHAR* InName, const FRigVMUserDefinedTypeResolver& InResolvalInfo = FRigVMUserDefinedTypeResolver()) const
 	{
-		FConditionalWriteScopeLock _(*this, bLockRegistry);
+		FConditionalWriteScopeLock _(*this);
 		return Super::FindFunction_NoLock(InStruct, InName, InResolvalInfo);
 	}
 
 	// Returns all current RigVM functions
-	const TChunkedArray<FRigVMFunction>& GetFunctions(bool bLockRegistry = true) const
+	const TChunkedArray<FRigVMFunction>& GetFunctions() const
 	{
-		FConditionalReadScopeLock _(*this, bLockRegistry);
+		FConditionalReadScopeLock _(*this);
 		return Super::GetFunctions_NoLock();
 	}
 
 	// Returns a template pointer given its notation (or nullptr)
-	const FRigVMTemplate* FindTemplate(const FName& InNotation, bool bIncludeDeprecated = false, bool bLockRegistry = true) const
+	const FRigVMTemplate* FindTemplate(const FName& InNotation, bool bIncludeDeprecated = false) const
 	{
-		FConditionalWriteScopeLock _(*this, bLockRegistry);
+		FConditionalWriteScopeLock _(*this);
 		return Super::FindTemplate_NoLock(InNotation, bIncludeDeprecated);
 	}
 
 	// Returns all current RigVM functions
-	const TChunkedArray<FRigVMTemplate>& GetTemplates(bool bLockRegistry = true) const
+	const TChunkedArray<FRigVMTemplate>& GetTemplates() const
 	{
-		FConditionalReadScopeLock _(*this, bLockRegistry);
+		FConditionalReadScopeLock _(*this);
 		return Super::GetTemplates_NoLock();
 	}
 
@@ -829,10 +838,9 @@ public:
 	const FRigVMTemplate* GetOrAddTemplateFromArguments(
 		const FName& InName,
 		const TArray<FRigVMTemplateArgumentInfo>& InInfos,
-		const FRigVMTemplateDelegates& InDelegates,
-		bool bLockRegistry = true)
+		const FRigVMTemplateDelegates& InDelegates)
 	{
-		FConditionalWriteScopeLock _(*this, bLockRegistry);
+		FConditionalWriteScopeLock _(*this);
 		return Super::GetOrAddTemplateFromArguments_NoLock(InName, InInfos, InDelegates);
 	}
 
@@ -840,24 +848,23 @@ public:
 	const FRigVMTemplate* AddTemplateFromArguments(
 		const FName& InName,
 		const TArray<FRigVMTemplateArgumentInfo>& InInfos,
-		const FRigVMTemplateDelegates& InDelegates,
-		bool bLockRegistry = true)
+		const FRigVMTemplateDelegates& InDelegates)
 	{
-		FConditionalWriteScopeLock _(*this, bLockRegistry);
+		FConditionalWriteScopeLock _(*this);
 		return Super::AddTemplateFromArguments_NoLock(InName, InInfos, InDelegates);
 	}
 
 	// Returns a dispatch factory given its name (or nullptr)
-	FRigVMDispatchFactory* FindDispatchFactory(const FName& InFactoryName, bool bLockRegistry = true) const
+	FRigVMDispatchFactory* FindDispatchFactory(const FName& InFactoryName) const
 	{
-		FConditionalReadScopeLock _(*this, bLockRegistry);
+		FConditionalReadScopeLock _(*this);
 		return Super::FindDispatchFactory_NoLock(InFactoryName);
 	}
 
 	// Returns a dispatch factory given its static struct (or nullptr)
-	FRigVMDispatchFactory* FindOrAddDispatchFactory(UScriptStruct* InFactoryStruct, bool bLockRegistry = true)
+	FRigVMDispatchFactory* FindOrAddDispatchFactory(UScriptStruct* InFactoryStruct)
 	{
-		FConditionalWriteScopeLock _(*this, bLockRegistry);
+		FConditionalWriteScopeLock _(*this);
 		return Super::FindOrAddDispatchFactory_NoLock(InFactoryStruct);
 	}
 
@@ -866,16 +873,16 @@ public:
 		typename T,
 		typename TEnableIf<TModels_V<CRigVMUStruct, T>>::Type * = nullptr
 	>
-	FRigVMDispatchFactory* FindOrAddDispatchFactory(bool bLockRegistry = true)
+	FRigVMDispatchFactory* FindOrAddDispatchFactory()
 	{
-		FConditionalWriteScopeLock _(*this, bLockRegistry);
+		FConditionalWriteScopeLock _(*this);
 		return Super::FindOrAddDispatchFactory_NoLock<T>();
 	}
 
 	// Returns a dispatch factory's singleton function name if that exists
-	FString FindOrAddSingletonDispatchFunction(UScriptStruct* InFactoryStruct, bool bLockRegistry = true)
+	FString FindOrAddSingletonDispatchFunction(UScriptStruct* InFactoryStruct)
 	{
-		FConditionalWriteScopeLock _(*this, bLockRegistry);
+		FConditionalWriteScopeLock _(*this);
 		return Super::FindOrAddSingletonDispatchFunction_NoLock(InFactoryStruct);
 	}
 
@@ -884,60 +891,60 @@ public:
 		typename T,
 		typename TEnableIf<TModels_V<CRigVMUStruct, T>>::Type * = nullptr
 	>
-	FString FindOrAddSingletonDispatchFunction(bool bLockRegistry = true)
+	FString FindOrAddSingletonDispatchFunction()
 	{
-		FConditionalWriteScopeLock _(*this, bLockRegistry);
+		FConditionalWriteScopeLock _(*this);
 		return Super::FindOrAddSingletonDispatchFunction_NoLock<T>();
 	}
 
 	// Returns all dispatch factories
-	const TArray<FRigVMDispatchFactory*>& GetFactories(bool bLockRegistry = true) const
+	const TArray<FRigVMDispatchFactory*>& GetFactories() const
 	{
-		FConditionalReadScopeLock _(*this, bLockRegistry);
+		FConditionalReadScopeLock _(*this);
 		return Super::GetFactories_NoLock();
 	}
 
 	// Given a struct name, return the predicates
-	const TArray<FRigVMFunction>* GetPredicatesForStruct(const FName& InStructName, bool bLockRegistry = true) const
+	const TArray<FRigVMFunction>* GetPredicatesForStruct(const FName& InStructName) const
 	{
-		FConditionalReadScopeLock _(*this, bLockRegistry);
+		FConditionalReadScopeLock _(*this);
 		return Super::GetPredicatesForStruct_NoLock(InStructName);
 	}
 
 	// Returns a unique hash per type index
-	uint32 GetHashForType(TRigVMTypeIndex InTypeIndex, bool bLockRegistry = true) const
+	uint32 GetHashForType(TRigVMTypeIndex InTypeIndex) const
 	{
-		FConditionalReadScopeLock _(*this, bLockRegistry);
+		FConditionalReadScopeLock _(*this);
 		return Super::GetHashForType_NoLock(InTypeIndex);
 	}
 	 
-	uint32 GetHashForScriptStruct(const UScriptStruct* InScriptStruct, bool bCheckTypeIndex = true, bool bLockRegistry = true) const
+	uint32 GetHashForScriptStruct(const UScriptStruct* InScriptStruct, bool bCheckTypeIndex = true) const
 	{
-		FConditionalReadScopeLock _(*this, bLockRegistry);
+		FConditionalReadScopeLock _(*this);
 		return Super::GetHashForScriptStruct_NoLock(InScriptStruct, bCheckTypeIndex);
 	}
 	 
-	uint32 GetHashForStruct(const UStruct* InStruct, bool bLockRegistry = true) const
+	uint32 GetHashForStruct(const UStruct* InStruct) const
 	{
-		FConditionalReadScopeLock _(*this, bLockRegistry);
+		FConditionalReadScopeLock _(*this);
 		return Super::GetHashForStruct_NoLock(InStruct);
 	}
 	 
-	uint32 GetHashForEnum(const UEnum* InEnum, bool bCheckTypeIndex = true, bool bLockRegistry = true) const
+	uint32 GetHashForEnum(const UEnum* InEnum, bool bCheckTypeIndex = true) const
 	{
-		const FConditionalReadScopeLock _(*this, bLockRegistry);
+		const FConditionalReadScopeLock _(*this);
 		return Super::GetHashForEnum_NoLock(InEnum, bCheckTypeIndex);
 	}
 	 
-	uint32 GetHashForProperty(const FProperty* InProperty, bool bLockRegistry = true) const
+	uint32 GetHashForProperty(const FProperty* InProperty) const
 	{
-		const FConditionalReadScopeLock _(*this, bLockRegistry);
+		const FConditionalReadScopeLock _(*this);
 		return Super::GetHashForProperty_NoLock(InProperty);
 	}
 
-	void RebuildRegistry(bool bLockRegistry = true)
+	void RebuildRegistry()
 	{
-		FConditionalWriteScopeLock _(*this, bLockRegistry);
+		FConditionalWriteScopeLock _(*this);
 		Super::RebuildRegistry_NoLock();
 	}
 
@@ -954,9 +961,9 @@ protected:
 
 	// Registers a dispatch factory given its struct.
 	// Note: Only call NoLock methods on the FRigVMRegistry_NoLock
-	virtual const FRigVMDispatchFactory* RegisterFactory_NoLock(UScriptStruct* InFactoryStruct) override
+	virtual const FRigVMDispatchFactory* RegisterFactory_NoLock(UScriptStruct* InFactoryStruct, const TArray<FRigVMTemplateArgumentInfo>& InArgumentInfos = {}) override
 	{
-		return Super::RegisterFactory_NoLock(InFactoryStruct);
+		return Super::RegisterFactory_NoLock(InFactoryStruct, InArgumentInfos);
 	}
 
 	// Register a predicate contained in the input struct
@@ -1295,13 +1302,13 @@ protected:
 
 private:
 
-	FRigVMRegistry_RWLock();
-
-	// Initialize the base types
-	virtual void Initialize(bool bLockRegistry) override;
+	FRigVMRegistry_RWLock(bool InIsGlobalRegistry);
 
 	static void EnsureLocked(ELockType InLockType);
 
+	UE_API void Initialize();
+	UE_API void Initialize_NoLock();
+
 	mutable FRWLock Lock;
 	mutable std::atomic<ELockType> LockType;
 	mutable std::atomic<int32> LockCount;
@@ -1317,7 +1324,7 @@ private:
 	friend struct FRigVMRegistry_NoLock;
 };
 
-typedef FRigVMRegistry_RWLock::FConditionalReadScopeLock FRigVMRegistryReadLock;
-typedef FRigVMRegistry_RWLock::FConditionalWriteScopeLock FRigVMRegistryWriteLock;
+using FRigVMRegistryReadLock = FRigVMRegistry_RWLock::FConditionalReadScopeLock;
+using FRigVMRegistryWriteLock = FRigVMRegistry_RWLock::FConditionalWriteScopeLock;
 
 #undef UE_API
diff --git a/Engine/Plugins/Runtime/RigVM/Source/RigVM/Public/RigVMCore/RigVMRegistryHandle.h b/Engine/Plugins/Runtime/RigVM/Source/RigVM/Public/RigVMCore/RigVMRegistryHandle.h
new file mode 100644
--- /dev/null
+++ b/Engine/Plugins/Runtime/RigVM/Source/RigVM/Public/RigVMCore/RigVMRegistryHandle.h
@@ -0,0 +1,50 @@
+// Copyright Epic Games, Inc. All Rights Reserved.
+
+#pragma once
+
+#include "CoreMinimal.h"
+
+struct FRigVMRegistry_NoLock;
+
+#define UE_API RIGVM_API
+
+struct FRigVMRegistryHandle
+{
+public:
+
+	UE_API FRigVMRegistryHandle();
+	FRigVMRegistryHandle(FRigVMRegistryHandle& InOther) = delete;
+
+	UE_API bool IsValid() const;
+	
+	UE_API FRigVMRegistry_NoLock* operator->();
+	UE_API const FRigVMRegistry_NoLock* operator->() const;
+
+	UE_API FRigVMRegistry_NoLock* GetRegistry();
+	UE_API const FRigVMRegistry_NoLock* GetRegistry() const;
+
+	FRigVMRegistryHandle& operator=(FRigVMRegistryHandle& InOther) = delete;
+	FRigVMRegistryHandle& operator=(const FRigVMRegistryHandle& InOther) = delete;
+
+protected:
+	UE_API FRigVMRegistryHandle(const FRigVMRegistry_NoLock* InRegistry);
+
+	FRigVMRegistry_NoLock* Registry;
+
+	friend struct FRigVMRegistry_NoLock;
+	friend class URigVM;
+};
+
+struct FRigVMRegistryHandleLock : public FRigVMRegistryHandle
+{
+	UE_API FRigVMRegistryHandleLock(const FRigVMRegistry_NoLock* InRegistry, bool bInLockEnabled)
+		: FRigVMRegistryHandle(InRegistry)
+		, bLockEnabled(bInLockEnabled)
+	{
+	}
+	
+protected:
+	const bool bLockEnabled;
+};
+
+#undef UE_API
diff --git a/Engine/Plugins/Runtime/RigVM/Source/RigVM/Public/RigVMCore/RigVMTemplate.h b/Engine/Plugins/Runtime/RigVM/Source/RigVM/Public/RigVMCore/RigVMTemplate.h
--- a/Engine/Plugins/Runtime/RigVM/Source/RigVM/Public/RigVMCore/RigVMTemplate.h
+++ b/Engine/Plugins/Runtime/RigVM/Source/RigVM/Public/RigVMCore/RigVMTemplate.h
@@ -28,14 +28,14 @@ struct FRigVMDispatchFactory;
 struct FRigVMTemplate;
 struct FRigVMDispatchContext;
 struct FRigVMUserDefinedTypeResolver;
-struct FRigVMRegistry_NoLock;
+struct FRigVMRegistryHandle;
 class FRigVMTypeCacheScope_NoLock;
 
 typedef TMap<FName, TRigVMTypeIndex> FRigVMTemplateTypeMap;
 
 // FRigVMTemplate_NewArgumentTypeDelegate is deprecated, use FRigVMTemplate_GetPermutationsFromArgumentTypeDelegate
-DECLARE_DELEGATE_RetVal_TwoParams(FRigVMTemplateTypeMap, FRigVMTemplate_NewArgumentTypeDelegate, const FName& /* InArgumentName */, TRigVMTypeIndex /* InTypeIndexToAdd */);
-DECLARE_DELEGATE_RetVal(FRigVMDispatchFactory*, FRigVMTemplate_GetDispatchFactoryDelegate);
+DECLARE_DELEGATE_RetVal_ThreeParams(FRigVMTemplateTypeMap, FRigVMTemplate_NewArgumentTypeDelegate, const FName& /* InArgumentName */, TRigVMTypeIndex /* InTypeIndexToAdd */, const FRigVMRegistryHandle& /*InContext*/);
+DECLARE_DELEGATE_RetVal_OneParam(FRigVMDispatchFactory*, FRigVMTemplate_GetDispatchFactoryDelegate, const FRigVMRegistryHandle& /* InContext */);
 
 struct FRigVMTemplateDelegates
 {
@@ -87,6 +87,8 @@ struct FRigVMTemplateArgumentType
 		return !CPPType.IsNone();
 	}
 
+	void Serialize_NoLock(FArchive& Ar);
+
 	static FRigVMTemplateArgumentType Array()
 	{
 		return FRigVMTemplateArgumentType(RigVMTypeUtils::GetWildCardArrayCPPTypeName(), RigVMTypeUtils::GetWildCardCPPTypeObject());
@@ -156,7 +158,7 @@ struct FRigVMTemplateArgumentType
  */
 struct FRigVMTemplateArgument
 {
-	DECLARE_DELEGATE_RetVal_OneParam(bool, FTypeFilter, const TRigVMTypeIndex&);
+	DECLARE_DELEGATE_RetVal_TwoParams(bool, FTypeFilter, const TRigVMTypeIndex&, const FRigVMRegistryHandle&);
 	
 	enum EArrayType
 	{
@@ -195,9 +197,9 @@ struct FRigVMTemplateArgument
 	RIGVM_API FRigVMTemplateArgument();
 
 	RIGVM_API FRigVMTemplateArgument(const FName& InName, ERigVMPinDirection InDirection);
-	RIGVM_API FRigVMTemplateArgument(const FName& InName, ERigVMPinDirection InDirection, TRigVMTypeIndex InType);
-	RIGVM_API FRigVMTemplateArgument(const FName& InName, ERigVMPinDirection InDirection, const TArray<TRigVMTypeIndex>& InTypeIndices);
-	RIGVM_API FRigVMTemplateArgument(const FName& InName, ERigVMPinDirection InDirection, const TArray<ETypeCategory>& InTypeCategories, TFunction<bool(const TRigVMTypeIndex&)> InFilterType = nullptr);
+	RIGVM_API FRigVMTemplateArgument(const FName& InName, ERigVMPinDirection InDirection, TRigVMTypeIndex InType, const FRigVMRegistryHandle& InRegistry);
+	RIGVM_API FRigVMTemplateArgument(const FName& InName, ERigVMPinDirection InDirection, const TArray<TRigVMTypeIndex>& InTypeIndices, const FRigVMRegistryHandle& InRegistry);
+	RIGVM_API FRigVMTemplateArgument(const FName& InName, ERigVMPinDirection InDirection, const TArray<ETypeCategory>& InTypeCategories, TFunction<bool(const TRigVMTypeIndex&, const FRigVMRegistryHandle&)> InFilterType, const FRigVMRegistryHandle& InRegistry);
 
 	// returns the name of the argument
 	const FName& GetName() const { return Name; }
@@ -207,27 +209,27 @@ struct FRigVMTemplateArgument
 
 #if WITH_EDITOR
 	// returns true if this argument supports a given type across a set of permutations
-	RIGVM_API bool SupportsTypeIndex(TRigVMTypeIndex InTypeIndex, TRigVMTypeIndex* OutTypeIndex = nullptr, const bool bLockRegistry = true) const;
-	RIGVM_API bool SupportsTypeIndex_NoLock(TRigVMTypeIndex InTypeIndex, TRigVMTypeIndex* OutTypeIndex = nullptr) const;
+	RIGVM_API bool SupportsTypeIndex(TRigVMTypeIndex InTypeIndex, TRigVMTypeIndex* OutTypeIndex = nullptr) const;
+	RIGVM_API bool SupportsTypeIndex_NoLock(const FRigVMRegistryHandle& InRegistry, TRigVMTypeIndex InTypeIndex, TRigVMTypeIndex* OutTypeIndex = nullptr) const;
 #endif
 
 	// returns the flat list of types (including duplicates) of this argument
-	RIGVM_API void GetAllTypes(TArray<TRigVMTypeIndex>& OutTypes, const bool bLockRegistry = true) const;
-	RIGVM_API void GetAllTypes_NoLock(TArray<TRigVMTypeIndex>& OutTypes) const;
-
-	RIGVM_API TRigVMTypeIndex GetTypeIndex(const int32 InIndex, const bool bLockRegistry = true) const;
-	RIGVM_API TRigVMTypeIndex GetTypeIndex_NoLock(const int32 InIndex) const;
-	RIGVM_API TOptional<TRigVMTypeIndex> TryToGetTypeIndex(const int32 InIndex, const bool bLockRegistry = true) const;
-	RIGVM_API TOptional<TRigVMTypeIndex> TryToGetTypeIndex_NoLock(const int32 InIndex) const;
-	RIGVM_API int32 GetNumTypes() const;
-	RIGVM_API int32 GetNumTypes_NoLock() const;
+	RIGVM_API void GetAllTypes(TArray<TRigVMTypeIndex>& OutTypes) const;
+	RIGVM_API void GetAllTypes_NoLock(TArray<TRigVMTypeIndex>& OutTypes, const FRigVMRegistryHandle& InRegistry) const;
+
+	RIGVM_API TRigVMTypeIndex GetTypeIndex(const int32 InIndex) const;
+	RIGVM_API TRigVMTypeIndex GetTypeIndex_NoLock(const int32 InIndex, const FRigVMRegistryHandle& InRegistry) const;
+	RIGVM_API TOptional<TRigVMTypeIndex> TryToGetTypeIndex(const int32 InIndex) const;
+	RIGVM_API TOptional<TRigVMTypeIndex> TryToGetTypeIndex_NoLock(const int32 InIndex, const FRigVMRegistryHandle& InRegistry) const;
+	RIGVM_API int32 GetNumTypes(const FRigVMRegistryHandle& InRegistry) const;
+	RIGVM_API int32 GetNumTypes_NoLock(const FRigVMRegistryHandle& InRegistry) const;
 	RIGVM_API void AddTypeIndex(const TRigVMTypeIndex InTypeIndex);
 	RIGVM_API void RemoveType(const int32 InIndex);
-	RIGVM_API void ForEachType(TFunction<bool(const TRigVMTypeIndex InType)>&& InCallback) const;
-	RIGVM_API int32 FindTypeIndex(const TRigVMTypeIndex InTypeIndex) const;
+	RIGVM_API void ForEachType(TFunction<bool(const TRigVMTypeIndex InType)>&& InCallback, const FRigVMRegistryHandle& InRegistry) const;
+	RIGVM_API int32 FindTypeIndex(const TRigVMTypeIndex InTypeIndex, const FRigVMRegistryHandle& InRegistry) const;
 
 	template <typename Predicate>
-	int32 IndexOfByPredicate(Predicate Pred) const
+	int32 IndexOfByPredicate(Predicate Pred, const FRigVMRegistryHandle& InRegistry) const
 	{
 		if (!bUseCategories)
 		{
@@ -237,9 +239,9 @@ struct FRigVMTemplateArgument
 		{
 			bool bFound = false;
 			int32 ValidIndex = 0;
-			CategoryViews(TypeCategories).ForEachType([this, &ValidIndex, &bFound, Pred](const TRigVMTypeIndex Type)
+			CategoryViews(TypeCategories, InRegistry).ForEachType([this, &ValidIndex, &bFound, Pred, &InRegistry](const TRigVMTypeIndex Type)
 			{
-				if (FilterType(Type))
+				if (FilterType(Type, InRegistry))
 				{
 					if (Pred(Type))
 					{
@@ -256,33 +258,35 @@ struct FRigVMTemplateArgument
 			}
 			return INDEX_NONE;
 		}
-		return CategoryViews(TypeCategories).IndexOfByPredicate(Pred);
+		return CategoryViews(TypeCategories, InRegistry).IndexOfByPredicate(Pred);
 	}
 
 	// returns an array of all of the supported types
 	RIGVM_API TArray<TRigVMTypeIndex> GetSupportedTypeIndices(const TArray<int32>& InPermutationIndices = TArray<int32>()) const;
+	RIGVM_API TArray<TRigVMTypeIndex> GetSupportedTypeIndices_NoLock(const TArray<int32>& InPermutationIndices, const FRigVMRegistryHandle& InRegistry) const;
 
 	// returns an array of all supported types as strings. this is used for automated testing only.
 	RIGVM_API TArray<FString> GetSupportedTypeStrings(const TArray<int32>& InPermutationIndices = TArray<int32>()) const;
+	RIGVM_API TArray<FString> GetSupportedTypeStrings_NoLock(const TArray<int32>& InPermutationIndices, const FRigVMRegistryHandle& InRegistry) const;
 	
 	// returns true if an argument is singleton (same type for all variants)
-	RIGVM_API bool IsSingleton(const TArray<int32>& InPermutationIndices = TArray<int32>(), const bool bLockRegistry = true) const;
-	RIGVM_API bool IsSingleton_NoLock(const TArray<int32>& InPermutationIndices = TArray<int32>()) const;
+	RIGVM_API bool IsSingleton(const TArray<int32>& InPermutationIndices = TArray<int32>()) const;
+	RIGVM_API bool IsSingleton_NoLock(const TArray<int32>& InPermutationIndices, const FRigVMRegistryHandle& InRegistry) const;
 
 	// returns true if this argument is an execute
 	RIGVM_API bool IsExecute() const;
-	RIGVM_API bool IsExecute_NoLock(const FRigVMRegistry_NoLock& InRegistry) const;
+	RIGVM_API bool IsExecute_NoLock(const FRigVMRegistryHandle& InRegistry) const;
 	
 	// returns true if the argument uses an array container
-	RIGVM_API EArrayType GetArrayType(const bool bLockRegistry = true) const;
-	RIGVM_API EArrayType GetArrayType_NoLock() const;
+	RIGVM_API EArrayType GetArrayType() const;
+	RIGVM_API EArrayType GetArrayType_NoLock(const FRigVMRegistryHandle& InRegistry) const;
 
-	RIGVM_API friend uint32 GetTypeHash_NoLock(const FRigVMTemplateArgument& InArgument);
+	RIGVM_API friend uint32 GetTypeHash_NoLock(const FRigVMTemplateArgument& InArgument, const FRigVMRegistryHandle& InRegistry);
 
 #if WITH_EDITOR
 	// Get the map of types to permutation indices
-	RIGVM_API const TArray<int32>& GetPermutations(const TRigVMTypeIndex InType, bool bLockRegistry = true) const;
-	RIGVM_API const TArray<int32>& GetPermutations_NoLock(const TRigVMTypeIndex InType) const;
+	RIGVM_API const TArray<int32>& GetPermutations(const TRigVMTypeIndex InType) const;
+	RIGVM_API const TArray<int32>& GetPermutations_NoLock(const TRigVMTypeIndex InType, const FRigVMRegistryHandle& InRegistry) const;
 	RIGVM_API void InvalidatePermutations(const TRigVMTypeIndex InType);
 #endif
 
@@ -299,23 +303,23 @@ protected:
 
 	bool bUseCategories = false;
 	TArray<ETypeCategory> TypeCategories;
-	TFunction<bool(const TRigVMTypeIndex&)> FilterType;
+	TFunction<bool(const TRigVMTypeIndex&, const FRigVMRegistryHandle&)> FilterType;
 	mutable TOptional<EArrayType> CachedArrayType;
 
 	// constructor from a property. this forces the type to be created
-	RIGVM_API FRigVMTemplateArgument(FProperty* InProperty, FRigVMRegistry_NoLock& InRegistry);
+	RIGVM_API FRigVMTemplateArgument(FProperty* InProperty, FRigVMRegistryHandle& InRegistry);
 
 	// static make function using a lock to create it
 	static RIGVM_API FRigVMTemplateArgument Make(FProperty* InProperty);
 
 	// static make function using a lock to create it
 	static RIGVM_API FRigVMTemplateArgument Make_NoLock(FProperty* InProperty);
-	static RIGVM_API FRigVMTemplateArgument Make_NoLock(FProperty* InProperty, FRigVMRegistry_NoLock& InRegistry);
+	static RIGVM_API FRigVMTemplateArgument Make_NoLock(FProperty* InProperty, FRigVMRegistryHandle& InRegistry);
 
-	RIGVM_API void EnsureValidExecuteType_NoLock(FRigVMRegistry_NoLock& InRegistry);
+	RIGVM_API void EnsureValidExecuteType_NoLock(const FRigVMRegistryHandle& InRegistry);
 
 #if WITH_EDITOR
-	RIGVM_API void UpdateTypeToPermutationsSlow();
+	RIGVM_API void UpdateTypeToPermutationsSlow(const FRigVMRegistryHandle& InRegistry);
 #endif
 
 	friend struct FRigVMTemplate;
@@ -325,12 +329,13 @@ protected:
 	friend struct FRigVMStructUpgradeInfo;
 	friend class URigVMCompiler;
 	friend class FRigVMTypeCacheScope_NoLock;
+	friend class URigVM;
 
 private:
 	struct CategoryViews
 	{
 		CategoryViews() = delete;
-		CategoryViews(const TArray<ETypeCategory>& InCategories);
+		CategoryViews(const TArray<ETypeCategory>& InCategories, const FRigVMRegistryHandle& InRegistry);
 		
 		void ForEachType(TFunction<bool(const TRigVMTypeIndex InType)>&& InCallback) const;
 
@@ -365,12 +370,12 @@ private:
 
 struct FRigVMTemplateArgumentInfo
 {
-	using ArgumentCallback = TFunction<FRigVMTemplateArgument(const FName /*InName*/, ERigVMPinDirection /*InDirection*/)>;
-	using TypeFilterCallback = TFunction<bool(const TRigVMTypeIndex& InNewType)>;
+	using ArgumentCallback = TFunction<FRigVMTemplateArgument(const FName /*InName*/, ERigVMPinDirection /*InDirection*/, const FRigVMRegistryHandle& /*InContext*/)>;
+	using TypeFilterCallback = TFunction<bool(const TRigVMTypeIndex& /*InNewType*/, const FRigVMRegistryHandle& /*InContext*/)>;
 	
 	FName Name = NAME_None;
 	ERigVMPinDirection Direction = ERigVMPinDirection::Invalid;
-	TFunction<FRigVMTemplateArgument(const FName /*InName*/, ERigVMPinDirection /*InDirection*/)> FactoryCallback = [](const FName, ERigVMPinDirection) { return FRigVMTemplateArgument(); };
+	ArgumentCallback FactoryCallback = [](const FName, ERigVMPinDirection, const FRigVMRegistryHandle&) { return FRigVMTemplateArgument(); };
 	
 	RIGVM_API FRigVMTemplateArgumentInfo(const FName InName, ERigVMPinDirection InDirection, TRigVMTypeIndex InTypeIndex);
 	RIGVM_API FRigVMTemplateArgumentInfo(const FName InName, ERigVMPinDirection InDirection, const TArray<TRigVMTypeIndex>& InTypeIndices);
@@ -379,14 +384,21 @@ struct FRigVMTemplateArgumentInfo
 	RIGVM_API FRigVMTemplateArgumentInfo(const FName InName, ERigVMPinDirection InDirection, ArgumentCallback&& InCallback);
 
 	RIGVM_API FRigVMTemplateArgument GetArgument() const;
-	
+	RIGVM_API FRigVMTemplateArgument GetArgument_NoLock(const FRigVMRegistryHandle& InRegistry) const;
+
 	static RIGVM_API FName ComputeTemplateNotation(
 		const FName InTemplateName,
 		const TArray<FRigVMTemplateArgumentInfo>& InInfos);
 	
 	static RIGVM_API TArray<TRigVMTypeIndex> GetTypesFromCategories(
 		const TArray<FRigVMTemplateArgument::ETypeCategory>& InTypeCategories,
-		const FRigVMTemplateArgument::FTypeFilter& InTypeFilter = {});	
+		const FRigVMTemplateArgument::FTypeFilter& InTypeFilter = {});
+
+	static TArray<TRigVMTypeIndex> GetTypesFromCategories_NoLock(
+    	const TArray<FRigVMTemplateArgument::ETypeCategory>& InTypeCategories,
+    	const FRigVMTemplateArgument::FTypeFilter& InTypeFilter,
+    	const FRigVMRegistryHandle& InRegistry);
+
 };
 
 /**
@@ -418,7 +430,7 @@ public:
 	RIGVM_API FName GetNodeName() const;
 
 	// returns true if this template can merge another one
-	RIGVM_API bool Merge(const FRigVMTemplate& InOther);
+	RIGVM_API bool Merge(const FRigVMTemplate& InOther, const FRigVMRegistryHandle& InRegistry);
 
 	// returns the number of args of this template
 	int32 NumArguments() const { return Arguments.Num(); }
@@ -439,10 +451,12 @@ public:
 	RIGVM_API const FRigVMExecuteArgument* FindExecuteArgument(const FName& InArgumentName, const FRigVMDispatchContext& InContext) const;
 
 	// returns the top level execute context struct this template uses
-	RIGVM_API const UScriptStruct* GetExecuteContextStruct(bool bLockRegistry = true) const;
+	RIGVM_API const UScriptStruct* GetExecuteContextStruct() const;
+	RIGVM_API const UScriptStruct* GetExecuteContextStruct_NoLock(const FRigVMRegistryHandle& InRegistry) const;
 
 	// returns true if this template supports a given execute context struct
 	RIGVM_API bool SupportsExecuteContextStruct(const UScriptStruct* InExecuteContextStruct) const;
+	RIGVM_API bool SupportsExecuteContextStruct_NoLock(const UScriptStruct* InExecuteContextStruct, const FRigVMRegistryHandle& InRegistry) const;
 
 #if WITH_EDITOR
 	// returns true if a given arg supports a type
@@ -450,47 +464,56 @@ public:
 #endif
 
 	// returns the number of permutations supported by this template
-	int32 NumPermutations() const { return Permutations.Num(); }
+	RIGVM_API int32 NumPermutations() const;
+	RIGVM_API int32 NumPermutations_NoLock(const FRigVMRegistryHandle& InRegistry) const { return Permutations.Num(); }
 
 	// returns the first / primary permutation of the template
-	RIGVM_API const FRigVMFunction* GetPrimaryPermutation(bool bLockRegistry = true) const;
+	RIGVM_API const FRigVMFunction* GetPrimaryPermutation() const;
+	RIGVM_API const FRigVMFunction* GetPrimaryPermutation_NoLock(const FRigVMRegistryHandle& InRegistry) const;
 
 	// returns a permutation given an index
-	RIGVM_API const FRigVMFunction* GetPermutation(int32 InIndex, bool bLockRegistry = true) const;
+	RIGVM_API const FRigVMFunction* GetPermutation(int32 InIndex) const;
+	RIGVM_API const FRigVMFunction* GetPermutation_NoLock(int32 InIndex, const FRigVMRegistryHandle& InRegistry) const;
 
 	// returns a permutation given an index and creates it using the backing factory if needed
-	RIGVM_API const FRigVMFunction* GetOrCreatePermutation(int32 InIndex, bool bLockRegistry = true);
+	RIGVM_API const FRigVMFunction* GetOrCreatePermutation(int32 InIndex);
+	RIGVM_API const FRigVMFunction* GetOrCreatePermutation_NoLock(int32 InIndex, FRigVMRegistryHandle& InRegistry);
 
 	// returns true if a given function is a permutation of this template
-	RIGVM_API bool ContainsPermutation(const FRigVMFunction* InPermutation, bool bLockRegistry = true) const;
-	RIGVM_API bool ContainsPermutation_NoLock(const FRigVMFunction* InPermutation) const;
+	RIGVM_API bool ContainsPermutation(const FRigVMFunction* InPermutation) const;
+	RIGVM_API bool ContainsPermutation_NoLock(const FRigVMFunction* InPermutation, const FRigVMRegistryHandle& InRegistry) const;
 
 	// returns the index of the permutation within the template of a given function (or INDEX_NONE)
 	RIGVM_API int32 FindPermutation(const FRigVMFunction* InPermutation) const;
+	RIGVM_API int32 FindPermutation_NoLock(const FRigVMFunction* InPermutation, const FRigVMRegistryHandle& InRegistry) const;
 
 	// returns the index of the permutation within the template of a given set of types
-	RIGVM_API int32 FindPermutation(const FTypeMap& InTypes, bool bLockRegistry = true) const;
+	RIGVM_API int32 FindPermutation(const FTypeMap& InTypes) const;
+	RIGVM_API int32 FindPermutation_NoLock(const FTypeMap& InTypes, const FRigVMRegistryHandle& InRegistry) const;
 
 	// returns true if the template was able to resolve to single permutation
-	RIGVM_API bool FullyResolve(FTypeMap& InOutTypes, int32& OutPermutationIndex, bool bLockRegistry = true) const;
+	RIGVM_API bool FullyResolve(FTypeMap& InOutTypes, int32& OutPermutationIndex) const;
+	RIGVM_API bool FullyResolve_NoLock(FTypeMap& InOutTypes, int32& OutPermutationIndex, const FRigVMRegistryHandle& InRegistry) const;
 
 	// returns true if the template was able to resolve to at least one permutation
-	RIGVM_API bool Resolve(FTypeMap& InOutTypes, TArray<int32> & OutPermutationIndices, bool bAllowFloatingPointCasts, bool bLockRegistry = true) const;
+	RIGVM_API bool Resolve(FTypeMap& InOutTypes, TArray<int32> & OutPermutationIndices, bool bAllowFloatingPointCasts) const;
+	RIGVM_API bool Resolve_NoLock(FTypeMap& InOutTypes, TArray<int32>& OutPermutationIndices, bool bAllowFloatingPointCasts, const FRigVMRegistryHandle& InRegistry) const;
 
 	// Will return the hash of the input type map, if it is a valid type map. Otherwise, will return 0.
 	// It is only a valid type map if it includes all arguments, and non of the types is a wildcard.
-	RIGVM_API uint32 GetTypesHashFromTypes(const FTypeMap& InTypes) const;
+	RIGVM_API uint32 GetTypesHashFromTypes(const FTypeMap& InTypes, const FRigVMRegistryHandle& InRegistry) const;
 
 	// returns true if the template was able to resolve to at least one permutation
-	RIGVM_API bool ContainsPermutation(const FTypeMap& InTypes, bool bLockRegistry = true) const;
-	RIGVM_API bool ContainsPermutation_NoLock(const FTypeMap& InTypes) const;
+	RIGVM_API bool ContainsPermutation(const FTypeMap& InTypes) const;
+	RIGVM_API bool ContainsPermutation_NoLock(const FTypeMap& InTypes, const FRigVMRegistryHandle& InRegistry) const;
 
 	// returns true if the template can resolve an argument to a new type
-	RIGVM_API bool ResolveArgument(const FName& InArgumentName, const TRigVMTypeIndex InTypeIndex, FTypeMap& InOutTypes, bool bLockRegistry = true) const;
+	RIGVM_API bool ResolveArgument(const FName& InArgumentName, const TRigVMTypeIndex InTypeIndex, FTypeMap& InOutTypes) const;
+	RIGVM_API bool ResolveArgument_NoLock(const FName& InArgumentName, const TRigVMTypeIndex InTypeIndex, FTypeMap& InOutTypes, const FRigVMRegistryHandle& InRegistry) const;
 
 	// returns the types for a specific permutation
-	RIGVM_API FRigVMTemplateTypeMap GetTypesForPermutation(const int32 InPermutationIndex, const bool bLockRegistry = true) const;
-	RIGVM_API FRigVMTemplateTypeMap GetTypesForPermutation_NoLock(const int32 InPermutationIndex) const;
+	RIGVM_API FRigVMTemplateTypeMap GetTypesForPermutation(const int32 InPermutationIndex) const;
+	RIGVM_API FRigVMTemplateTypeMap GetTypesForPermutation_NoLock(const int32 InPermutationIndex, const FRigVMRegistryHandle& InRegistry) const;
 
 	// returns true if a given argument is valid for a template
 	static RIGVM_API bool IsValidArgumentForTemplate(const ERigVMPinDirection InDirection);
@@ -511,7 +534,7 @@ public:
 	RIGVM_API FTypeMap GetArgumentTypesFromString(const FString& InTypeString, const FRigVMUserDefinedTypeResolver* InTypeResolver = nullptr) const;
 
 	// converts the types provided to a string (like "A:float,B:int32")
-	static RIGVM_API FString GetStringFromArgumentTypes(const FTypeMap& InTypes, bool bLockRegistry = true); 
+	static RIGVM_API FString GetStringFromArgumentTypes(const FTypeMap& InTypes, const FRigVMRegistryHandle& InRegistry); 
 
 #if WITH_EDITOR
 
@@ -534,65 +557,64 @@ public:
 
 	// Updates the template's argument types. This only affects templates which have category based
 	// arguments and will resolve the other arguments to the expected types.
-	RIGVM_API bool UpdateAllArgumentTypesSlow();
+	RIGVM_API bool UpdateAllArgumentTypesSlow(FRigVMRegistryHandle& InRegistry);
 
 	// Updates the template's argument types by adding a single new type
-	RIGVM_API bool HandlePropagatedArgumentType(const TRigVMTypeIndex InTypeIndex);
+	RIGVM_API bool HandlePropagatedArgumentType(const TRigVMTypeIndex InTypeIndex, FRigVMRegistryHandle& InRegistry);
 
 	// Invalidates template permutations whenever a type such as a user defined struct is removed
-	RIGVM_API void HandleTypeRemoval(TRigVMTypeIndex InTypeIndex);
+	RIGVM_API void HandleTypeRemoval(TRigVMTypeIndex InTypeIndex, const FRigVMRegistryHandle& InRegistry);
 	
 	// Returns the delegate to be able to react to type changes dynamically
 	// This delegate is deprecated
 	FRigVMTemplate_NewArgumentTypeDelegate& OnNewArgumentType() { return Delegates.NewArgumentTypeDelegate; }
 
 	// Returns the factory this template was created by
-	RIGVM_API const FRigVMDispatchFactory* GetDispatchFactory(const bool bLockRegistry = true) const;
+	RIGVM_API const FRigVMDispatchFactory* GetDispatchFactory() const;
 
 	// Returns the factory this template was created by
-	const FRigVMDispatchFactory* GetDispatchFactory_NoLock() const
-	{
-		return UsesDispatch() ? Delegates.GetDispatchFactoryDelegate.Execute() : nullptr;
-	}
-
+	const FRigVMDispatchFactory* GetDispatchFactory_NoLock(const FRigVMRegistryHandle& InRegistry) const;
+	
 	// Returns true if this template is backed by a dispatch factory
 	bool UsesDispatch() const
 	{
 		return Delegates.GetDispatchFactoryDelegate.IsBound();
 	}
 
-	RIGVM_API void RecomputeTypesHashToPermutations();
-	RIGVM_API void RecomputeTypesHashToPermutations(const TArray<FRigVMTypeCacheScope_NoLock>& InTypeCaches);
+	RIGVM_API void RecomputeTypesHashToPermutations(const FRigVMRegistryHandle& InRegistry);
+	RIGVM_API void RecomputeTypesHashToPermutations(const TArray<FRigVMTypeCacheScope_NoLock>& InTypeCaches, const FRigVMRegistryHandle& InRegistry);
 
-	RIGVM_API void UpdateTypesHashToPermutation(const int32 InPermutation);
+	RIGVM_API void UpdateTypesHashToPermutation(const int32 InPermutation, const FRigVMRegistryHandle& InRegistry);
 
 	RIGVM_API friend uint32 GetTypeHash(const FRigVMTemplate& InTemplate);
-	RIGVM_API friend uint32 GetTypeHash_NoLock(const FRigVMTemplate& InTemplate);
+	RIGVM_API friend uint32 GetTypeHash_NoLock(const FRigVMTemplate& InTemplate, const FRigVMRegistryHandle& InRegistry);
+
+	RIGVM_API TArray<FRigVMTemplateArgumentInfo> GetFlattenedArgumentInfos_NoLock(FRigVMRegistryHandle& InRegistry) const;
+	RIGVM_API TArray<FRigVMTemplateArgumentInfo> GetFlattenedArgumentInfosForFunctions_NoLock(const TArray<const FRigVMFunction*>& InFunctions, FRigVMRegistryHandle& InRegistry) const;
+	RIGVM_API TArray<FRigVMTemplateArgumentInfo> GetFlattenedArgumentInfosForPermutations_NoLock(const TArray<int32>& InPermutations, FRigVMRegistryHandle& InRegistry) const;
 
 private:
 
 	// Constructor from a struct, a template name and a function index
-	FRigVMTemplate(UScriptStruct* InStruct, const FString& InTemplateName, int32 InFunctionIndex = INDEX_NONE);
+	FRigVMTemplate(UScriptStruct* InStruct, const FString& InTemplateName, int32 InFunctionIndex, FRigVMRegistryHandle& InRegistry);
 
 	// Constructor from a template name and argument infos
-	FRigVMTemplate(const FName& InTemplateName, const TArray<FRigVMTemplateArgumentInfo>& InInfos);
+	FRigVMTemplate(const FName& InTemplateName, const TArray<FRigVMTemplateArgumentInfo>& InInfos, const FRigVMRegistryHandle& InRegistry);
 
 	static FLinearColor GetColorFromMetadata(FString InMetadata);
 
 	void InvalidateHash() { Hash = UINT32_MAX; }
 	RIGVM_API const TArray<FRigVMExecuteArgument>& GetExecuteArguments(const FRigVMDispatchContext& InContext) const;
+	RIGVM_API const TArray<FRigVMExecuteArgument>& GetExecuteArguments_NoLock(const FRigVMDispatchContext& InContext, const FRigVMRegistryHandle& InRegistry) const;
 
-	const FRigVMFunction* GetPermutation_NoLock(int32 InIndex) const;
-	const FRigVMFunction* GetOrCreatePermutation_NoLock(int32 InIndex);
-
-	FTypeMap GetArgumentTypesFromString_Impl(const FString& InTypeString, const FRigVMUserDefinedTypeResolver* InTypeResolver, bool bLockRegistry) const;
+	FTypeMap GetArgumentTypesFromString_NoLock(const FString& InTypeString, const FRigVMUserDefinedTypeResolver* InTypeResolver, FRigVMRegistryHandle& InRegistry) const;
 
-	uint32 ComputeTypeHash() const;
+	uint32 ComputeTypeHash(const FRigVMRegistryHandle& InRegistry) const;
 
 	bool UpdateArgumentTypes_Impl(
 		const FRigVMTemplateArgument& InPrimaryArgument,
 		const TRigVMTypeIndex InPrimaryTypeIndex,
-		const FRigVMRegistry_NoLock& InRegistry, 
+		FRigVMRegistryHandle& InRegistry, 
 		const FRigVMDispatchFactory* InFactory,
 		TArray<FRigVMTemplateTypeMap, TInlineAllocator<1>>& InOutTypesArray);
 
@@ -603,6 +625,7 @@ private:
 	TArray<int32> Permutations;
 	TMap<uint32, int32> TypesHashToPermutation;
 	mutable uint32 Hash;
+	FRigVMRegistry_NoLock* OwnerRegistry;
 
 	FRigVMTemplateDelegates Delegates;
 
@@ -616,11 +639,11 @@ class FRigVMTypeCacheScope_NoLock
 {
 public:
 	FRigVMTypeCacheScope_NoLock();
-	FRigVMTypeCacheScope_NoLock(const FRigVMTemplateArgument& InArgument);
+	FRigVMTypeCacheScope_NoLock(const FRigVMRegistryHandle& InRegistry, const FRigVMTemplateArgument& InArgument);
 	~FRigVMTypeCacheScope_NoLock();
 
 	bool IsValid() const { return Argument != nullptr; }
-	const FRigVMTypeCacheScope_NoLock& UpdateIfRequired(const FRigVMTemplateArgument& InArgument);
+	const FRigVMTypeCacheScope_NoLock& UpdateIfRequired(const FRigVMRegistryHandle& InRegistry, const FRigVMTemplateArgument& InArgument);
 	int32 GetNumTypes_NoLock() const;
 	TRigVMTypeIndex GetTypeIndex_NoLock(int32 InIndex) const;
 	
@@ -628,6 +651,7 @@ private:
 
 	void UpdateTypesIfRequired() const;
 
+	FRigVMRegistry_NoLock* Registry;
 	const FRigVMTemplateArgument* Argument;
 	bool bShouldCopyTypes;
 	mutable TOptional<int32> NumTypes;
diff --git a/Engine/Plugins/Runtime/RigVM/Source/RigVM/Public/RigVMFunctions/RigVMDispatch_Array.h b/Engine/Plugins/Runtime/RigVM/Source/RigVM/Public/RigVMFunctions/RigVMDispatch_Array.h
--- a/Engine/Plugins/Runtime/RigVM/Source/RigVM/Public/RigVMFunctions/RigVMDispatch_Array.h
+++ b/Engine/Plugins/Runtime/RigVM/Source/RigVM/Public/RigVMFunctions/RigVMDispatch_Array.h
@@ -60,18 +60,18 @@ public:
 		FactoryScriptStruct = StaticStruct();
 	}
 
-	RIGVM_API virtual const TArray<FRigVMTemplateArgumentInfo>& GetArgumentInfos() const override;
-	RIGVM_API virtual FRigVMTemplateTypeMap OnNewArgumentType(const FName& InArgumentName, TRigVMTypeIndex InTypeIndex) const override;
+	RIGVM_API virtual const TArray<FRigVMTemplateArgumentInfo>& GetArgumentInfos(FRigVMRegistryHandle& InRegistry) const override;
+	RIGVM_API virtual FRigVMTemplateTypeMap OnNewArgumentType(const FName& InArgumentName, TRigVMTypeIndex InTypeIndex, FRigVMRegistryHandle& InRegistry) const override;
 #if WITH_EDITOR
 	RIGVM_API virtual FText GetNodeTooltip(const FRigVMTemplateTypeMap& InTypes) const override;
 	RIGVM_API virtual FText GetArgumentTooltip(const FName& InArgumentName, TRigVMTypeIndex InTypeIndex) const override;
 	RIGVM_API virtual FString GetArgumentMetaData(const FName& InArgumentName, const FName& InMetaDataKey) const override;
 	RIGVM_API virtual FString GetArgumentDefaultValue(const FName& InArgumentName, TRigVMTypeIndex InTypeIndex) const override;
 #endif
-	RIGVM_API virtual FName GetArgumentNameForOperandIndex(int32 InOperandIndex, int32 InTotalOperands) const override;
+	RIGVM_API virtual FName GetArgumentNameForOperandIndex(int32 InOperandIndex, int32 InTotalOperands, FRigVMRegistryHandle& InRegistry) const override;
 
 protected:
-	virtual FRigVMFunctionPtr GetDispatchFunctionImpl(const FRigVMTemplateTypeMap& InTypes) const override { return &FRigVMDispatch_ArrayMake::Execute; }
+	virtual FRigVMFunctionPtr GetDispatchFunctionImpl(const FRigVMTemplateTypeMap& InTypes, const FRigVMRegistryHandle& InRegistry) const override { return &FRigVMDispatch_ArrayMake::Execute; }
 	static RIGVM_API void Execute(FRigVMExtendedExecuteContext& InContext, FRigVMMemoryHandleArray Handles, FRigVMPredicateBranchArray Predicates);
 };
 
@@ -87,16 +87,16 @@ public:
 	}
 
 	virtual ERigVMOpCode GetOpCode() const override { return ERigVMOpCode::ArrayReset; }
-	RIGVM_API virtual FName GetArgumentNameForOperandIndex(int32 InOperandIndex, int32 InTotalOperands) const override;
-	RIGVM_API virtual const TArray<FRigVMTemplateArgumentInfo>& GetArgumentInfos() const override;
-	RIGVM_API virtual FRigVMTemplateTypeMap OnNewArgumentType(const FName& InArgumentName, TRigVMTypeIndex InTypeIndex) const override;
+	RIGVM_API virtual FName GetArgumentNameForOperandIndex(int32 InOperandIndex, int32 InTotalOperands, FRigVMRegistryHandle& InRegistry) const override;
+	RIGVM_API virtual const TArray<FRigVMTemplateArgumentInfo>& GetArgumentInfos(FRigVMRegistryHandle& InRegistry) const override;
+	RIGVM_API virtual FRigVMTemplateTypeMap OnNewArgumentType(const FName& InArgumentName, TRigVMTypeIndex InTypeIndex, FRigVMRegistryHandle& InRegistry) const override;
 #if WITH_EDITOR
 	RIGVM_API virtual FText GetNodeTooltip(const FRigVMTemplateTypeMap& InTypes) const override;
 	RIGVM_API virtual FText GetArgumentTooltip(const FName& InArgumentName, TRigVMTypeIndex InTypeIndex) const override;
 #endif
 
 protected:
-	virtual FRigVMFunctionPtr GetDispatchFunctionImpl(const FRigVMTemplateTypeMap& InTypes) const override { return &FRigVMDispatch_ArrayReset::Execute; }
+	virtual FRigVMFunctionPtr GetDispatchFunctionImpl(const FRigVMTemplateTypeMap& InTypes, const FRigVMRegistryHandle& InRegistry) const override { return &FRigVMDispatch_ArrayReset::Execute; }
 	static RIGVM_API void Execute(FRigVMExtendedExecuteContext& InContext, FRigVMMemoryHandleArray Handles, FRigVMPredicateBranchArray Predicates);
 };
 
@@ -112,16 +112,16 @@ public:
 	}
 
 	virtual ERigVMOpCode GetOpCode() const override { return ERigVMOpCode::ArrayGetNum; }
-	RIGVM_API virtual FName GetArgumentNameForOperandIndex(int32 InOperandIndex, int32 InTotalOperands) const override;
-	RIGVM_API virtual const TArray<FRigVMTemplateArgumentInfo>& GetArgumentInfos() const override;
-	RIGVM_API virtual FRigVMTemplateTypeMap OnNewArgumentType(const FName& InArgumentName, TRigVMTypeIndex InTypeIndex) const override;
+	RIGVM_API virtual FName GetArgumentNameForOperandIndex(int32 InOperandIndex, int32 InTotalOperands, FRigVMRegistryHandle& InRegistry) const override;
+	RIGVM_API virtual const TArray<FRigVMTemplateArgumentInfo>& GetArgumentInfos(FRigVMRegistryHandle& InRegistry) const override;
+	RIGVM_API virtual FRigVMTemplateTypeMap OnNewArgumentType(const FName& InArgumentName, TRigVMTypeIndex InTypeIndex, FRigVMRegistryHandle& InRegistry) const override;
 #if WITH_EDITOR
 	RIGVM_API virtual FText GetNodeTooltip(const FRigVMTemplateTypeMap& InTypes) const override;
 	RIGVM_API virtual FText GetArgumentTooltip(const FName& InArgumentName, TRigVMTypeIndex InTypeIndex) const override;
 #endif
 
 protected:
-	virtual FRigVMFunctionPtr GetDispatchFunctionImpl(const FRigVMTemplateTypeMap& InTypes) const override { return &FRigVMDispatch_ArrayGetNum::Execute; }
+	virtual FRigVMFunctionPtr GetDispatchFunctionImpl(const FRigVMTemplateTypeMap& InTypes, const FRigVMRegistryHandle& InRegistry) const override { return &FRigVMDispatch_ArrayGetNum::Execute; }
 	static RIGVM_API void Execute(FRigVMExtendedExecuteContext& InContext, FRigVMMemoryHandleArray Handles, FRigVMPredicateBranchArray Predicates);
 };
 
@@ -137,16 +137,16 @@ public:
 	}
 
 	virtual ERigVMOpCode GetOpCode() const override { return ERigVMOpCode::ArraySetNum; }
-	RIGVM_API virtual FName GetArgumentNameForOperandIndex(int32 InOperandIndex, int32 InTotalOperands) const override;
-	RIGVM_API virtual const TArray<FRigVMTemplateArgumentInfo>& GetArgumentInfos() const override;
-	RIGVM_API virtual FRigVMTemplateTypeMap OnNewArgumentType(const FName& InArgumentName, TRigVMTypeIndex InTypeIndex) const override;
+	RIGVM_API virtual FName GetArgumentNameForOperandIndex(int32 InOperandIndex, int32 InTotalOperands, FRigVMRegistryHandle& InRegistry) const override;
+	RIGVM_API virtual const TArray<FRigVMTemplateArgumentInfo>& GetArgumentInfos(FRigVMRegistryHandle& InRegistry) const override;
+	RIGVM_API virtual FRigVMTemplateTypeMap OnNewArgumentType(const FName& InArgumentName, TRigVMTypeIndex InTypeIndex, FRigVMRegistryHandle& InRegistry) const override;
 #if WITH_EDITOR
 	RIGVM_API virtual FText GetNodeTooltip(const FRigVMTemplateTypeMap& InTypes) const override;
 	RIGVM_API virtual FText GetArgumentTooltip(const FName& InArgumentName, TRigVMTypeIndex InTypeIndex) const override;
 #endif
 
 protected:
-	virtual FRigVMFunctionPtr GetDispatchFunctionImpl(const FRigVMTemplateTypeMap& InTypes) const override { return &FRigVMDispatch_ArraySetNum::Execute; }
+	virtual FRigVMFunctionPtr GetDispatchFunctionImpl(const FRigVMTemplateTypeMap& InTypes, const FRigVMRegistryHandle& InRegistry) const override { return &FRigVMDispatch_ArraySetNum::Execute; }
 	static RIGVM_API void Execute(FRigVMExtendedExecuteContext& InContext, FRigVMMemoryHandleArray Handles, FRigVMPredicateBranchArray Predicates);
 };
 
@@ -161,16 +161,16 @@ public:
 		FactoryScriptStruct = StaticStruct();
 	}
 
-	RIGVM_API virtual FName GetArgumentNameForOperandIndex(int32 InOperandIndex, int32 InTotalOperands) const override;
-	RIGVM_API virtual const TArray<FRigVMTemplateArgumentInfo>& GetArgumentInfos() const override;
-	RIGVM_API virtual FRigVMTemplateTypeMap OnNewArgumentType(const FName& InArgumentName, TRigVMTypeIndex InTypeIndex) const override;
+	RIGVM_API virtual FName GetArgumentNameForOperandIndex(int32 InOperandIndex, int32 InTotalOperands, FRigVMRegistryHandle& InRegistry) const override;
+	RIGVM_API virtual const TArray<FRigVMTemplateArgumentInfo>& GetArgumentInfos(FRigVMRegistryHandle& InRegistry) const override;
+	RIGVM_API virtual FRigVMTemplateTypeMap OnNewArgumentType(const FName& InArgumentName, TRigVMTypeIndex InTypeIndex, FRigVMRegistryHandle& InRegistry) const override;
 #if WITH_EDITOR
 	RIGVM_API virtual FText GetNodeTooltip(const FRigVMTemplateTypeMap& InTypes) const override;
 	RIGVM_API virtual FText GetArgumentTooltip(const FName& InArgumentName, TRigVMTypeIndex InTypeIndex) const override;
 #endif
 
 protected:
-	virtual FRigVMFunctionPtr GetDispatchFunctionImpl(const FRigVMTemplateTypeMap& InTypes) const override { return &FRigVMDispatch_ArrayInit::Execute; }
+	virtual FRigVMFunctionPtr GetDispatchFunctionImpl(const FRigVMTemplateTypeMap& InTypes, const FRigVMRegistryHandle& InRegistry) const override { return &FRigVMDispatch_ArrayInit::Execute; }
 	static RIGVM_API void Execute(FRigVMExtendedExecuteContext& InContext, FRigVMMemoryHandleArray Handles, FRigVMPredicateBranchArray Predicates);
 };
 
@@ -186,16 +186,16 @@ public:
 	}
 
 	virtual ERigVMOpCode GetOpCode() const override { return ERigVMOpCode::ArrayGetAtIndex; }
-	RIGVM_API virtual FName GetArgumentNameForOperandIndex(int32 InOperandIndex, int32 InTotalOperands) const override;
-	RIGVM_API virtual const TArray<FRigVMTemplateArgumentInfo>& GetArgumentInfos() const override;
-	RIGVM_API virtual FRigVMTemplateTypeMap OnNewArgumentType(const FName& InArgumentName, TRigVMTypeIndex InTypeIndex) const override;
+	RIGVM_API virtual FName GetArgumentNameForOperandIndex(int32 InOperandIndex, int32 InTotalOperands, FRigVMRegistryHandle& InRegistry) const override;
+	RIGVM_API virtual const TArray<FRigVMTemplateArgumentInfo>& GetArgumentInfos(FRigVMRegistryHandle& InRegistry) const override;
+	RIGVM_API virtual FRigVMTemplateTypeMap OnNewArgumentType(const FName& InArgumentName, TRigVMTypeIndex InTypeIndex, FRigVMRegistryHandle& InRegistry) const override;
 #if WITH_EDITOR
 	RIGVM_API virtual FText GetNodeTooltip(const FRigVMTemplateTypeMap& InTypes) const override;
 	RIGVM_API virtual FText GetArgumentTooltip(const FName& InArgumentName, TRigVMTypeIndex InTypeIndex) const override;
 #endif
 
 protected:
-	virtual FRigVMFunctionPtr GetDispatchFunctionImpl(const FRigVMTemplateTypeMap& InTypes) const override { return &FRigVMDispatch_ArrayGetAtIndex::Execute; }
+	virtual FRigVMFunctionPtr GetDispatchFunctionImpl(const FRigVMTemplateTypeMap& InTypes, const FRigVMRegistryHandle& InRegistry) const override { return &FRigVMDispatch_ArrayGetAtIndex::Execute; }
 	static RIGVM_API void Execute(FRigVMExtendedExecuteContext& InContext, FRigVMMemoryHandleArray Handles, FRigVMPredicateBranchArray Predicates);
 };
 
@@ -211,16 +211,16 @@ public:
 	}
 
 	virtual ERigVMOpCode GetOpCode() const override { return ERigVMOpCode::ArraySetAtIndex; }
-	RIGVM_API virtual FName GetArgumentNameForOperandIndex(int32 InOperandIndex, int32 InTotalOperands) const override;
-	RIGVM_API virtual const TArray<FRigVMTemplateArgumentInfo>& GetArgumentInfos() const override;
-	RIGVM_API virtual FRigVMTemplateTypeMap OnNewArgumentType(const FName& InArgumentName, TRigVMTypeIndex InTypeIndex) const override;
+	RIGVM_API virtual FName GetArgumentNameForOperandIndex(int32 InOperandIndex, int32 InTotalOperands, FRigVMRegistryHandle& InRegistry) const override;
+	RIGVM_API virtual const TArray<FRigVMTemplateArgumentInfo>& GetArgumentInfos(FRigVMRegistryHandle& InRegistry) const override;
+	RIGVM_API virtual FRigVMTemplateTypeMap OnNewArgumentType(const FName& InArgumentName, TRigVMTypeIndex InTypeIndex, FRigVMRegistryHandle& InRegistry) const override;
 #if WITH_EDITOR
 	RIGVM_API virtual FText GetNodeTooltip(const FRigVMTemplateTypeMap& InTypes) const override;
 	RIGVM_API virtual FText GetArgumentTooltip(const FName& InArgumentName, TRigVMTypeIndex InTypeIndex) const override;
 #endif
 
 protected:
-	virtual FRigVMFunctionPtr GetDispatchFunctionImpl(const FRigVMTemplateTypeMap& InTypes) const override { return &FRigVMDispatch_ArraySetAtIndex::Execute; }
+	virtual FRigVMFunctionPtr GetDispatchFunctionImpl(const FRigVMTemplateTypeMap& InTypes, const FRigVMRegistryHandle& InRegistry) const override { return &FRigVMDispatch_ArraySetAtIndex::Execute; }
 	static RIGVM_API void Execute(FRigVMExtendedExecuteContext& InContext, FRigVMMemoryHandleArray Handles, FRigVMPredicateBranchArray Predicates);
 };
 
@@ -236,15 +236,15 @@ public:
 	}
 
 	virtual ERigVMOpCode GetOpCode() const override { return ERigVMOpCode::ArrayAdd; }
-	RIGVM_API virtual FName GetArgumentNameForOperandIndex(int32 InOperandIndex, int32 InTotalOperands) const override;
-	RIGVM_API virtual const TArray<FRigVMTemplateArgumentInfo>& GetArgumentInfos() const override;
+	RIGVM_API virtual FName GetArgumentNameForOperandIndex(int32 InOperandIndex, int32 InTotalOperands, FRigVMRegistryHandle& InRegistry) const override;
+	RIGVM_API virtual const TArray<FRigVMTemplateArgumentInfo>& GetArgumentInfos(FRigVMRegistryHandle& InRegistry) const override;
 #if WITH_EDITOR
 	RIGVM_API virtual FText GetNodeTooltip(const FRigVMTemplateTypeMap& InTypes) const override;
 	RIGVM_API virtual FText GetArgumentTooltip(const FName& InArgumentName, TRigVMTypeIndex InTypeIndex) const override;
 #endif
 
 protected:
-	virtual FRigVMFunctionPtr GetDispatchFunctionImpl(const FRigVMTemplateTypeMap& InTypes) const override { return &FRigVMDispatch_ArrayAdd::Execute; }
+	virtual FRigVMFunctionPtr GetDispatchFunctionImpl(const FRigVMTemplateTypeMap& InTypes, const FRigVMRegistryHandle& InRegistry) const override { return &FRigVMDispatch_ArrayAdd::Execute; }
 	static RIGVM_API void Execute(FRigVMExtendedExecuteContext& InContext, FRigVMMemoryHandleArray Handles, FRigVMPredicateBranchArray Predicates);
 };
 
@@ -266,7 +266,7 @@ public:
 #endif
 
 protected:
-	virtual FRigVMFunctionPtr GetDispatchFunctionImpl(const FRigVMTemplateTypeMap& InTypes) const override { return &FRigVMDispatch_ArrayInsert::Execute; }
+	virtual FRigVMFunctionPtr GetDispatchFunctionImpl(const FRigVMTemplateTypeMap& InTypes, const FRigVMRegistryHandle& InRegistry) const override { return &FRigVMDispatch_ArrayInsert::Execute; }
 	static RIGVM_API void Execute(FRigVMExtendedExecuteContext& InContext, FRigVMMemoryHandleArray Handles, FRigVMPredicateBranchArray Predicates);
 };
 
@@ -282,16 +282,16 @@ public:
 	}
 
 	virtual ERigVMOpCode GetOpCode() const override { return ERigVMOpCode::ArrayRemove; }
-	RIGVM_API virtual FName GetArgumentNameForOperandIndex(int32 InOperandIndex, int32 InTotalOperands) const override;
-	RIGVM_API virtual const TArray<FRigVMTemplateArgumentInfo>& GetArgumentInfos() const override;
-	RIGVM_API virtual FRigVMTemplateTypeMap OnNewArgumentType(const FName& InArgumentName, TRigVMTypeIndex InTypeIndex) const override;
+	RIGVM_API virtual FName GetArgumentNameForOperandIndex(int32 InOperandIndex, int32 InTotalOperands, FRigVMRegistryHandle& InRegistry) const override;
+	RIGVM_API virtual const TArray<FRigVMTemplateArgumentInfo>& GetArgumentInfos(FRigVMRegistryHandle& InRegistry) const override;
+	RIGVM_API virtual FRigVMTemplateTypeMap OnNewArgumentType(const FName& InArgumentName, TRigVMTypeIndex InTypeIndex, FRigVMRegistryHandle& InRegistry) const override;
 #if WITH_EDITOR
    	RIGVM_API virtual FText GetNodeTooltip(const FRigVMTemplateTypeMap& InTypes) const override;
    	RIGVM_API virtual FText GetArgumentTooltip(const FName& InArgumentName, TRigVMTypeIndex InTypeIndex) const override;
 #endif
 
 protected:
-	virtual FRigVMFunctionPtr GetDispatchFunctionImpl(const FRigVMTemplateTypeMap& InTypes) const override { return &FRigVMDispatch_ArrayRemove::Execute; }
+	virtual FRigVMFunctionPtr GetDispatchFunctionImpl(const FRigVMTemplateTypeMap& InTypes, const FRigVMRegistryHandle& InRegistry) const override { return &FRigVMDispatch_ArrayRemove::Execute; }
 	static RIGVM_API void Execute(FRigVMExtendedExecuteContext& InContext, FRigVMMemoryHandleArray Handles, FRigVMPredicateBranchArray Predicates);
 };
 
@@ -314,7 +314,7 @@ public:
 
 protected:
 
-	virtual FRigVMFunctionPtr GetDispatchFunctionImpl(const FRigVMTemplateTypeMap& InTypes) const override { return &FRigVMDispatch_ArrayReverse::Execute; }
+	virtual FRigVMFunctionPtr GetDispatchFunctionImpl(const FRigVMTemplateTypeMap& InTypes, const FRigVMRegistryHandle& InRegistry) const override { return &FRigVMDispatch_ArrayReverse::Execute; }
 	static RIGVM_API void Execute(FRigVMExtendedExecuteContext& InContext, FRigVMMemoryHandleArray Handles, FRigVMPredicateBranchArray Predicates);
 };
 
@@ -330,16 +330,16 @@ public:
 	}
 
 	virtual ERigVMOpCode GetOpCode() const override { return ERigVMOpCode::ArrayFind; }
-	RIGVM_API virtual FName GetArgumentNameForOperandIndex(int32 InOperandIndex, int32 InTotalOperands) const override;
-	RIGVM_API virtual const TArray<FRigVMTemplateArgumentInfo>& GetArgumentInfos() const override;
-	RIGVM_API virtual FRigVMTemplateTypeMap OnNewArgumentType(const FName& InArgumentName, TRigVMTypeIndex InTypeIndex) const override;
+	RIGVM_API virtual FName GetArgumentNameForOperandIndex(int32 InOperandIndex, int32 InTotalOperands, FRigVMRegistryHandle& InRegistry) const override;
+	RIGVM_API virtual const TArray<FRigVMTemplateArgumentInfo>& GetArgumentInfos(FRigVMRegistryHandle& InRegistry) const override;
+	RIGVM_API virtual FRigVMTemplateTypeMap OnNewArgumentType(const FName& InArgumentName, TRigVMTypeIndex InTypeIndex, FRigVMRegistryHandle& InRegistry) const override;
 #if WITH_EDITOR
    	RIGVM_API virtual FText GetNodeTooltip(const FRigVMTemplateTypeMap& InTypes) const override;
    	RIGVM_API virtual FText GetArgumentTooltip(const FName& InArgumentName, TRigVMTypeIndex InTypeIndex) const override;
 #endif
 
 protected:
-	virtual FRigVMFunctionPtr GetDispatchFunctionImpl(const FRigVMTemplateTypeMap& InTypes) const override { return &FRigVMDispatch_ArrayFind::Execute; }
+	virtual FRigVMFunctionPtr GetDispatchFunctionImpl(const FRigVMTemplateTypeMap& InTypes, const FRigVMRegistryHandle& InRegistry) const override { return &FRigVMDispatch_ArrayFind::Execute; }
 	static RIGVM_API void Execute(FRigVMExtendedExecuteContext& InContext, FRigVMMemoryHandleArray Handles, FRigVMPredicateBranchArray Predicates);
 };
 
@@ -355,16 +355,16 @@ public:
 	}
 
 	virtual ERigVMOpCode GetOpCode() const override { return ERigVMOpCode::ArrayAppend; }
-	RIGVM_API virtual FName GetArgumentNameForOperandIndex(int32 InOperandIndex, int32 InTotalOperands) const override;
-	RIGVM_API virtual const TArray<FRigVMTemplateArgumentInfo>& GetArgumentInfos() const override;
-	RIGVM_API virtual FRigVMTemplateTypeMap OnNewArgumentType(const FName& InArgumentName, TRigVMTypeIndex InTypeIndex) const override;
+	RIGVM_API virtual FName GetArgumentNameForOperandIndex(int32 InOperandIndex, int32 InTotalOperands, FRigVMRegistryHandle& InRegistry) const override;
+	RIGVM_API virtual const TArray<FRigVMTemplateArgumentInfo>& GetArgumentInfos(FRigVMRegistryHandle& InRegistry) const override;
+	RIGVM_API virtual FRigVMTemplateTypeMap OnNewArgumentType(const FName& InArgumentName, TRigVMTypeIndex InTypeIndex, FRigVMRegistryHandle& InRegistry) const override;
 #if WITH_EDITOR
    	RIGVM_API virtual FText GetNodeTooltip(const FRigVMTemplateTypeMap& InTypes) const override;
    	RIGVM_API virtual FText GetArgumentTooltip(const FName& InArgumentName, TRigVMTypeIndex InTypeIndex) const override;
 #endif
 
 protected:
-	virtual FRigVMFunctionPtr GetDispatchFunctionImpl(const FRigVMTemplateTypeMap& InTypes) const override { return &FRigVMDispatch_ArrayAppend::Execute; }
+	virtual FRigVMFunctionPtr GetDispatchFunctionImpl(const FRigVMTemplateTypeMap& InTypes, const FRigVMRegistryHandle& InRegistry) const override { return &FRigVMDispatch_ArrayAppend::Execute; }
 	static RIGVM_API void Execute(FRigVMExtendedExecuteContext& InContext, FRigVMMemoryHandleArray Handles, FRigVMPredicateBranchArray Predicates);
 };
 
@@ -380,16 +380,16 @@ public:
 	}
 
 	virtual ERigVMOpCode GetOpCode() const override { return ERigVMOpCode::ArrayClone; }
-	RIGVM_API virtual FName GetArgumentNameForOperandIndex(int32 InOperandIndex, int32 InTotalOperands) const override;
-	RIGVM_API virtual const TArray<FRigVMTemplateArgumentInfo>& GetArgumentInfos() const override;
-	RIGVM_API virtual FRigVMTemplateTypeMap OnNewArgumentType(const FName& InArgumentName, TRigVMTypeIndex InTypeIndex) const override;
+	RIGVM_API virtual FName GetArgumentNameForOperandIndex(int32 InOperandIndex, int32 InTotalOperands, FRigVMRegistryHandle& InRegistry) const override;
+	RIGVM_API virtual const TArray<FRigVMTemplateArgumentInfo>& GetArgumentInfos(FRigVMRegistryHandle& InRegistry) const override;
+	RIGVM_API virtual FRigVMTemplateTypeMap OnNewArgumentType(const FName& InArgumentName, TRigVMTypeIndex InTypeIndex, FRigVMRegistryHandle& InRegistry) const override;
 #if WITH_EDITOR
    	RIGVM_API virtual FText GetNodeTooltip(const FRigVMTemplateTypeMap& InTypes) const override;
    	RIGVM_API virtual FText GetArgumentTooltip(const FName& InArgumentName, TRigVMTypeIndex InTypeIndex) const override;
 #endif
 
 protected:
-	virtual FRigVMFunctionPtr GetDispatchFunctionImpl(const FRigVMTemplateTypeMap& InTypes) const override { return &FRigVMDispatch_ArrayClone::Execute; }
+	virtual FRigVMFunctionPtr GetDispatchFunctionImpl(const FRigVMTemplateTypeMap& InTypes, const FRigVMRegistryHandle& InRegistry) const override { return &FRigVMDispatch_ArrayClone::Execute; }
 	static RIGVM_API void Execute(FRigVMExtendedExecuteContext& InContext, FRigVMMemoryHandleArray Handles, FRigVMPredicateBranchArray Predicates);
 };
 
@@ -411,7 +411,7 @@ public:
 #endif
 
 protected:
-	virtual FRigVMFunctionPtr GetDispatchFunctionImpl(const FRigVMTemplateTypeMap& InTypes) const override { return &FRigVMDispatch_ArrayUnion::Execute; }
+	virtual FRigVMFunctionPtr GetDispatchFunctionImpl(const FRigVMTemplateTypeMap& InTypes, const FRigVMRegistryHandle& InRegistry) const override { return &FRigVMDispatch_ArrayUnion::Execute; }
 	static RIGVM_API void Execute(FRigVMExtendedExecuteContext& InContext, FRigVMMemoryHandleArray Handles, FRigVMPredicateBranchArray Predicates);
 };
 
@@ -427,16 +427,16 @@ public:
 	}
 
 	virtual ERigVMOpCode GetOpCode() const override { return ERigVMOpCode::ArrayDifference; }
-	RIGVM_API virtual FName GetArgumentNameForOperandIndex(int32 InOperandIndex, int32 InTotalOperands) const override;
-	RIGVM_API virtual const TArray<FRigVMTemplateArgumentInfo>& GetArgumentInfos() const override;
-	RIGVM_API virtual FRigVMTemplateTypeMap OnNewArgumentType(const FName& InArgumentName, TRigVMTypeIndex InTypeIndex) const override;
+	RIGVM_API virtual FName GetArgumentNameForOperandIndex(int32 InOperandIndex, int32 InTotalOperands, FRigVMRegistryHandle& InRegistry) const override;
+	RIGVM_API virtual const TArray<FRigVMTemplateArgumentInfo>& GetArgumentInfos(FRigVMRegistryHandle& InRegistry) const override;
+	RIGVM_API virtual FRigVMTemplateTypeMap OnNewArgumentType(const FName& InArgumentName, TRigVMTypeIndex InTypeIndex, FRigVMRegistryHandle& InRegistry) const override;
 #if WITH_EDITOR
    	RIGVM_API virtual FText GetNodeTooltip(const FRigVMTemplateTypeMap& InTypes) const override;
    	RIGVM_API virtual FText GetArgumentTooltip(const FName& InArgumentName, TRigVMTypeIndex InTypeIndex) const override;
 #endif
 
 protected:
-	virtual FRigVMFunctionPtr GetDispatchFunctionImpl(const FRigVMTemplateTypeMap& InTypes) const override { return &FRigVMDispatch_ArrayDifference::Execute; }
+	virtual FRigVMFunctionPtr GetDispatchFunctionImpl(const FRigVMTemplateTypeMap& InTypes, const FRigVMRegistryHandle& InRegistry) const override { return &FRigVMDispatch_ArrayDifference::Execute; }
 	static RIGVM_API void Execute(FRigVMExtendedExecuteContext& InContext, FRigVMMemoryHandleArray Handles, FRigVMPredicateBranchArray Predicates);
 };
 
@@ -458,7 +458,7 @@ public:
 #endif
 
 protected:
-	virtual FRigVMFunctionPtr GetDispatchFunctionImpl(const FRigVMTemplateTypeMap& InTypes) const override { return &FRigVMDispatch_ArrayIntersection::Execute; }
+	virtual FRigVMFunctionPtr GetDispatchFunctionImpl(const FRigVMTemplateTypeMap& InTypes, const FRigVMRegistryHandle& InRegistry) const override { return &FRigVMDispatch_ArrayIntersection::Execute; }
 	static RIGVM_API void Execute(FRigVMExtendedExecuteContext& InContext, FRigVMMemoryHandleArray Handles, FRigVMPredicateBranchArray Predicates);
 };
 
@@ -474,18 +474,18 @@ public:
 	}
 
 	virtual ERigVMOpCode GetOpCode() const override { return ERigVMOpCode::ArrayIterator; }
-	RIGVM_API virtual FName GetArgumentNameForOperandIndex(int32 InOperandIndex, int32 InTotalOperands) const override;
-	RIGVM_API virtual const TArray<FRigVMTemplateArgumentInfo>& GetArgumentInfos() const override;
+	RIGVM_API virtual FName GetArgumentNameForOperandIndex(int32 InOperandIndex, int32 InTotalOperands, FRigVMRegistryHandle& InRegistry) const override;
+	RIGVM_API virtual const TArray<FRigVMTemplateArgumentInfo>& GetArgumentInfos(FRigVMRegistryHandle& InRegistry) const override;
 	RIGVM_API virtual const TArray<FRigVMExecuteArgument>& GetExecuteArguments_Impl(const FRigVMDispatchContext& InContext) const override;
 	RIGVM_API virtual const TArray<FName>& GetControlFlowBlocks_Impl(const FRigVMDispatchContext& InContext) const override;
 	RIGVM_API virtual const bool IsControlFlowBlockSliced(const FName& InBlockName) const override;
-	RIGVM_API virtual FRigVMTemplateTypeMap OnNewArgumentType(const FName& InArgumentName, TRigVMTypeIndex InTypeIndex) const override;
+	RIGVM_API virtual FRigVMTemplateTypeMap OnNewArgumentType(const FName& InArgumentName, TRigVMTypeIndex InTypeIndex, FRigVMRegistryHandle& InRegistry) const override;
 #if WITH_EDITOR
    	RIGVM_API virtual FText GetNodeTooltip(const FRigVMTemplateTypeMap& InTypes) const override;
    	RIGVM_API virtual FText GetArgumentTooltip(const FName& InArgumentName, TRigVMTypeIndex InTypeIndex) const override;
 #endif
 
 protected:
-	virtual FRigVMFunctionPtr GetDispatchFunctionImpl(const FRigVMTemplateTypeMap& InTypes) const override { return &FRigVMDispatch_ArrayIterator::Execute; }
+	virtual FRigVMFunctionPtr GetDispatchFunctionImpl(const FRigVMTemplateTypeMap& InTypes, const FRigVMRegistryHandle& InRegistry) const override { return &FRigVMDispatch_ArrayIterator::Execute; }
 	static RIGVM_API void Execute(FRigVMExtendedExecuteContext& InContext, FRigVMMemoryHandleArray Handles, FRigVMPredicateBranchArray Predicates);
 };
diff --git a/Engine/Plugins/Runtime/RigVM/Source/RigVM/Public/RigVMFunctions/RigVMDispatch_CastEnum.h b/Engine/Plugins/Runtime/RigVM/Source/RigVM/Public/RigVMFunctions/RigVMDispatch_CastEnum.h
--- a/Engine/Plugins/Runtime/RigVM/Source/RigVM/Public/RigVMFunctions/RigVMDispatch_CastEnum.h
+++ b/Engine/Plugins/Runtime/RigVM/Source/RigVM/Public/RigVMFunctions/RigVMDispatch_CastEnum.h
@@ -25,8 +25,8 @@ public:
 		FactoryScriptStruct = StaticStruct();
 	}
 	
-	RIGVM_API virtual const TArray<FRigVMTemplateArgumentInfo>& GetArgumentInfos() const override;
-	RIGVM_API virtual bool GetPermutationsFromArgumentType(const FName& InArgumentName, const TRigVMTypeIndex& InTypeIndex, TArray<FRigVMTemplateTypeMap, TInlineAllocator<1>>& OutPermutations) const override;
+	RIGVM_API virtual const TArray<FRigVMTemplateArgumentInfo>& GetArgumentInfos(FRigVMRegistryHandle& InRegistry) const override;
+	RIGVM_API virtual bool GetPermutationsFromArgumentType(const FName& InArgumentName, const TRigVMTypeIndex& InTypeIndex, TArray<FRigVMTemplateTypeMap, TInlineAllocator<1>>& OutPermutations, FRigVMRegistryHandle& InRegistry) const override;
 #if WITH_EDITOR
 	RIGVM_API virtual FString GetNodeTitle(const FRigVMTemplateTypeMap& InTypes) const override;
 	RIGVM_API virtual FText GetNodeTooltip(const FRigVMTemplateTypeMap& InTypes) const override;
@@ -34,7 +34,7 @@ public:
 	virtual bool IsSingleton() const override { return true; }
 	
 protected:
-	virtual FRigVMFunctionPtr GetDispatchFunctionImpl(const FRigVMTemplateTypeMap& InTypes) const override { return &FRigVMDispatch_CastEnumToInt::Execute; }
+	virtual FRigVMFunctionPtr GetDispatchFunctionImpl(const FRigVMTemplateTypeMap& InTypes, const FRigVMRegistryHandle& InRegistry) const override { return &FRigVMDispatch_CastEnumToInt::Execute; }
 	static RIGVM_API void Execute(FRigVMExtendedExecuteContext& InContext, FRigVMMemoryHandleArray Handles, FRigVMPredicateBranchArray RigVMBranches);
 };
 
@@ -52,8 +52,8 @@ public:
 		FactoryScriptStruct = StaticStruct();
 	}
 	
-	RIGVM_API virtual const TArray<FRigVMTemplateArgumentInfo>& GetArgumentInfos() const override;
-	RIGVM_API virtual bool GetPermutationsFromArgumentType(const FName& InArgumentName, const TRigVMTypeIndex& InTypeIndex, TArray<FRigVMTemplateTypeMap, TInlineAllocator<1>>& OutPermutations) const override;
+	RIGVM_API virtual const TArray<FRigVMTemplateArgumentInfo>& GetArgumentInfos(FRigVMRegistryHandle& InRegistry) const override;
+	RIGVM_API virtual bool GetPermutationsFromArgumentType(const FName& InArgumentName, const TRigVMTypeIndex& InTypeIndex, TArray<FRigVMTemplateTypeMap, TInlineAllocator<1>>& OutPermutations, FRigVMRegistryHandle& InRegistry) const override;
 #if WITH_EDITOR
 	RIGVM_API virtual FString GetNodeTitle(const FRigVMTemplateTypeMap& InTypes) const override;
 	RIGVM_API virtual FText GetNodeTooltip(const FRigVMTemplateTypeMap& InTypes) const override;
@@ -61,6 +61,6 @@ public:
 	virtual bool IsSingleton() const override { return true; }
 	
 protected:
-	virtual FRigVMFunctionPtr GetDispatchFunctionImpl(const FRigVMTemplateTypeMap& InTypes) const override { return &FRigVMDispatch_CastIntToEnum::Execute; }
+	virtual FRigVMFunctionPtr GetDispatchFunctionImpl(const FRigVMTemplateTypeMap& InTypes, const FRigVMRegistryHandle& InRegistry) const override { return &FRigVMDispatch_CastIntToEnum::Execute; }
 	static RIGVM_API void Execute(FRigVMExtendedExecuteContext& InContext, FRigVMMemoryHandleArray Handles, FRigVMPredicateBranchArray RigVMBranches);
 };
diff --git a/Engine/Plugins/Runtime/RigVM/Source/RigVM/Public/RigVMFunctions/RigVMDispatch_CastObject.h b/Engine/Plugins/Runtime/RigVM/Source/RigVM/Public/RigVMFunctions/RigVMDispatch_CastObject.h
--- a/Engine/Plugins/Runtime/RigVM/Source/RigVM/Public/RigVMFunctions/RigVMDispatch_CastObject.h
+++ b/Engine/Plugins/Runtime/RigVM/Source/RigVM/Public/RigVMFunctions/RigVMDispatch_CastObject.h
@@ -3,6 +3,7 @@
 #pragma once
 
 #include "RigVMCore/RigVMDispatchFactory.h"
+#include "RigVMCore/RigVMRegistry.h"
 #include "RigVMDispatch_CastObject.generated.h"
 
 USTRUCT(meta=(DisplayName = "Cast", Category = "Object", Keywords = "As", NodeColor = "1,1,1,1"))
@@ -16,8 +17,8 @@ public:
 		FactoryScriptStruct = StaticStruct();
 	}
 	
-	RIGVM_API virtual const TArray<FRigVMTemplateArgumentInfo>& GetArgumentInfos() const override;
-	RIGVM_API virtual bool GetPermutationsFromArgumentType(const FName& InArgumentName, const TRigVMTypeIndex& InTypeIndex, TArray<FRigVMTemplateTypeMap, TInlineAllocator<1>>& OutPermutations) const override;
+	RIGVM_API virtual const TArray<FRigVMTemplateArgumentInfo>& GetArgumentInfos(FRigVMRegistryHandle& InRegistry) const override;
+	RIGVM_API virtual bool GetPermutationsFromArgumentType(const FName& InArgumentName, const TRigVMTypeIndex& InTypeIndex, TArray<FRigVMTemplateTypeMap, TInlineAllocator<1>>& OutPermutations, FRigVMRegistryHandle& InRegistry) const override;
 #if WITH_EDITOR
 	RIGVM_API virtual FString GetNodeTitle(const FRigVMTemplateTypeMap& InTypes) const override;
 	RIGVM_API virtual FText GetNodeTooltip(const FRigVMTemplateTypeMap& InTypes) const override;
@@ -25,7 +26,7 @@ public:
 	virtual bool IsSingleton() const override { return true; }
 	
 protected:
-	virtual FRigVMFunctionPtr GetDispatchFunctionImpl(const FRigVMTemplateTypeMap& InTypes) const override { return &FRigVMDispatch_CastObject::Execute; }
+	virtual FRigVMFunctionPtr GetDispatchFunctionImpl(const FRigVMTemplateTypeMap& InTypes, const FRigVMRegistryHandle& InRegistry) const override { return &FRigVMDispatch_CastObject::Execute; }
 	static RIGVM_API void Execute(FRigVMExtendedExecuteContext& InContext, FRigVMMemoryHandleArray Handles, FRigVMPredicateBranchArray RigVMBranches);
 
 	static inline const FLazyName ValueName = FLazyName(TEXT("Value"));
diff --git a/Engine/Plugins/Runtime/RigVM/Source/RigVM/Public/RigVMFunctions/RigVMDispatch_Constant.h b/Engine/Plugins/Runtime/RigVM/Source/RigVM/Public/RigVMFunctions/RigVMDispatch_Constant.h
--- a/Engine/Plugins/Runtime/RigVM/Source/RigVM/Public/RigVMFunctions/RigVMDispatch_Constant.h
+++ b/Engine/Plugins/Runtime/RigVM/Source/RigVM/Public/RigVMFunctions/RigVMDispatch_Constant.h
@@ -18,8 +18,8 @@ public:
 		FactoryScriptStruct = StaticStruct();
 	}
 
-	RIGVM_API virtual const TArray<FRigVMTemplateArgumentInfo>& GetArgumentInfos() const override;
-	RIGVM_API virtual FRigVMTemplateTypeMap OnNewArgumentType(const FName& InArgumentName, TRigVMTypeIndex InTypeIndex) const override;
+	RIGVM_API virtual const TArray<FRigVMTemplateArgumentInfo>& GetArgumentInfos(FRigVMRegistryHandle& InRegistry) const override;
+	RIGVM_API virtual FRigVMTemplateTypeMap OnNewArgumentType(const FName& InArgumentName, TRigVMTypeIndex InTypeIndex, FRigVMRegistryHandle& InRegistry) const override;
 #if WITH_EDITOR
 	RIGVM_API virtual FString GetNodeTitle(const FRigVMTemplateTypeMap& InTypes) const override;
 	RIGVM_API virtual FText GetNodeTooltip(const FRigVMTemplateTypeMap& InTypes) const override;
@@ -27,7 +27,7 @@ public:
 	virtual bool IsSingleton() const override { return true; }
 	
 protected:
-	virtual FRigVMFunctionPtr GetDispatchFunctionImpl(const FRigVMTemplateTypeMap& InTypes) const override { return &FRigVMDispatch_Constant::Execute; }
+	virtual FRigVMFunctionPtr GetDispatchFunctionImpl(const FRigVMTemplateTypeMap& InTypes, const FRigVMRegistryHandle& InRegistry) const override { return &FRigVMDispatch_Constant::Execute; }
 	static RIGVM_API void Execute(FRigVMExtendedExecuteContext& InContext, FRigVMMemoryHandleArray Handles, FRigVMPredicateBranchArray RigVMBranches);
 
 	static inline const FLazyName ValueName = FLazyName(TEXT("Value"));
diff --git a/Engine/Plugins/Runtime/RigVM/Source/RigVM/Public/RigVMFunctions/RigVMDispatch_Core.h b/Engine/Plugins/Runtime/RigVM/Source/RigVM/Public/RigVMFunctions/RigVMDispatch_Core.h
--- a/Engine/Plugins/Runtime/RigVM/Source/RigVM/Public/RigVMFunctions/RigVMDispatch_Core.h
+++ b/Engine/Plugins/Runtime/RigVM/Source/RigVM/Public/RigVMFunctions/RigVMDispatch_Core.h
@@ -26,13 +26,13 @@ public:
 		FactoryScriptStruct = StaticStruct();
 	}
 
-	RIGVM_API virtual FName GetArgumentNameForOperandIndex(int32 InOperandIndex, int32 InTotalOperands) const override;
-	RIGVM_API virtual const TArray<FRigVMTemplateArgumentInfo>& GetArgumentInfos() const override;
-	RIGVM_API virtual FRigVMTemplateTypeMap OnNewArgumentType(const FName& InArgumentName, TRigVMTypeIndex InTypeIndex) const override;
+	RIGVM_API virtual FName GetArgumentNameForOperandIndex(int32 InOperandIndex, int32 InTotalOperands, FRigVMRegistryHandle& InRegistry) const override;
+	RIGVM_API virtual const TArray<FRigVMTemplateArgumentInfo>& GetArgumentInfos(FRigVMRegistryHandle& InRegistry) const override;
+	RIGVM_API virtual FRigVMTemplateTypeMap OnNewArgumentType(const FName& InArgumentName, TRigVMTypeIndex InTypeIndex, FRigVMRegistryHandle& InRegistry) const override;
 
 protected:
 
-	RIGVM_API virtual FRigVMFunctionPtr GetDispatchFunctionImpl(const FRigVMTemplateTypeMap& InTypes) const override;
+	RIGVM_API virtual FRigVMFunctionPtr GetDispatchFunctionImpl(const FRigVMTemplateTypeMap& InTypes, const FRigVMRegistryHandle& InRegistry) const override;
 	static RIGVM_API bool AdaptResult(bool bResult, const FRigVMExtendedExecuteContext& InContext);
 	static RIGVM_API void Execute(FRigVMExtendedExecuteContext& InContext, FRigVMMemoryHandleArray Handles, FRigVMPredicateBranchArray Predicates);
 
@@ -100,7 +100,7 @@ struct FRigVMDispatch_CoreNotEquals : public FRigVMDispatch_CoreEquals
 		FactoryScriptStruct = StaticStruct();
 	}
 
-	RIGVM_API virtual const TArray<FRigVMTemplateArgumentInfo>& GetArgumentInfos() const override;
+	RIGVM_API virtual const TArray<FRigVMTemplateArgumentInfo>& GetArgumentInfos(FRigVMRegistryHandle& InRegistry) const override;
 
 	// we are inheriting everything from the equals dispatch,
 	// and due to the check of the factory within FRigVMDispatch_CoreEquals::Execute we can
diff --git a/Engine/Plugins/Runtime/RigVM/Source/RigVM/Public/RigVMFunctions/RigVMDispatch_If.h b/Engine/Plugins/Runtime/RigVM/Source/RigVM/Public/RigVMFunctions/RigVMDispatch_If.h
--- a/Engine/Plugins/Runtime/RigVM/Source/RigVM/Public/RigVMFunctions/RigVMDispatch_If.h
+++ b/Engine/Plugins/Runtime/RigVM/Source/RigVM/Public/RigVMFunctions/RigVMDispatch_If.h
@@ -21,9 +21,9 @@ public:
 		FactoryScriptStruct = StaticStruct();
 	}
 
-	RIGVM_API virtual FName GetArgumentNameForOperandIndex(int32 InOperandIndex, int32 InTotalOperands) const override;
-	RIGVM_API virtual const TArray<FRigVMTemplateArgumentInfo>& GetArgumentInfos() const override;
-	RIGVM_API virtual FRigVMTemplateTypeMap OnNewArgumentType(const FName& InArgumentName, TRigVMTypeIndex InTypeIndex) const override;
+	RIGVM_API virtual FName GetArgumentNameForOperandIndex(int32 InOperandIndex, int32 InTotalOperands, FRigVMRegistryHandle& InRegistry) const override;
+	RIGVM_API virtual const TArray<FRigVMTemplateArgumentInfo>& GetArgumentInfos(FRigVMRegistryHandle& InRegistry) const override;
+	RIGVM_API virtual FRigVMTemplateTypeMap OnNewArgumentType(const FName& InArgumentName, TRigVMTypeIndex InTypeIndex, FRigVMRegistryHandle& InRegistry) const override;
 	virtual bool IsSingleton() const override { return true; } 
 
 #if WITH_EDITOR
@@ -32,7 +32,7 @@ public:
 
 protected:
 
-	RIGVM_API virtual FRigVMFunctionPtr GetDispatchFunctionImpl(const FRigVMTemplateTypeMap& InTypes) const override;
+	RIGVM_API virtual FRigVMFunctionPtr GetDispatchFunctionImpl(const FRigVMTemplateTypeMap& InTypes, const FRigVMRegistryHandle& InRegistry) const override;
 	static RIGVM_API void Execute(FRigVMExtendedExecuteContext& InContext, FRigVMMemoryHandleArray Handles, FRigVMPredicateBranchArray Predicates);
 
 	static inline const FLazyName ConditionName = FLazyName(TEXT("Condition"));
diff --git a/Engine/Plugins/Runtime/RigVM/Source/RigVM/Public/RigVMFunctions/RigVMDispatch_MakeStruct.h b/Engine/Plugins/Runtime/RigVM/Source/RigVM/Public/RigVMFunctions/RigVMDispatch_MakeStruct.h
--- a/Engine/Plugins/Runtime/RigVM/Source/RigVM/Public/RigVMFunctions/RigVMDispatch_MakeStruct.h
+++ b/Engine/Plugins/Runtime/RigVM/Source/RigVM/Public/RigVMFunctions/RigVMDispatch_MakeStruct.h
@@ -18,8 +18,8 @@ public:
 		FactoryScriptStruct = StaticStruct();
 	}
 	
-	RIGVM_API virtual const TArray<FRigVMTemplateArgumentInfo>& GetArgumentInfos() const override;
-	RIGVM_API virtual FRigVMTemplateTypeMap OnNewArgumentType(const FName& InArgumentName, TRigVMTypeIndex InTypeIndex) const override;
+	RIGVM_API virtual const TArray<FRigVMTemplateArgumentInfo>& GetArgumentInfos(FRigVMRegistryHandle& InRegistry) const override;
+	RIGVM_API virtual FRigVMTemplateTypeMap OnNewArgumentType(const FName& InArgumentName, TRigVMTypeIndex InTypeIndex, FRigVMRegistryHandle& InRegistry) const override;
 #if WITH_EDITOR
 	RIGVM_API virtual FString GetNodeTitle(const FRigVMTemplateTypeMap& InTypes) const override;
 	RIGVM_API virtual FText GetNodeTooltip(const FRigVMTemplateTypeMap& InTypes) const override;
@@ -28,7 +28,7 @@ public:
 #endif
 
 protected:
-	virtual FRigVMFunctionPtr GetDispatchFunctionImpl(const FRigVMTemplateTypeMap& InTypes) const override { return &FRigVMDispatch_MakeStruct::Execute; }
+	virtual FRigVMFunctionPtr GetDispatchFunctionImpl(const FRigVMTemplateTypeMap& InTypes, const FRigVMRegistryHandle& InRegistry) const override { return &FRigVMDispatch_MakeStruct::Execute; }
 	static RIGVM_API void Execute(FRigVMExtendedExecuteContext& InContext, FRigVMMemoryHandleArray Handles, FRigVMPredicateBranchArray RigVMBranches);
 
 	static inline const FLazyName ElementsName = FLazyName(TEXT("Elements"));
@@ -49,7 +49,7 @@ public:
 		FactoryScriptStruct = StaticStruct();
 	}
 
-	RIGVM_API virtual const TArray<FRigVMTemplateArgumentInfo>& GetArgumentInfos() const override;
+	RIGVM_API virtual const TArray<FRigVMTemplateArgumentInfo>& GetArgumentInfos(FRigVMRegistryHandle& InRegistry) const override;
 #if WITH_EDITOR
 	RIGVM_API virtual FText GetNodeTooltip(const FRigVMTemplateTypeMap& InTypes) const override;
 	RIGVM_API virtual FString GetKeywords() const override;
diff --git a/Engine/Plugins/Runtime/RigVM/Source/RigVM/Public/RigVMFunctions/RigVMDispatch_Print.h b/Engine/Plugins/Runtime/RigVM/Source/RigVM/Public/RigVMFunctions/RigVMDispatch_Print.h
--- a/Engine/Plugins/Runtime/RigVM/Source/RigVM/Public/RigVMFunctions/RigVMDispatch_Print.h
+++ b/Engine/Plugins/Runtime/RigVM/Source/RigVM/Public/RigVMFunctions/RigVMDispatch_Print.h
@@ -20,10 +20,10 @@ public:
 		FactoryScriptStruct = StaticStruct();
 	}
 
-	RIGVM_API virtual FName GetArgumentNameForOperandIndex(int32 InOperandIndex, int32 InTotalOperands) const override;
-	RIGVM_API virtual const TArray<FRigVMTemplateArgumentInfo>& GetArgumentInfos() const override;
+	RIGVM_API virtual FName GetArgumentNameForOperandIndex(int32 InOperandIndex, int32 InTotalOperands, FRigVMRegistryHandle& InRegistry) const override;
+	RIGVM_API virtual const TArray<FRigVMTemplateArgumentInfo>& GetArgumentInfos(FRigVMRegistryHandle& InRegistry) const override;
 	RIGVM_API virtual const TArray<FRigVMExecuteArgument>& GetExecuteArguments_Impl(const FRigVMDispatchContext& InContext) const override;
-	RIGVM_API virtual FRigVMTemplateTypeMap OnNewArgumentType(const FName& InArgumentName, TRigVMTypeIndex InTypeIndex) const override;
+	RIGVM_API virtual FRigVMTemplateTypeMap OnNewArgumentType(const FName& InArgumentName, TRigVMTypeIndex InTypeIndex, FRigVMRegistryHandle& InRegistry) const override;
 	virtual bool IsSingleton() const override { return true; } 
 
 #if WITH_EDITOR
@@ -33,7 +33,7 @@ public:
 
 protected:
 
-	virtual FRigVMFunctionPtr GetDispatchFunctionImpl(const FRigVMTemplateTypeMap& InTypes) const override
+	virtual FRigVMFunctionPtr GetDispatchFunctionImpl(const FRigVMTemplateTypeMap& InTypes, const FRigVMRegistryHandle& InRegistry) const override
 	{
 		return &FRigVMDispatch_Print::Execute;
 	}
diff --git a/Engine/Plugins/Runtime/RigVM/Source/RigVM/Public/RigVMFunctions/RigVMDispatch_Select.h b/Engine/Plugins/Runtime/RigVM/Source/RigVM/Public/RigVMFunctions/RigVMDispatch_Select.h
--- a/Engine/Plugins/Runtime/RigVM/Source/RigVM/Public/RigVMFunctions/RigVMDispatch_Select.h
+++ b/Engine/Plugins/Runtime/RigVM/Source/RigVM/Public/RigVMFunctions/RigVMDispatch_Select.h
@@ -21,9 +21,9 @@ public:
 		FactoryScriptStruct = StaticStruct();
 	}
 
-	RIGVM_API virtual FName GetArgumentNameForOperandIndex(int32 InOperandIndex, int32 InTotalOperands) const override;
-	RIGVM_API virtual const TArray<FRigVMTemplateArgumentInfo>& GetArgumentInfos() const override;
-	RIGVM_API virtual FRigVMTemplateTypeMap OnNewArgumentType(const FName& InArgumentName, TRigVMTypeIndex InTypeIndex) const override;
+	RIGVM_API virtual FName GetArgumentNameForOperandIndex(int32 InOperandIndex, int32 InTotalOperands, FRigVMRegistryHandle& InRegistry) const override;
+	RIGVM_API virtual const TArray<FRigVMTemplateArgumentInfo>& GetArgumentInfos(FRigVMRegistryHandle& InRegistry) const override;
+	RIGVM_API virtual FRigVMTemplateTypeMap OnNewArgumentType(const FName& InArgumentName, TRigVMTypeIndex InTypeIndex, FRigVMRegistryHandle& InRegistry) const override;
 #if WITH_EDITOR
 	RIGVM_API virtual FString GetArgumentMetaData(const FName& InArgumentName, const FName& InMetaDataKey) const override;
 	RIGVM_API virtual FString GetArgumentDefaultValue(const FName& InArgumentName, TRigVMTypeIndex InTypeIndex) const override;
@@ -32,7 +32,7 @@ public:
 
 protected:
 
-	RIGVM_API virtual FRigVMFunctionPtr GetDispatchFunctionImpl(const FRigVMTemplateTypeMap& InTypes) const override;
+	RIGVM_API virtual FRigVMFunctionPtr GetDispatchFunctionImpl(const FRigVMTemplateTypeMap& InTypes, const FRigVMRegistryHandle& InRegistry) const override;
 	static RIGVM_API void Execute(FRigVMExtendedExecuteContext& InContext, FRigVMMemoryHandleArray Handles, FRigVMPredicateBranchArray Predicates);
 
 	static inline FLazyName IndexName = FLazyName(TEXT("Index"));
