RigVM: Localized registry in cooked environments - first steps.

diff --git a/Engine/Plugins/Runtime/RigVM/Source/RigVM/Public/RigVMFunctions/RigVMDispatch_Switch.h b/Engine/Plugins/Runtime/RigVM/Source/RigVM/Public/RigVMFunctions/RigVMDispatch_Switch.h
--- a/Engine/Plugins/Runtime/RigVM/Source/RigVM/Public/RigVMFunctions/RigVMDispatch_Switch.h
+++ b/Engine/Plugins/Runtime/RigVM/Source/RigVM/Public/RigVMFunctions/RigVMDispatch_Switch.h
@@ -21,8 +21,8 @@ public:
 		FactoryScriptStruct = StaticStruct();
 	}
 
-	RIGVM_API virtual FName GetArgumentNameForOperandIndex(int32 InOperandIndex, int32 InTotalOperands) const override;
-	RIGVM_API virtual const TArray<FRigVMTemplateArgumentInfo>& GetArgumentInfos() const override;
+	RIGVM_API virtual FName GetArgumentNameForOperandIndex(int32 InOperandIndex, int32 InTotalOperands, FRigVMRegistryHandle& InRegistry) const override;
+	RIGVM_API virtual const TArray<FRigVMTemplateArgumentInfo>& GetArgumentInfos(FRigVMRegistryHandle& InRegistry) const override;
 	RIGVM_API virtual const TArray<FRigVMExecuteArgument>& GetExecuteArguments_Impl(const FRigVMDispatchContext& InContext) const override;
 	virtual bool IsSingleton() const override { return true; } 
 
@@ -36,7 +36,7 @@ public:
 
 protected:
 
-	RIGVM_API virtual FRigVMFunctionPtr GetDispatchFunctionImpl(const FRigVMTemplateTypeMap& InTypes) const override;
+	RIGVM_API virtual FRigVMFunctionPtr GetDispatchFunctionImpl(const FRigVMTemplateTypeMap& InTypes, const FRigVMRegistryHandle& InRegistry) const override;
 	static RIGVM_API void Execute(FRigVMExtendedExecuteContext& InContext, FRigVMMemoryHandleArray Handles, FRigVMPredicateBranchArray Predicates);
 	static RIGVM_API FName GetCaseName(int32 InIndex);
 	static RIGVM_API FName GetCaseDisplayName(int32 InIndex);
diff --git a/Engine/Plugins/Runtime/RigVM/Source/RigVM/Public/RigVMFunctions/RigVMFunction_String.h b/Engine/Plugins/Runtime/RigVM/Source/RigVM/Public/RigVMFunctions/RigVMFunction_String.h
--- a/Engine/Plugins/Runtime/RigVM/Source/RigVM/Public/RigVMFunctions/RigVMFunction_String.h
+++ b/Engine/Plugins/Runtime/RigVM/Source/RigVM/Public/RigVMFunctions/RigVMFunction_String.h
@@ -510,12 +510,12 @@ struct FRigDispatch_ToString : public FRigVMDispatchFactory
 
 public:
 
-	virtual const TArray<FRigVMTemplateArgumentInfo>& GetArgumentInfos() const override;
-	virtual FRigVMTemplateTypeMap OnNewArgumentType(const FName& InArgumentName, TRigVMTypeIndex InTypeIndex) const override;
+	virtual const TArray<FRigVMTemplateArgumentInfo>& GetArgumentInfos(FRigVMRegistryHandle& InRegistry) const override;
+	virtual FRigVMTemplateTypeMap OnNewArgumentType(const FName& InArgumentName, TRigVMTypeIndex InTypeIndex, FRigVMRegistryHandle& InRegistry) const override;
 
 protected:
 
-	virtual FRigVMFunctionPtr GetDispatchFunctionImpl(const FRigVMTemplateTypeMap& InTypes) const override;
+	virtual FRigVMFunctionPtr GetDispatchFunctionImpl(const FRigVMTemplateTypeMap& InTypes, const FRigVMRegistryHandle& InRegistry) const override;
 	static void Execute(FRigVMExtendedExecuteContext& InContext, FRigVMMemoryHandleArray Handles, FRigVMPredicateBranchArray Predicates);
 };
 
@@ -529,11 +529,11 @@ struct FRigDispatch_FromString : public FRigVMDispatchFactory
 
 public:
 
-	virtual const TArray<FRigVMTemplateArgumentInfo>& GetArgumentInfos() const override;
-	virtual FRigVMTemplateTypeMap OnNewArgumentType(const FName& InArgumentName, TRigVMTypeIndex InTypeIndex) const override;
+	virtual const TArray<FRigVMTemplateArgumentInfo>& GetArgumentInfos(FRigVMRegistryHandle& InRegistry) const override;
+	virtual FRigVMTemplateTypeMap OnNewArgumentType(const FName& InArgumentName, TRigVMTypeIndex InTypeIndex, FRigVMRegistryHandle& InRegistry) const override;
 
 protected:
 
-	virtual FRigVMFunctionPtr GetDispatchFunctionImpl(const FRigVMTemplateTypeMap& InTypes) const override;
+	virtual FRigVMFunctionPtr GetDispatchFunctionImpl(const FRigVMTemplateTypeMap& InTypes, const FRigVMRegistryHandle& InRegistry) const override;
 	static void Execute(FRigVMExtendedExecuteContext& InContext, FRigVMMemoryHandleArray Handles, FRigVMPredicateBranchArray Predicates);
 };
\ No newline at end of file
diff --git a/Engine/Plugins/Runtime/RigVM/Source/RigVM/Public/RigVMObjectVersion.h b/Engine/Plugins/Runtime/RigVM/Source/RigVM/Public/RigVMObjectVersion.h
--- a/Engine/Plugins/Runtime/RigVM/Source/RigVM/Public/RigVMObjectVersion.h
+++ b/Engine/Plugins/Runtime/RigVM/Source/RigVM/Public/RigVMObjectVersion.h
@@ -69,6 +69,9 @@ struct FRigVMObjectVersion
 		// from RigVMCopyOp
 		ByteCodeCleanup,
 
+		// The VM stores a local snapshot registry to use in cooked environments instead of the shared global registry
+		LocalizedRegistry,
+
 		// -----<new versions can be added above this line>-------------------------------------------------
 		VersionPlusOne,
 		LatestVersion = VersionPlusOne - 1
diff --git a/Engine/Plugins/Runtime/RigVM/Source/RigVMDeveloper/Private/RigVMCompiler/RigVMCodeGenerator.cpp b/Engine/Plugins/Runtime/RigVM/Source/RigVMDeveloper/Private/RigVMCompiler/RigVMCodeGenerator.cpp
--- a/Engine/Plugins/Runtime/RigVM/Source/RigVMDeveloper/Private/RigVMCompiler/RigVMCodeGenerator.cpp
+++ b/Engine/Plugins/Runtime/RigVM/Source/RigVMDeveloper/Private/RigVMCompiler/RigVMCodeGenerator.cpp
@@ -1237,7 +1237,7 @@ void FRigVMCodeGenerator::ParseRequiredUProperties(const FRigVMExtendedExecuteCo
 				}
 
 				FRigVMDispatchContext DispatchContext;
-				if(URigVMDispatchNode* DispatchNode = Cast<URigVMDispatchNode>(ByteCode.GetSubjectForInstruction(InstructionIndex)))
+				if(const URigVMDispatchNode* DispatchNode = Cast<URigVMDispatchNode>(ByteCode.GetSubjectForInstruction(InstructionIndex)))
 				{
 					DispatchContext = DispatchNode->GetDispatchContext();
 				}
diff --git a/Engine/Plugins/Runtime/RigVM/Source/RigVMDeveloper/Private/RigVMCompiler/RigVMCompiler.cpp b/Engine/Plugins/Runtime/RigVM/Source/RigVMDeveloper/Private/RigVMCompiler/RigVMCompiler.cpp
--- a/Engine/Plugins/Runtime/RigVM/Source/RigVMDeveloper/Private/RigVMCompiler/RigVMCompiler.cpp
+++ b/Engine/Plugins/Runtime/RigVM/Source/RigVMDeveloper/Private/RigVMCompiler/RigVMCompiler.cpp
@@ -1547,6 +1547,12 @@ bool URigVMCompiler::Compile(const FRigVMCompileSettings& InSettings, TArray<URi
 		OutFunctionCompilationData->OperandToDebugRegisters = WorkData.VM->OperandToDebugRegisters;
 	}
 
+	// Localize all of the registry information
+	if(CurrentCompilationFunction == nullptr)
+	{
+		WorkData.VM->CreateLocalizedRegistryIfRequired();
+	}
+
 	// make sure all functions are known and resolved now.
 	WorkData.VM->ResolveFunctionsIfRequired();
 	
diff --git a/Engine/Plugins/Runtime/RigVM/Source/RigVMDeveloper/Private/RigVMDeveloperTypeUtils.cpp b/Engine/Plugins/Runtime/RigVM/Source/RigVMDeveloper/Private/RigVMDeveloperTypeUtils.cpp
--- a/Engine/Plugins/Runtime/RigVM/Source/RigVMDeveloper/Private/RigVMDeveloperTypeUtils.cpp
+++ b/Engine/Plugins/Runtime/RigVM/Source/RigVMDeveloper/Private/RigVMDeveloperTypeUtils.cpp
@@ -617,7 +617,7 @@ const TArray<TRigVMTypeIndex>& RigVMTypeUtils::GetAvailableCasts(const TRigVMTyp
 				}
 				Index++;
 				return true;
-			});
+			}, Registry.GetHandle_NoLock());
 
 			return AvailableCasts;
 		}
diff --git a/Engine/Plugins/Runtime/RigVM/Source/RigVMDeveloper/Private/RigVMModel/Nodes/RigVMIfNode.cpp b/Engine/Plugins/Runtime/RigVM/Source/RigVMDeveloper/Private/RigVMModel/Nodes/RigVMIfNode.cpp
--- a/Engine/Plugins/Runtime/RigVM/Source/RigVMDeveloper/Private/RigVMModel/Nodes/RigVMIfNode.cpp
+++ b/Engine/Plugins/Runtime/RigVM/Source/RigVMDeveloper/Private/RigVMModel/Nodes/RigVMIfNode.cpp
@@ -31,23 +31,21 @@ const FRigVMTemplate* UDEPRECATED_RigVMIfNode::GetTemplate() const
 		static const FLazyName FalseFName(FalseName);
 		static const FLazyName ResultFName(ResultName);
 
-		static TArray<FRigVMTemplateArgumentInfo> Infos;
-		if(Infos.IsEmpty())
-		{
-			static const TArray<FRigVMTemplateArgument::ETypeCategory> Categories = {
-				FRigVMTemplateArgument::ETypeCategory_SingleAnyValue,
-				FRigVMTemplateArgument::ETypeCategory_ArrayAnyValue
-			};
-			Infos.Reserve(4);
-			Infos.Emplace(ConditionFName, ERigVMPinDirection::Input, RigVMTypeUtils::TypeIndex::Bool);
-			Infos.Emplace(TrueFName, ERigVMPinDirection::Input, Categories);
-			Infos.Emplace(FalseFName, ERigVMPinDirection::Input, Categories);
-			Infos.Emplace(ResultFName, ERigVMPinDirection::Output, Categories);
-		}
+		static const TArray<FRigVMTemplateArgument::ETypeCategory> Categories = {
+			FRigVMTemplateArgument::ETypeCategory_SingleAnyValue,
+			FRigVMTemplateArgument::ETypeCategory_ArrayAnyValue
+		};
+
+		TArray<FRigVMTemplateArgumentInfo> ArgumentInfos;
+		ArgumentInfos.Reserve(4);
+		ArgumentInfos.Emplace(ConditionFName, ERigVMPinDirection::Input, RigVMTypeUtils::TypeIndex::Bool);
+		ArgumentInfos.Emplace(TrueFName, ERigVMPinDirection::Input, Categories);
+		ArgumentInfos.Emplace(FalseFName, ERigVMPinDirection::Input, Categories);
+		ArgumentInfos.Emplace(ResultFName, ERigVMPinDirection::Output, Categories);
 
 		FRigVMTemplateDelegates Delegates;
 		Delegates.NewArgumentTypeDelegate = 
-			FRigVMTemplate_NewArgumentTypeDelegate::CreateLambda([](const FName& InArgumentName, int32 InTypeIndex)
+			FRigVMTemplate_NewArgumentTypeDelegate::CreateLambda([](const FName& InArgumentName, int32 InTypeIndex, const FRigVMRegistryHandle& InRegistry)
 			{
 				FRigVMTemplateTypeMap Types;
 
@@ -62,7 +60,7 @@ const FRigVMTemplate* UDEPRECATED_RigVMIfNode::GetTemplate() const
 				return Types;
 			});
 
-		IfNodeTemplate = CachedTemplate = FRigVMRegistry::Get().GetOrAddTemplateFromArguments(IfName, Infos, Delegates);
+		IfNodeTemplate = CachedTemplate = FRigVMRegistry::Get().GetOrAddTemplateFromArguments(IfName, ArgumentInfos, Delegates);
 	}
 	return CachedTemplate;
 }
diff --git a/Engine/Plugins/Runtime/RigVM/Source/RigVMDeveloper/Private/RigVMModel/Nodes/RigVMSelectNode.cpp b/Engine/Plugins/Runtime/RigVM/Source/RigVMDeveloper/Private/RigVMModel/Nodes/RigVMSelectNode.cpp
--- a/Engine/Plugins/Runtime/RigVM/Source/RigVMDeveloper/Private/RigVMModel/Nodes/RigVMSelectNode.cpp
+++ b/Engine/Plugins/Runtime/RigVM/Source/RigVMDeveloper/Private/RigVMModel/Nodes/RigVMSelectNode.cpp
@@ -43,26 +43,24 @@ const FRigVMTemplate* UDEPRECATED_RigVMSelectNode::GetTemplate() const
 		static const FLazyName ValueFName(ValueName);
 		static const FLazyName ResultFName(ResultName);
 
-		static TArray<FRigVMTemplateArgumentInfo> Infos;
-		if(Infos.IsEmpty())
-		{
-			static const TArray<FRigVMTemplateArgument::ETypeCategory> ValueTypeCategories = {
-				FRigVMTemplateArgument::ETypeCategory_ArrayAnyValue,
-				FRigVMTemplateArgument::ETypeCategory_ArrayArrayAnyValue
-			};
-			static const TArray<FRigVMTemplateArgument::ETypeCategory> ResultTypeCategories = {
-				FRigVMTemplateArgument::ETypeCategory_SingleAnyValue,
-				FRigVMTemplateArgument::ETypeCategory_ArrayAnyValue
-			};
-			Infos.Reserve(3);
-			Infos.Emplace(IndexFName, ERigVMPinDirection::Input, RigVMTypeUtils::TypeIndex::Int32);
-			Infos.Emplace(ValueFName, ERigVMPinDirection::Input, ValueTypeCategories);
-			Infos.Emplace(ResultFName, ERigVMPinDirection::Output, ResultTypeCategories);
-		}
+		static const TArray<FRigVMTemplateArgument::ETypeCategory> ValueTypeCategories = {
+			FRigVMTemplateArgument::ETypeCategory_ArrayAnyValue,
+			FRigVMTemplateArgument::ETypeCategory_ArrayArrayAnyValue
+		};
+		static const TArray<FRigVMTemplateArgument::ETypeCategory> ResultTypeCategories = {
+			FRigVMTemplateArgument::ETypeCategory_SingleAnyValue,
+			FRigVMTemplateArgument::ETypeCategory_ArrayAnyValue
+		};
+		
+		TArray<FRigVMTemplateArgumentInfo> ArgumentInfos;
+		ArgumentInfos.Reserve(3);
+		ArgumentInfos.Emplace(IndexFName, ERigVMPinDirection::Input, RigVMTypeUtils::TypeIndex::Int32);
+		ArgumentInfos.Emplace(ValueFName, ERigVMPinDirection::Input, ValueTypeCategories);
+		ArgumentInfos.Emplace(ResultFName, ERigVMPinDirection::Output, ResultTypeCategories);
 		
 		FRigVMTemplateDelegates Delegates;
 		Delegates.NewArgumentTypeDelegate = 
-			FRigVMTemplate_NewArgumentTypeDelegate::CreateLambda([](const FName& InArgumentName, int32 InTypeIndex)
+			FRigVMTemplate_NewArgumentTypeDelegate::CreateLambda([](const FName& InArgumentName, int32 InTypeIndex, const FRigVMRegistryHandle& InRegistry)
 			{
 				FRigVMTemplateTypeMap Types;
 
@@ -72,11 +70,11 @@ const FRigVMTemplate* UDEPRECATED_RigVMSelectNode::GetTemplate() const
 				if(InArgumentName == ValueFName)
 				{
 					ValueTypeIndex = InTypeIndex;
-					ResultTypeIndex = FRigVMRegistry_NoLock::GetForRead().GetBaseTypeFromArrayTypeIndex_NoLock(InTypeIndex);
+					ResultTypeIndex = FRigVMRegistry::Get().GetBaseTypeFromArrayTypeIndex(InTypeIndex);
 				}
 				else if(InArgumentName == ResultFName)
 				{
-					ValueTypeIndex = FRigVMRegistry_NoLock::GetForRead().GetArrayTypeFromBaseTypeIndex_NoLock(InTypeIndex);;
+					ValueTypeIndex = FRigVMRegistry::Get().GetArrayTypeFromBaseTypeIndex(InTypeIndex);;
 					ResultTypeIndex = InTypeIndex;
 				}
 				
@@ -90,7 +88,7 @@ const FRigVMTemplate* UDEPRECATED_RigVMSelectNode::GetTemplate() const
 				return Types;
 			});
 
-		SelectNodeTemplate = CachedTemplate = FRigVMRegistry::Get().GetOrAddTemplateFromArguments(SelectName, Infos, Delegates);
+		SelectNodeTemplate = CachedTemplate = FRigVMRegistry::Get().GetOrAddTemplateFromArguments(SelectName, ArgumentInfos, Delegates);
 	}
 	return CachedTemplate;
 }
diff --git a/Engine/Plugins/Runtime/RigVM/Source/RigVMDeveloper/Private/RigVMModel/Nodes/RigVMTemplateNode.cpp b/Engine/Plugins/Runtime/RigVM/Source/RigVMDeveloper/Private/RigVMModel/Nodes/RigVMTemplateNode.cpp
--- a/Engine/Plugins/Runtime/RigVM/Source/RigVMDeveloper/Private/RigVMModel/Nodes/RigVMTemplateNode.cpp
+++ b/Engine/Plugins/Runtime/RigVM/Source/RigVMDeveloper/Private/RigVMModel/Nodes/RigVMTemplateNode.cpp
@@ -196,7 +196,7 @@ FText URigVMTemplateNode::GetToolTipTextForPin(const URigVMPin* InPin) const
 			{
 				FString Tooltip;
 				
-				if(Arg->GetNumTypes() > 100)
+				if(Arg->GetNumTypes(FRigVMRegistry::Get().GetHandle_NoLock()) > UE_RIGVM_MANYTYPES_LIMIT)
 				{
 					Tooltip = TEXT("Supports any type.");
 				}
diff --git a/Engine/Plugins/Runtime/RigVM/Source/RigVMDeveloper/Private/RigVMModel/RigVMController.cpp b/Engine/Plugins/Runtime/RigVM/Source/RigVMDeveloper/Private/RigVMModel/RigVMController.cpp
--- a/Engine/Plugins/Runtime/RigVM/Source/RigVMDeveloper/Private/RigVMModel/RigVMController.cpp
+++ b/Engine/Plugins/Runtime/RigVM/Source/RigVMDeveloper/Private/RigVMModel/RigVMController.cpp
@@ -22429,7 +22429,8 @@ void URigVMController::RegisterUseOfTemplate(const URigVMTemplateNode* InNode)
 	const FName& Notation = Template->GetNotation();
 	FRigVMController_CommonTypePerTemplate& TypesForTemplate = Settings->TemplateDefaultTypes.FindOrAdd(Notation);
 
-	const FString TypesString = FRigVMTemplate::GetStringFromArgumentTypes(Template->GetTypesForPermutation(ResolvedPermutation));
+	FRigVMRegistryReadLock ReadLock;
+	const FString TypesString = FRigVMTemplate::GetStringFromArgumentTypes(Template->GetTypesForPermutation(ResolvedPermutation), ReadLock);
 	int32& Count = TypesForTemplate.Counts.FindOrAdd(TypesString);
 	Count++;
 }
diff --git a/Engine/Plugins/Runtime/RigVM/Source/RigVMDeveloper/Public/RigVMModel/Nodes/RigVMTemplateNode.h b/Engine/Plugins/Runtime/RigVM/Source/RigVMDeveloper/Public/RigVMModel/Nodes/RigVMTemplateNode.h
--- a/Engine/Plugins/Runtime/RigVM/Source/RigVMDeveloper/Public/RigVMModel/Nodes/RigVMTemplateNode.h
+++ b/Engine/Plugins/Runtime/RigVM/Source/RigVMDeveloper/Public/RigVMModel/Nodes/RigVMTemplateNode.h
@@ -7,6 +7,7 @@
 #include "RigVMTemplateNode.generated.h"
 
 #define UE_API RIGVMDEVELOPER_API
+#define UE_RIGVM_MANYTYPES_LIMIT 100
 
 USTRUCT()
 struct FRigVMTemplatePreferredType
