Control Rig: Enable user provided names for constant nodes

#rnx

diff --git a/Engine/Plugins/Animation/ControlRig/Source/ControlRig/Private/Units/Hierarchy/RigUnit_BoneName.cpp b/Engine/Plugins/Animation/ControlRig/Source/ControlRig/Private/Units/Hierarchy/RigUnit_BoneName.cpp
--- a/Engine/Plugins/Animation/ControlRig/Source/ControlRig/Private/Units/Hierarchy/RigUnit_BoneName.cpp
+++ b/Engine/Plugins/Animation/ControlRig/Source/ControlRig/Private/Units/Hierarchy/RigUnit_BoneName.cpp
@@ -1,6 +1,8 @@
 // Copyright Epic Games, Inc. All Rights Reserved.
 
 #include "Units/Hierarchy/RigUnit_BoneName.h"
+
+#include "RigVMFunctions/RigVMDispatch_Constant.h"
 #include "Units/RigUnitContext.h"
 
 #include UE_INLINE_GENERATED_CPP_BY_NAME(RigUnit_BoneName)
@@ -13,6 +15,13 @@ FRigUnit_ItemArray_Execute()
 {
 }
 
+FRigVMStructUpgradeInfo FRigUnit_ItemArray::GetUpgradeInfo() const
+{
+	FRigVMStructUpgradeInfo Info = FRigVMStructUpgradeInfo::MakeFromStructToFactory(StaticStruct(), FRigVMDispatch_Constant::StaticStruct());
+	Info.AddRemappedPin(TEXT("Items"), TEXT("Value"));
+	return Info;
+}
+
 FRigUnit_BoneName_Execute()
 {
 }
diff --git a/Engine/Plugins/Animation/ControlRig/Source/ControlRig/Public/Units/Hierarchy/RigUnit_BoneName.h b/Engine/Plugins/Animation/ControlRig/Source/ControlRig/Public/Units/Hierarchy/RigUnit_BoneName.h
--- a/Engine/Plugins/Animation/ControlRig/Source/ControlRig/Public/Units/Hierarchy/RigUnit_BoneName.h
+++ b/Engine/Plugins/Animation/ControlRig/Source/ControlRig/Public/Units/Hierarchy/RigUnit_BoneName.h
@@ -31,7 +31,7 @@ struct FRigUnit_Item : public FRigUnit
 /**
  * The Item Array node is used to share an array of items across the graph
  */
-USTRUCT(meta=(DisplayName="Item Array", Category="Hierarchy", NodeColor = "0.4627450108528137 1.0 0.3294120132923126", DocumentationPolicy = "Strict", Constant))
+USTRUCT(meta=(DisplayName="Item Array", Category="Hierarchy", NodeColor = "0.4627450108528137 1.0 0.3294120132923126", DocumentationPolicy = "Strict", Constant, Deprecated = "5.7"))
 struct FRigUnit_ItemArray : public FRigUnit
 {
 	GENERATED_BODY()
@@ -47,6 +47,9 @@ struct FRigUnit_ItemArray : public FRigUnit
 	 */
 	UPROPERTY(meta = (Input, Output, ExpandByDefault))
 	TArray<FRigElementKey> Items;
+
+	RIGVM_METHOD()
+	UE_API virtual FRigVMStructUpgradeInfo GetUpgradeInfo() const override;
 };
 
 /**
diff --git a/Engine/Plugins/Animation/ControlRig/Source/ControlRigEditor/Private/Editor/ControlRigEditor.cpp b/Engine/Plugins/Animation/ControlRig/Source/ControlRigEditor/Private/Editor/ControlRigEditor.cpp
--- a/Engine/Plugins/Animation/ControlRig/Source/ControlRigEditor/Private/Editor/ControlRigEditor.cpp
+++ b/Engine/Plugins/Animation/ControlRig/Source/ControlRigEditor/Private/Editor/ControlRigEditor.cpp
@@ -125,6 +125,7 @@
 #include "PropertyPath.h"
 #include "Editor/SRigVMDetailsInspector.h"
 #include "Overrides/OverrideStatusDetailsObjectFilter.h"
+#include "RigVMFunctions/RigVMDispatch_Constant.h"
 
 #if WITH_RIGVMLEGACYEDITOR
 #include "SKismetInspector.h"
@@ -4721,6 +4722,7 @@ void FControlRigBaseEditor::HandleModularRigModified(EModularRigNotification InN
 			break;
 		}
 		case EModularRigNotification::ModuleReparented:
+		case EModularRigNotification::ModuleReordered:
 		case EModularRigNotification::ModuleRenamed:
 		{
 			break;
@@ -5318,24 +5320,40 @@ void FControlRigBaseEditor::CreateRigHierarchyToGraphDragAndDropMenu() const
 											{
 												Controller->OpenUndoBracket(TEXT("Create Item Array From Selection"));
 
-												if (URigVMNode* ItemsNode = Controller->AddUnitNode(FRigUnit_ItemArray::StaticStruct(), TEXT("Execute"), NodePosition))
+												const FRigVMDispatchFactory* ConstantFactory = FRigVMRegistry::Get().FindOrAddDispatchFactory(FRigVMDispatch_Constant::StaticStruct());
+												check(ConstantFactory);
+												
+												if (URigVMNode* ItemsNode = Controller->AddTemplateNode(ConstantFactory->GetTemplateNotation(), NodePosition, TEXT("ItemArray"), true, true))
 												{
-													if (URigVMPin* ItemsPin = ItemsNode->FindPin(TEXT("Items")))
+													if (URigVMPin* ItemsPin = ItemsNode->FindPin(TEXT("Value")))
 													{
+														const TRigVMTypeIndex RigElementKeyType = FRigVMRegistry::Get().FindOrAddType(FRigElementKey::StaticStruct());
+														const TRigVMTypeIndex RigElementKeyArrayType = FRigVMRegistry::Get().GetArrayTypeFromBaseTypeIndex(RigElementKeyType);
+														Controller->ResolveWildCardPin(ItemsPin->GetPinPath(), RigElementKeyArrayType, true, true);
 														Controller->SetArrayPinSize(ItemsPin->GetPinPath(), DraggedElements.Num());
 
 														TArray<URigVMPin*> ItemPins = ItemsPin->GetSubPins();
-														ensure(ItemPins.Num() == DraggedElements.Num());
-
-														for (int32 ItemIndex = 0; ItemIndex < DraggedElements.Num(); ItemIndex++)
+														if(ensure(ItemPins.Num() == DraggedElements.Num()))
 														{
-															FString DefaultValue;
-															FRigElementKey::StaticStruct()->ExportText(DefaultValue, &DraggedElements[ItemIndex], nullptr, nullptr, PPF_None, nullptr);
-															Controller->SetPinDefaultValue(ItemPins[ItemIndex]->GetPinPath(), DefaultValue, true, true, false, true);
-															Controller->SetPinExpansion(ItemPins[ItemIndex]->GetPinPath(), true, true, true);
+															for (int32 ItemIndex = 0; ItemIndex < DraggedElements.Num(); ItemIndex++)
+															{
+																FString DefaultValue;
+																FRigElementKey::StaticStruct()->ExportText(DefaultValue, &DraggedElements[ItemIndex], nullptr, nullptr, PPF_None, nullptr);
+																Controller->SetPinDefaultValue(ItemPins[ItemIndex]->GetPinPath(), DefaultValue, true, true, false, true);
+																Controller->SetPinExpansion(ItemPins[ItemIndex]->GetPinPath(), true, true, true);
+															}
+															Controller->SetPinExpansion(ItemsPin->GetPinPath(), true, true, true);
 														}
 
-														Controller->SetPinExpansion(ItemsPin->GetPinPath(), true, true, true);
+														Controller->SelectNodeByName(ItemsNode->GetFName());
+													}
+
+													if(URigVMEdGraph* EdGraph = Cast<URigVMEdGraph>(ControlRigEditor->GetControlRigBlueprint()->GetEdGraph(ItemsNode->GetGraph())))
+													{
+														if(URigVMEdGraphNode* EdGraphNode = Cast<URigVMEdGraphNode>(EdGraph->FindNodeForModelNodeName(ItemsNode->GetFName())))
+														{
+															EdGraphNode->RequestRename();
+														}
 													}
 												}
 
diff --git a/Engine/Plugins/Runtime/RigVM/Source/RigVM/Private/RigVMFunctions/Math/RigVMFunction_MathBool.cpp b/Engine/Plugins/Runtime/RigVM/Source/RigVM/Private/RigVMFunctions/Math/RigVMFunction_MathBool.cpp
--- a/Engine/Plugins/Runtime/RigVM/Source/RigVM/Private/RigVMFunctions/Math/RigVMFunction_MathBool.cpp
+++ b/Engine/Plugins/Runtime/RigVM/Source/RigVM/Private/RigVMFunctions/Math/RigVMFunction_MathBool.cpp
@@ -1,6 +1,7 @@
 // Copyright Epic Games, Inc. All Rights Reserved.
 
 #include "RigVMFunctions/Math/RigVMFunction_MathBool.h"
+#include "RigVMFunctions/RigVMDispatch_Constant.h"
 #include "RigVMFunctions/RigVMDispatch_Core.h"
 
 #include UE_INLINE_GENERATED_CPP_BY_NAME(RigVMFunction_MathBool)
@@ -9,6 +10,11 @@ FRigVMFunction_MathBoolMake_Execute()
 {
 }
 
+FRigVMStructUpgradeInfo FRigVMFunction_MathBoolMake::GetUpgradeInfo() const
+{
+	return FRigVMStructUpgradeInfo::MakeFromStructToFactory(StaticStruct(), FRigVMDispatch_Constant::StaticStruct());
+}
+
 FRigVMFunction_MathBoolConstTrue_Execute()
 {
 	Value = true;
diff --git a/Engine/Plugins/Runtime/RigVM/Source/RigVM/Private/RigVMFunctions/Math/RigVMFunction_MathColor.cpp b/Engine/Plugins/Runtime/RigVM/Source/RigVM/Private/RigVMFunctions/Math/RigVMFunction_MathColor.cpp
--- a/Engine/Plugins/Runtime/RigVM/Source/RigVM/Private/RigVMFunctions/Math/RigVMFunction_MathColor.cpp
+++ b/Engine/Plugins/Runtime/RigVM/Source/RigVM/Private/RigVMFunctions/Math/RigVMFunction_MathColor.cpp
@@ -1,6 +1,7 @@
 // Copyright Epic Games, Inc. All Rights Reserved.
 
 #include "RigVMFunctions/Math/RigVMFunction_MathColor.h"
+#include "RigVMFunctions/RigVMDispatch_Constant.h"
 
 #include UE_INLINE_GENERATED_CPP_BY_NAME(RigVMFunction_MathColor)
 
@@ -9,6 +10,17 @@ FRigVMFunction_MathColorMake_Execute()
 	Result = FLinearColor(R, G, B, A);
 }
 
+FRigVMStructUpgradeInfo FRigVMFunction_MathColorMake::GetUpgradeInfo() const
+{
+	FRigVMStructUpgradeInfo Info = FRigVMStructUpgradeInfo::MakeFromStructToFactory(StaticStruct(), FRigVMDispatch_Constant::StaticStruct());
+	Info.AddRemappedPin(TEXT("R"), TEXT("Value.R"));
+	Info.AddRemappedPin(TEXT("G"), TEXT("Value.G"));
+	Info.AddRemappedPin(TEXT("B"), TEXT("Value.B"));
+	Info.AddRemappedPin(TEXT("A"), TEXT("Value.A"));
+	Info.AddRemappedPin(TEXT("Result"), TEXT("Value"));
+	return Info;
+}
+
 FRigVMFunction_MathColorFromFloat_Execute()
 {
 	Result = FLinearColor(Value, Value, Value);
diff --git a/Engine/Plugins/Runtime/RigVM/Source/RigVM/Private/RigVMFunctions/Math/RigVMFunction_MathDouble.cpp b/Engine/Plugins/Runtime/RigVM/Source/RigVM/Private/RigVMFunctions/Math/RigVMFunction_MathDouble.cpp
--- a/Engine/Plugins/Runtime/RigVM/Source/RigVM/Private/RigVMFunctions/Math/RigVMFunction_MathDouble.cpp
+++ b/Engine/Plugins/Runtime/RigVM/Source/RigVM/Private/RigVMFunctions/Math/RigVMFunction_MathDouble.cpp
@@ -1,6 +1,7 @@
 // Copyright Epic Games, Inc. All Rights Reserved.
 
 #include "RigVMFunctions/Math/RigVMFunction_MathDouble.h"
+#include "RigVMFunctions/RigVMDispatch_Constant.h"
 #include "RigVMFunctions/RigVMDispatch_Core.h"
 
 #include UE_INLINE_GENERATED_CPP_BY_NAME(RigVMFunction_MathDouble)
@@ -9,6 +10,11 @@ FRigVMFunction_MathDoubleMake_Execute()
 {
 }
 
+FRigVMStructUpgradeInfo FRigVMFunction_MathDoubleMake::GetUpgradeInfo() const
+{
+	return FRigVMStructUpgradeInfo::MakeFromStructToFactory(StaticStruct(), FRigVMDispatch_Constant::StaticStruct());
+}
+
 FRigVMFunction_MathDoubleConstPi_Execute()
 {
 	Value = PI;
diff --git a/Engine/Plugins/Runtime/RigVM/Source/RigVM/Private/RigVMFunctions/Math/RigVMFunction_MathFloat.cpp b/Engine/Plugins/Runtime/RigVM/Source/RigVM/Private/RigVMFunctions/Math/RigVMFunction_MathFloat.cpp
--- a/Engine/Plugins/Runtime/RigVM/Source/RigVM/Private/RigVMFunctions/Math/RigVMFunction_MathFloat.cpp
+++ b/Engine/Plugins/Runtime/RigVM/Source/RigVM/Private/RigVMFunctions/Math/RigVMFunction_MathFloat.cpp
@@ -1,6 +1,7 @@
 // Copyright Epic Games, Inc. All Rights Reserved.
 
 #include "RigVMFunctions/Math/RigVMFunction_MathFloat.h"
+#include "RigVMFunctions/RigVMDispatch_Constant.h"
 #include "RigVMFunctions/RigVMDispatch_Core.h"
 
 #include UE_INLINE_GENERATED_CPP_BY_NAME(RigVMFunction_MathFloat)
@@ -9,6 +10,11 @@ FRigVMFunction_MathFloatMake_Execute()
 {
 }
 
+FRigVMStructUpgradeInfo FRigVMFunction_MathFloatMake::GetUpgradeInfo() const
+{
+	return FRigVMStructUpgradeInfo::MakeFromStructToFactory(StaticStruct(), FRigVMDispatch_Constant::StaticStruct());
+}
+
 FRigVMFunction_MathFloatConstPi_Execute()
 {
 	Value = PI;
diff --git a/Engine/Plugins/Runtime/RigVM/Source/RigVM/Private/RigVMFunctions/Math/RigVMFunction_MathInt.cpp b/Engine/Plugins/Runtime/RigVM/Source/RigVM/Private/RigVMFunctions/Math/RigVMFunction_MathInt.cpp
--- a/Engine/Plugins/Runtime/RigVM/Source/RigVM/Private/RigVMFunctions/Math/RigVMFunction_MathInt.cpp
+++ b/Engine/Plugins/Runtime/RigVM/Source/RigVM/Private/RigVMFunctions/Math/RigVMFunction_MathInt.cpp
@@ -1,6 +1,7 @@
 // Copyright Epic Games, Inc. All Rights Reserved.
 
 #include "RigVMFunctions/Math/RigVMFunction_MathInt.h"
+#include "RigVMFunctions/RigVMDispatch_Constant.h"
 #include "RigVMFunctions/RigVMDispatch_Core.h"
 #include "RigVMCore/RigVMStructTest.h"
 
@@ -12,6 +13,11 @@ FRigVMFunction_MathIntMake_Execute()
 {
 }
 
+FRigVMStructUpgradeInfo FRigVMFunction_MathIntMake::GetUpgradeInfo() const
+{
+	return FRigVMStructUpgradeInfo::MakeFromStructToFactory(StaticStruct(), FRigVMDispatch_Constant::StaticStruct());
+}
+
 FRigVMFunction_MathIntAdd_Execute()
 {
 	Result = A + B;
diff --git a/Engine/Plugins/Runtime/RigVM/Source/RigVM/Private/RigVMFunctions/Math/RigVMFunction_MathTransform.cpp b/Engine/Plugins/Runtime/RigVM/Source/RigVM/Private/RigVMFunctions/Math/RigVMFunction_MathTransform.cpp
--- a/Engine/Plugins/Runtime/RigVM/Source/RigVM/Private/RigVMFunctions/Math/RigVMFunction_MathTransform.cpp
+++ b/Engine/Plugins/Runtime/RigVM/Source/RigVM/Private/RigVMFunctions/Math/RigVMFunction_MathTransform.cpp
@@ -2,6 +2,7 @@
 
 #include "RigVMFunctions/Math/RigVMFunction_MathTransform.h"
 #include "RigVMFunctions/Math/RigVMFunction_MathVector.h"
+#include "RigVMFunctions/RigVMDispatch_Constant.h"
 #include "AnimationCoreLibrary.h"
 
 #include UE_INLINE_GENERATED_CPP_BY_NAME(RigVMFunction_MathTransform)
@@ -11,6 +12,16 @@ FRigVMFunction_MathTransformMake_Execute()
 	Result = FTransform(Rotation, Translation, Scale);
 }
 
+FRigVMStructUpgradeInfo FRigVMFunction_MathTransformMake::GetUpgradeInfo() const
+{
+	FRigVMStructUpgradeInfo Info = FRigVMStructUpgradeInfo::MakeFromStructToFactory(StaticStruct(), FRigVMDispatch_Constant::StaticStruct());
+	Info.AddRemappedPin(TEXT("Translation"), TEXT("Value.Translation"));
+	Info.AddRemappedPin(TEXT("Rotation"), TEXT("Value.Rotation"));
+	Info.AddRemappedPin(TEXT("Scale"), TEXT("Value.Scale"));
+	Info.AddRemappedPin(TEXT("Result"), TEXT("Value"));
+	return Info;
+}
+
 FRigVMFunction_MathTransformFromEulerTransform_Execute()
 {
 	Result = EulerTransform.ToFTransform();
diff --git a/Engine/Plugins/Runtime/RigVM/Source/RigVM/Private/RigVMFunctions/Math/RigVMFunction_MathVector.cpp b/Engine/Plugins/Runtime/RigVM/Source/RigVM/Private/RigVMFunctions/Math/RigVMFunction_MathVector.cpp
--- a/Engine/Plugins/Runtime/RigVM/Source/RigVM/Private/RigVMFunctions/Math/RigVMFunction_MathVector.cpp
+++ b/Engine/Plugins/Runtime/RigVM/Source/RigVM/Private/RigVMFunctions/Math/RigVMFunction_MathVector.cpp
@@ -2,6 +2,7 @@
 
 #include "RigVMFunctions/Math/RigVMFunction_MathVector.h"
 #include "RigVMFunctions/Math/RigVMFunction_MathTransform.h"
+#include "RigVMFunctions/RigVMDispatch_Constant.h"
 #include "RigVMFunctions/RigVMDispatch_Core.h"
 #include "Engine/Engine.h"
 #include "Materials/MaterialInstanceDynamic.h"
@@ -13,6 +14,16 @@ FRigVMFunction_MathVectorMake_Execute()
 	Result = FVector(X, Y, Z);
 }
 
+FRigVMStructUpgradeInfo FRigVMFunction_MathVectorMake::GetUpgradeInfo() const
+{
+	FRigVMStructUpgradeInfo Info = FRigVMStructUpgradeInfo::MakeFromStructToFactory(StaticStruct(), FRigVMDispatch_Constant::StaticStruct());
+	Info.AddRemappedPin(TEXT("X"), TEXT("Value.X"));
+	Info.AddRemappedPin(TEXT("Y"), TEXT("Value.Y"));
+	Info.AddRemappedPin(TEXT("Z"), TEXT("Value.Z"));
+	Info.AddRemappedPin(TEXT("Result"), TEXT("Value"));
+	return Info;
+}
+
 FRigVMFunction_MathVectorFromFloat_Execute()
 {
 	Result = FVector(Value, Value, Value);
diff --git a/Engine/Plugins/Runtime/RigVM/Source/RigVM/Public/RigVMCore/RigVMDispatchFactory.h b/Engine/Plugins/Runtime/RigVM/Source/RigVM/Public/RigVMCore/RigVMDispatchFactory.h
--- a/Engine/Plugins/Runtime/RigVM/Source/RigVM/Public/RigVMCore/RigVMDispatchFactory.h
+++ b/Engine/Plugins/Runtime/RigVM/Source/RigVM/Public/RigVMCore/RigVMDispatchFactory.h
@@ -79,6 +79,9 @@ public:
 	// returns the registry this factory belongs top
 	const FRigVMRegistry_NoLock* GetOwnerRegistry() const { return OwnerRegistry; }
 
+	// returns true if nodes provided by this factory can be renamed
+	virtual bool SupportsRenaming() const { return false; }
+
 #if WITH_EDITOR
 	
 	// returns the title of the node for a given type set.
diff --git a/Engine/Plugins/Runtime/RigVM/Source/RigVM/Public/RigVMFunctions/Math/RigVMFunction_MathBool.h b/Engine/Plugins/Runtime/RigVM/Source/RigVM/Public/RigVMFunctions/Math/RigVMFunction_MathBool.h
--- a/Engine/Plugins/Runtime/RigVM/Source/RigVM/Public/RigVMFunctions/Math/RigVMFunction_MathBool.h
+++ b/Engine/Plugins/Runtime/RigVM/Source/RigVM/Public/RigVMFunctions/Math/RigVMFunction_MathBool.h
@@ -85,7 +85,7 @@ struct FRigVMFunction_MathBoolBinaryAggregateOp : public FRigVMFunction_MathBool
 /**
  * A boolean constant
  */
-USTRUCT(meta=(DisplayName="Bool", Keywords="Make,Construct,Constant"))
+USTRUCT(meta=(DisplayName="Bool", Keywords="Make,Construct,Constant", Deprecated="5.7"))
 struct FRigVMFunction_MathBoolMake : public FRigVMFunction_MathBoolBase
 {
 	GENERATED_BODY()
@@ -100,6 +100,9 @@ struct FRigVMFunction_MathBoolMake : public FRigVMFunction_MathBoolBase
 
 	UPROPERTY(meta=(Input, Output))
 	bool Value;
+
+	RIGVM_METHOD()
+	RIGVM_API virtual FRigVMStructUpgradeInfo GetUpgradeInfo() const override;
 };
 
 /**
diff --git a/Engine/Plugins/Runtime/RigVM/Source/RigVM/Public/RigVMFunctions/Math/RigVMFunction_MathColor.h b/Engine/Plugins/Runtime/RigVM/Source/RigVM/Public/RigVMFunctions/Math/RigVMFunction_MathColor.h
--- a/Engine/Plugins/Runtime/RigVM/Source/RigVM/Public/RigVMFunctions/Math/RigVMFunction_MathColor.h
+++ b/Engine/Plugins/Runtime/RigVM/Source/RigVM/Public/RigVMFunctions/Math/RigVMFunction_MathColor.h
@@ -54,7 +54,7 @@ struct FRigVMFunction_MathColorBinaryAggregateOp : public FRigVMFunction_MathCol
 /**
  * Makes a color from its components
  */
-USTRUCT(meta=(DisplayName="Make Color", Keywords="Make,Construct,Constant"))
+USTRUCT(meta=(DisplayName="Make Color", Keywords="Make,Construct,Constant", Deprecated="5.7"))
 struct FRigVMFunction_MathColorMake : public FRigVMFunction_MathColorBase
 {
 	GENERATED_BODY()
@@ -82,6 +82,9 @@ struct FRigVMFunction_MathColorMake : public FRigVMFunction_MathColorBase
 
 	UPROPERTY(meta=(Output))
 	FLinearColor Result;
+
+	RIGVM_METHOD()
+	RIGVM_API virtual FRigVMStructUpgradeInfo GetUpgradeInfo() const override;
 };
 
 /**
diff --git a/Engine/Plugins/Runtime/RigVM/Source/RigVM/Public/RigVMFunctions/Math/RigVMFunction_MathDouble.h b/Engine/Plugins/Runtime/RigVM/Source/RigVM/Public/RigVMFunctions/Math/RigVMFunction_MathDouble.h
--- a/Engine/Plugins/Runtime/RigVM/Source/RigVM/Public/RigVMFunctions/Math/RigVMFunction_MathDouble.h
+++ b/Engine/Plugins/Runtime/RigVM/Source/RigVM/Public/RigVMFunctions/Math/RigVMFunction_MathDouble.h
@@ -86,7 +86,7 @@ struct FRigVMFunction_MathDoubleBinaryAggregateOp : public FRigVMFunction_MathDo
 /**
  * A double constant
  */
-USTRUCT(meta=(DisplayName="Double", Keywords="Make,Construct,Constant"))
+USTRUCT(meta=(DisplayName="Double", Keywords="Make,Construct,Constant", Deprecated="5.7"))
 struct FRigVMFunction_MathDoubleMake : public FRigVMFunction_MathDoubleBase
 {
 	GENERATED_BODY()
@@ -101,6 +101,9 @@ struct FRigVMFunction_MathDoubleMake : public FRigVMFunction_MathDoubleBase
 
 	UPROPERTY(meta=(Input, Output))
 	double Value;
+	
+	RIGVM_METHOD()
+	RIGVM_API virtual FRigVMStructUpgradeInfo GetUpgradeInfo() const override;
 };
 
 /**
@@ -975,4 +978,4 @@ struct FRigVMFunction_MathDoubleArrayAverage : public FRigVMFunction_MathDoubleB
 
 	UPROPERTY(meta = (Output))
 	double Average;
-};
\ No newline at end of file
+};
diff --git a/Engine/Plugins/Runtime/RigVM/Source/RigVM/Public/RigVMFunctions/Math/RigVMFunction_MathFloat.h b/Engine/Plugins/Runtime/RigVM/Source/RigVM/Public/RigVMFunctions/Math/RigVMFunction_MathFloat.h
--- a/Engine/Plugins/Runtime/RigVM/Source/RigVM/Public/RigVMFunctions/Math/RigVMFunction_MathFloat.h
+++ b/Engine/Plugins/Runtime/RigVM/Source/RigVM/Public/RigVMFunctions/Math/RigVMFunction_MathFloat.h
@@ -87,7 +87,7 @@ struct FRigVMFunction_MathFloatBinaryAggregateOp : public FRigVMFunction_MathFlo
 /**
  * A float constant
  */
-USTRUCT(meta=(DisplayName="Float", Keywords="Make,Construct,Constant"))
+USTRUCT(meta=(DisplayName="Float", Keywords="Make,Construct,Constant", Deprecated="5.7"))
 struct FRigVMFunction_MathFloatMake : public FRigVMFunction_MathFloatBase
 {
 	GENERATED_BODY()
@@ -102,6 +102,9 @@ struct FRigVMFunction_MathFloatMake : public FRigVMFunction_MathFloatBase
 
 	UPROPERTY(meta=(Input, Output))
 	float Value;
+
+	RIGVM_METHOD()
+	RIGVM_API virtual FRigVMStructUpgradeInfo GetUpgradeInfo() const override;
 };
 
 /**
diff --git a/Engine/Plugins/Runtime/RigVM/Source/RigVM/Public/RigVMFunctions/Math/RigVMFunction_MathInt.h b/Engine/Plugins/Runtime/RigVM/Source/RigVM/Public/RigVMFunctions/Math/RigVMFunction_MathInt.h
--- a/Engine/Plugins/Runtime/RigVM/Source/RigVM/Public/RigVMFunctions/Math/RigVMFunction_MathInt.h
+++ b/Engine/Plugins/Runtime/RigVM/Source/RigVM/Public/RigVMFunctions/Math/RigVMFunction_MathInt.h
@@ -71,7 +71,7 @@ struct FRigVMFunction_MathIntBinaryAggregateOp : public FRigVMFunction_MathIntBa
 /**
  * A integer constant
  */
-USTRUCT(meta=(DisplayName="Integer", Keywords="Make,Construct,Constant"))
+USTRUCT(meta=(DisplayName="Integer", Keywords="Make,Construct,Constant", Deprecated="5.7"))
 struct FRigVMFunction_MathIntMake : public FRigVMFunction_MathIntBase
 {
 	GENERATED_BODY()
@@ -86,6 +86,9 @@ struct FRigVMFunction_MathIntMake : public FRigVMFunction_MathIntBase
 
 	UPROPERTY(meta=(Input, Output))
 	int32 Value;
+
+	RIGVM_METHOD()
+	RIGVM_API virtual FRigVMStructUpgradeInfo GetUpgradeInfo() const override;
 };
 
 /**
diff --git a/Engine/Plugins/Runtime/RigVM/Source/RigVM/Public/RigVMFunctions/Math/RigVMFunction_MathTransform.h b/Engine/Plugins/Runtime/RigVM/Source/RigVM/Public/RigVMFunctions/Math/RigVMFunction_MathTransform.h
--- a/Engine/Plugins/Runtime/RigVM/Source/RigVM/Public/RigVMFunctions/Math/RigVMFunction_MathTransform.h
+++ b/Engine/Plugins/Runtime/RigVM/Source/RigVM/Public/RigVMFunctions/Math/RigVMFunction_MathTransform.h
@@ -80,7 +80,7 @@ struct FRigVMFunction_MathTransformBinaryAggregateOp : public FRigVMFunction_Mat
 /**
  * Makes a transform from its components
  */
-USTRUCT(meta=(DisplayName="Make Transform", Keywords="Make,Construct,Constant"))
+USTRUCT(meta=(DisplayName="Make Transform", Keywords="Make,Construct,Constant", Deprecated="5.7"))
 struct FRigVMFunction_MathTransformMake : public FRigVMFunction_MathTransformBase
 {
 	GENERATED_BODY()
@@ -106,6 +106,9 @@ struct FRigVMFunction_MathTransformMake : public FRigVMFunction_MathTransformBas
 
 	UPROPERTY(meta=(Output))
 	FTransform Result;
+
+	RIGVM_METHOD()
+	RIGVM_API virtual FRigVMStructUpgradeInfo GetUpgradeInfo() const override;
 };
 
 /**
diff --git a/Engine/Plugins/Runtime/RigVM/Source/RigVM/Public/RigVMFunctions/Math/RigVMFunction_MathVector.h b/Engine/Plugins/Runtime/RigVM/Source/RigVM/Public/RigVMFunctions/Math/RigVMFunction_MathVector.h
--- a/Engine/Plugins/Runtime/RigVM/Source/RigVM/Public/RigVMFunctions/Math/RigVMFunction_MathVector.h
+++ b/Engine/Plugins/Runtime/RigVM/Source/RigVM/Public/RigVMFunctions/Math/RigVMFunction_MathVector.h
@@ -72,7 +72,7 @@ struct FRigVMFunction_MathVectorBinaryAggregateOp : public FRigVMFunction_MathVe
 /**
  * Makes a vector from its components
  */
-USTRUCT(meta=(DisplayName="Make Vector", Keywords="Make,Construct,Constant"))
+USTRUCT(meta=(DisplayName="Make Vector", Keywords="Make,Construct,Constant", Deprecated="5.7"))
 struct FRigVMFunction_MathVectorMake : public FRigVMFunction_MathVectorBase
 {
 	GENERATED_BODY()
@@ -96,6 +96,9 @@ struct FRigVMFunction_MathVectorMake : public FRigVMFunction_MathVectorBase
 
 	UPROPERTY(meta=(Output))
 	FVector Result;
+		
+	RIGVM_METHOD()
+	RIGVM_API virtual FRigVMStructUpgradeInfo GetUpgradeInfo() const override;
 };
 
 /**
diff --git a/Engine/Plugins/Runtime/RigVM/Source/RigVM/Public/RigVMFunctions/RigVMDispatch_Constant.h b/Engine/Plugins/Runtime/RigVM/Source/RigVM/Public/RigVMFunctions/RigVMDispatch_Constant.h
--- a/Engine/Plugins/Runtime/RigVM/Source/RigVM/Public/RigVMFunctions/RigVMDispatch_Constant.h
+++ b/Engine/Plugins/Runtime/RigVM/Source/RigVM/Public/RigVMFunctions/RigVMDispatch_Constant.h
@@ -25,7 +25,8 @@ public:
 	RIGVM_API virtual FText GetNodeTooltip(const FRigVMTemplateTypeMap& InTypes) const override;
 #endif
 	virtual bool IsSingleton() const override { return true; }
-	
+	virtual bool SupportsRenaming() const override { return true; }
+
 protected:
 	virtual FRigVMFunctionPtr GetDispatchFunctionImpl(const FRigVMTemplateTypeMap& InTypes, const FRigVMRegistryHandle& InRegistry) const override { return &FRigVMDispatch_Constant::Execute; }
 	static RIGVM_API void Execute(FRigVMExtendedExecuteContext& InContext, FRigVMMemoryHandleArray Handles, FRigVMPredicateBranchArray RigVMBranches);
@@ -33,4 +34,5 @@ protected:
 	static inline const FLazyName ValueName = FLazyName(TEXT("Value"));
 
 	friend class URigVMController;
+	friend class FRigVMParserAST;
 };
diff --git a/Engine/Plugins/Runtime/RigVM/Source/RigVM/Public/RigVMFunctions/RigVMDispatch_MakeStruct.h b/Engine/Plugins/Runtime/RigVM/Source/RigVM/Public/RigVMFunctions/RigVMDispatch_MakeStruct.h
--- a/Engine/Plugins/Runtime/RigVM/Source/RigVM/Public/RigVMFunctions/RigVMDispatch_MakeStruct.h
+++ b/Engine/Plugins/Runtime/RigVM/Source/RigVM/Public/RigVMFunctions/RigVMDispatch_MakeStruct.h
@@ -26,6 +26,7 @@ public:
 	RIGVM_API virtual FString GetArgumentMetaData(const FName& InArgumentName, const FName& InMetaDataKey) const override;
 	RIGVM_API virtual FString GetKeywords() const override;
 #endif
+	virtual bool SupportsRenaming() const override { return true; }
 
 protected:
 	virtual FRigVMFunctionPtr GetDispatchFunctionImpl(const FRigVMTemplateTypeMap& InTypes, const FRigVMRegistryHandle& InRegistry) const override { return &FRigVMDispatch_MakeStruct::Execute; }
@@ -36,6 +37,7 @@ protected:
 
 	friend struct FRigVMDispatch_BreakStruct;
 	friend class URigVMController;
+	friend class FRigVMParserAST;
 };
 
 USTRUCT(meta=(DisplayName = "Break", Keywords = "Decompose,Decomposition"))
diff --git a/Engine/Plugins/Runtime/RigVM/Source/RigVMDeveloper/Private/EdGraph/RigVMEdGraphNode.cpp b/Engine/Plugins/Runtime/RigVM/Source/RigVMDeveloper/Private/EdGraph/RigVMEdGraphNode.cpp
--- a/Engine/Plugins/Runtime/RigVM/Source/RigVMDeveloper/Private/EdGraph/RigVMEdGraphNode.cpp
+++ b/Engine/Plugins/Runtime/RigVM/Source/RigVMDeveloper/Private/EdGraph/RigVMEdGraphNode.cpp
@@ -54,6 +54,7 @@ URigVMEdGraphNode::URigVMEdGraphNode()
 #endif
 , CachedTemplate(nullptr)
 , MicroSeconds(0)
+, bRenameIsPending(false)
 {
 	bHasCompilerMessage = false;
 	ErrorType = (int32)EMessageSeverity::Info + 1;
@@ -70,7 +71,12 @@ FText URigVMEdGraphNode::GetNodeTitle(ENodeTitleType::Type TitleType) const
 		FString SubTitle;
 		if(URigVMNode* ModelNode = GetModelNode())
 		{
-			if (const URigVMUnitNode* UnitNode = Cast<URigVMUnitNode>(ModelNode))
+			if(bCanRenameNode)
+			{
+				NodeTitle = FText::FromString(ModelNode->GetName());
+			}
+
+			else if (const URigVMUnitNode* UnitNode = Cast<URigVMUnitNode>(ModelNode))
 			{
 				const UScriptStruct* ScriptStruct = UnitNode->GetScriptStruct();
 				if (ScriptStruct && ScriptStruct->IsChildOf(FRigVMStruct::StaticStruct()))
@@ -387,12 +393,20 @@ void URigVMEdGraphNode::PostReconstructNode()
 
 	bCanRenameNode = false;
 
+	if(URigVMController* Controller = GetController())
+	{
+		if(const URigVMSchema* Schema = Controller->GetSchema())
+		{
+			bCanRenameNode = Schema->CanSetNodeTitle(Controller, GetModelNode());
+		}
+	}
+
 	if(IsOutDated())
 	{
 		static const FLinearColor WarningColor = FAppStyle::GetColor("ErrorReporting.WarningBackgroundColor");
 		SetColorFromModel(WarningColor);
 	}
-	else if(URigVMNode* ModelNode = GetModelNode())
+	else if(const URigVMNode* ModelNode = GetModelNode())
 	{
 		SetColorFromModel(ModelNode->GetNodeColor());
 	}
@@ -2310,6 +2324,7 @@ void URigVMEdGraphNode::UpdateProfilingSettings()
 		}
 	}
 }
+
 #endif
 
 FString URigVMEdGraphNode::GetPinMetaData(FName InPinName, FName InKey)
@@ -2331,5 +2346,102 @@ FString URigVMEdGraphNode::GetPinMetaData(FName InPinName, FName InKey)
 	return Pin->GetMetaData(InKey);
 }
 
+TSharedPtr<INameValidatorInterface> URigVMEdGraphNode::MakeNameValidator() const
+{
+	const URigVMController* Controller = GetController();
+	check(Controller);
+	return MakeShareable(new FRigVMGraphNodeNameValidator(Controller));
+}
+
+void URigVMEdGraphNode::RequestRename(float InDelay)
+{
+	if(GetCanRenameNode())
+	{
+		if(RequestRenameDelegate.IsBound())
+		{
+			if(IsInGameThread())
+			{
+				RequestRenameDelegate.Execute();
+			}
+			else
+			{
+				TWeakObjectPtr<URigVMEdGraphNode> WeakThis = this;
+				auto RequestRenameLambda = [WeakThis]()
+				{
+					if(WeakThis.IsValid())
+					{
+						(void)WeakThis->RequestRenameDelegate.ExecuteIfBound();
+						WeakThis->bRenameIsPending = false;
+					}
+				};
+				
+				if(InDelay < SMALL_NUMBER)
+				{
+					FFunctionGraphTask::CreateAndDispatchWhenReady(RequestRenameLambda, TStatId(), NULL, ENamedThreads::GameThread);
+				}
+				else
+				{
+					FTimerHandle RequestRenameHandle;
+					GEditor->GetTimerManager()->SetTimer(RequestRenameHandle, FTimerDelegate::CreateLambda([RequestRenameLambda]()
+					{
+						RequestRenameLambda();
+					}), 1, false, InDelay);
+				}
+			}
+		}
+		else
+		{
+			bRenameIsPending = true;
+		}
+	}
+}
+
+FRigVMGraphNodeNameValidator::FRigVMGraphNodeNameValidator(const URigVMController* InController)
+: Controller(InController)
+{
+}
+
+EValidatorResult FRigVMGraphNodeNameValidator::IsValid(const FString& Name, bool bOriginal)
+{
+	if(Name.IsEmpty())
+	{
+		return EValidatorResult::EmptyName;
+	}
+	
+	const URigVMGraph* Graph = Controller->GetGraph();
+	check(Graph);
+
+	if(Graph->FindNode(Name))
+	{
+		return EValidatorResult::AlreadyInUse;
+	}
+
+	const URigVMSchema* Schema = Graph->GetSchema();
+	check(Schema);
+
+	const FString SanitizedName = Schema->GetValidNodeName(Controller->GetGraph(), Name);
+
+	if(SanitizedName.Len() < Name.Len())
+	{
+		return EValidatorResult::TooLong;
+	}
+	
+	if(!SanitizedName.Equals(Name, ESearchCase::CaseSensitive))
+	{
+		return EValidatorResult::ContainsInvalidCharacters;
+	}
+
+	return EValidatorResult::Ok;
+}
+
+EValidatorResult FRigVMGraphNodeNameValidator::IsValid(const FName& Name, bool bOriginal)
+{
+	if(Name.IsNone())
+	{
+		return EValidatorResult::EmptyName;
+	}
+	return IsValid(Name.ToString(), bOriginal);
+}
+
 #undef LOCTEXT_NAMESPACE
 
diff --git a/Engine/Plugins/Runtime/RigVM/Source/RigVMDeveloper/Private/RigVMModel/Nodes/RigVMDispatchNode.cpp b/Engine/Plugins/Runtime/RigVM/Source/RigVMDeveloper/Private/RigVMModel/Nodes/RigVMDispatchNode.cpp
--- a/Engine/Plugins/Runtime/RigVM/Source/RigVMDeveloper/Private/RigVMModel/Nodes/RigVMDispatchNode.cpp
+++ b/Engine/Plugins/Runtime/RigVM/Source/RigVMDeveloper/Private/RigVMModel/Nodes/RigVMDispatchNode.cpp
@@ -14,6 +14,15 @@ FString URigVMDispatchNode::GetNodeTitle() const
 	return Super::GetNodeTitle();
 }
 
+bool URigVMDispatchNode::SupportsRenaming() const
+{
+	if(const FRigVMDispatchFactory* Factory = GetFactory())
+	{
+		return Factory->SupportsRenaming();
+	}
+	return Super::SupportsRenaming();
+}
+
 FText URigVMDispatchNode::GetToolTipText() const
 {
 	if (const FRigVMDispatchFactory* Factory = GetFactory())
diff --git a/Engine/Plugins/Runtime/RigVM/Source/RigVMDeveloper/Private/RigVMModel/Nodes/RigVMRerouteNode.cpp b/Engine/Plugins/Runtime/RigVM/Source/RigVMDeveloper/Private/RigVMModel/Nodes/RigVMRerouteNode.cpp
--- a/Engine/Plugins/Runtime/RigVM/Source/RigVMDeveloper/Private/RigVMModel/Nodes/RigVMRerouteNode.cpp
+++ b/Engine/Plugins/Runtime/RigVM/Source/RigVMDeveloper/Private/RigVMModel/Nodes/RigVMRerouteNode.cpp
@@ -57,6 +57,26 @@ FLinearColor URigVMRerouteNode::GetNodeColor() const
 	return FLinearColor::White;
 }
 
+FText URigVMRerouteNode::GetToolTipText() const
+{
+	for(const URigVMLink* Link : GetPins()[0]->GetLinks())
+	{
+		if(const URigVMPin* SourcePin = Link->GetSourcePin())
+		{
+			if(SourcePin->GetNode() != this)
+			{
+				if(const URigVMRerouteNode* SourceRerouteNode = Cast<URigVMRerouteNode>(SourcePin->GetNode()))
+				{
+					return SourceRerouteNode->GetToolTipText();
+				}
+				const FText Format = NSLOCTEXT("URigVMRerouteNode", "ToolTipFormat", "Points to: {0}\n\n{1}");
+				return FText::Format(Format, FText::FromString(SourcePin->GetPinPath()), SourcePin->GetNode()->GetToolTipText());
+			}
+		}
+	}
+	return Super::GetToolTipText();
+}
+
 bool URigVMRerouteNode::IsLiteral() const
 {
 	return GetPins()[0]->GetSourceLinks(true).IsEmpty();
diff --git a/Engine/Plugins/Runtime/RigVM/Source/RigVMDeveloper/Private/RigVMModel/RigVMController.cpp b/Engine/Plugins/Runtime/RigVM/Source/RigVMDeveloper/Private/RigVMModel/RigVMController.cpp
--- a/Engine/Plugins/Runtime/RigVM/Source/RigVMDeveloper/Private/RigVMModel/RigVMController.cpp
+++ b/Engine/Plugins/Runtime/RigVM/Source/RigVMDeveloper/Private/RigVMModel/RigVMController.cpp
@@ -2818,14 +2818,23 @@ TArray<URigVMNode*> URigVMController::UpgradeNodes(const TArray<URigVMNode*>& In
 		return TArray<URigVMNode*>();
 	}
 
+	FRestoreLinkedPathSettings Settings;
+
 	TArray<URigVMNode*> UpgradedNodes;
 	TMap<FString,FRigVMController_PinPathRemapDelegate> RemapPinDelegates;
 	for(URigVMNode* Node : InNodes)
 	{
+		const FName OriginalName = Node->GetFName();
 		FRigVMController_PinPathRemapDelegate RemapPinDelegate;
 		URigVMNode* UpgradedNode = UpgradeNode(Node, bSetupUndoRedo, &RemapPinDelegate);
 		if(UpgradedNode)
 		{
+			const FName NewName = UpgradedNode->GetFName();
+			if(NewName != OriginalName)
+			{
+				Settings.NodeNameMap.Add(OriginalName.ToString(), NewName.ToString());
+			}
+			
 			UpgradedNodes.Add(UpgradedNode);
 			if(RemapPinDelegate.IsBound())
 			{
@@ -2834,7 +2843,6 @@ TArray<URigVMNode*> URigVMController::UpgradeNodes(const TArray<URigVMNode*>& In
 		}
 	}
 
-	FRestoreLinkedPathSettings Settings;
 	Settings.RemapDelegates = RemapPinDelegates;
 	RestoreLinkedPaths(LinkedPaths, Settings, bSetupUndoRedo);
 
@@ -2868,6 +2876,7 @@ URigVMNode* URigVMController::UpgradeNode(URigVMNode* InNode, bool bSetupUndoRed
 	EjectAllInjectedNodes(InNode, bSetupUndoRedo);
 
 	const FString NodeName = InNode->GetName();
+	const FString NodeTitle = InNode->GetNodeTitle();
 	const FVector2D NodePosition = InNode->GetPosition();
 
 	FRigVMBaseAction Action(this);
@@ -2995,6 +3004,11 @@ URigVMNode* URigVMController::UpgradeNode(URigVMNode* InNode, bool bSetupUndoRed
 
 	ApplyPinStates(UpgradedNode, PinStates, RedirectedPinPaths, bSetupUndoRedo);
 
+	if(!NodeTitle.IsEmpty() && UpgradedNode->SupportsRenaming())
+	{
+		(void)RenameNode(UpgradedNode, *NodeTitle, bSetupUndoRedo, false);
+	}
+
 	if(bSetupUndoRedo)
 	{
 		GetActionStack()->EndAction(Action);
diff --git a/Engine/Plugins/Runtime/RigVM/Source/RigVMDeveloper/Private/RigVMModel/RigVMNode.cpp b/Engine/Plugins/Runtime/RigVM/Source/RigVMDeveloper/Private/RigVMModel/RigVMNode.cpp
--- a/Engine/Plugins/Runtime/RigVM/Source/RigVMDeveloper/Private/RigVMModel/RigVMNode.cpp
+++ b/Engine/Plugins/Runtime/RigVM/Source/RigVMDeveloper/Private/RigVMModel/RigVMNode.cpp
@@ -659,6 +659,11 @@ const FString& URigVMNode::GetNodeTitleRaw() const
 	return NodeTitle;
 }
 
+bool URigVMNode::SupportsRenaming() const
+{
+	return false;
+}
+
 FVector2D URigVMNode::GetPosition() const
 {
 	return Position;
diff --git a/Engine/Plugins/Runtime/RigVM/Source/RigVMDeveloper/Private/RigVMModel/RigVMSchema.cpp b/Engine/Plugins/Runtime/RigVM/Source/RigVMDeveloper/Private/RigVMModel/RigVMSchema.cpp
--- a/Engine/Plugins/Runtime/RigVM/Source/RigVMDeveloper/Private/RigVMModel/RigVMSchema.cpp
+++ b/Engine/Plugins/Runtime/RigVM/Source/RigVMDeveloper/Private/RigVMModel/RigVMSchema.cpp
@@ -4,6 +4,7 @@
 #include "RigVMStringUtils.h"
 #include "RigVMModel/RigVMClient.h"
 #include "RigVMModel/RigVMController.h"
+#include "RigVMFunctions/RigVMDispatch_Constant.h"
 
 TAutoConsoleVariable<bool> CVarRigVMEnableNodeLayouts(TEXT("RigVM.EnableNodeLayouts"), false, TEXT("Set to true to turn on support for node layouts"));
 
@@ -473,6 +474,17 @@ bool URigVMSchema::CanSetNodeTitle(URigVMController* InController, const URigVMN
 {
 	RIGVMSCHEMA_DEFAULT_FUNCTION_BODY
 
+	if(const URigVMTemplateNode* TemplateNode = Cast<URigVMTemplateNode>(InNode))
+	{
+		if(const FRigVMTemplate* Template = TemplateNode->GetTemplate())
+		{
+			if(const FRigVMDispatchFactory* Factory = Template->GetDispatchFactory())
+			{
+				return Factory->SupportsRenaming();
+			}
+		}
+	}
+	
 	return false;
 }
 
diff --git a/Engine/Plugins/Runtime/RigVM/Source/RigVMDeveloper/Public/EdGraph/RigVMEdGraphNode.h b/Engine/Plugins/Runtime/RigVM/Source/RigVMDeveloper/Public/EdGraph/RigVMEdGraphNode.h
--- a/Engine/Plugins/Runtime/RigVM/Source/RigVMDeveloper/Public/EdGraph/RigVMEdGraphNode.h
+++ b/Engine/Plugins/Runtime/RigVM/Source/RigVMDeveloper/Public/EdGraph/RigVMEdGraphNode.h
@@ -136,6 +136,8 @@ public:
 #endif
 #endif
 	UE_API virtual FString GetPinMetaData(FName InPinName, FName InKey) override;
+	UE_API virtual TSharedPtr<class INameValidatorInterface> MakeNameValidator() const override;
+	UE_API void RequestRename(float InDelay = 0.f);
 
 	UE_API virtual bool IsDeprecated() const override;
 	UE_API bool IsOutDated() const;
@@ -316,6 +318,8 @@ private:
 	mutable TOptional<bool> DrawAsCompactNodeCache;
 	mutable double MicroSeconds;
 	mutable TArray<double> MicroSecondsFrames;
+	FSimpleDelegate RequestRenameDelegate;
+	mutable bool bRenameIsPending;
 
 	friend class SRigVMGraphNode;
 	friend class FRigVMFunctionArgumentLayout;
@@ -326,4 +330,19 @@ private:
 	friend class URigVMEdGraphSelectNodeSpawner;
 };
 
+class FRigVMGraphNodeNameValidator : public INameValidatorInterface
+{
+public:
+	FRigVMGraphNodeNameValidator(const URigVMController* InController);
+
+	// Begin FNameValidatorInterface
+	virtual EValidatorResult IsValid(const FString& Name, bool bOriginal) override;
+	virtual EValidatorResult IsValid(const FName& Name, bool bOriginal) override;
+	// End FNameValidatorInterface
+
+private:
+
+	const URigVMController* Controller;
+};
+
 #undef UE_API
diff --git a/Engine/Plugins/Runtime/RigVM/Source/RigVMDeveloper/Public/RigVMModel/Nodes/RigVMDispatchNode.h b/Engine/Plugins/Runtime/RigVM/Source/RigVMDeveloper/Public/RigVMModel/Nodes/RigVMDispatchNode.h
--- a/Engine/Plugins/Runtime/RigVM/Source/RigVMDeveloper/Public/RigVMModel/Nodes/RigVMDispatchNode.h
+++ b/Engine/Plugins/Runtime/RigVM/Source/RigVMDeveloper/Public/RigVMModel/Nodes/RigVMDispatchNode.h
@@ -37,6 +37,7 @@ public:
 
 	// Override node functions
 	UE_API virtual FString GetNodeTitle() const override;
+	UE_API virtual bool SupportsRenaming() const override;
 	UE_API virtual FText GetToolTipText() const override;
 	UE_API virtual FLinearColor GetNodeColor() const override;
 	UE_API virtual bool IsDefinedAsConstant() const override;
diff --git a/Engine/Plugins/Runtime/RigVM/Source/RigVMDeveloper/Public/RigVMModel/Nodes/RigVMRerouteNode.h b/Engine/Plugins/Runtime/RigVM/Source/RigVMDeveloper/Public/RigVMModel/Nodes/RigVMRerouteNode.h
--- a/Engine/Plugins/Runtime/RigVM/Source/RigVMDeveloper/Public/RigVMModel/Nodes/RigVMRerouteNode.h
+++ b/Engine/Plugins/Runtime/RigVM/Source/RigVMDeveloper/Public/RigVMModel/Nodes/RigVMRerouteNode.h
@@ -29,6 +29,8 @@ public:
 
 	UE_API virtual FLinearColor GetNodeColor() const override;
 
+	UE_API virtual FText GetToolTipText() const override;
+
 	// Has no source connections
 	UE_API bool IsLiteral() const;
 
diff --git a/Engine/Plugins/Runtime/RigVM/Source/RigVMDeveloper/Public/RigVMModel/RigVMNode.h b/Engine/Plugins/Runtime/RigVM/Source/RigVMDeveloper/Public/RigVMModel/RigVMNode.h
--- a/Engine/Plugins/Runtime/RigVM/Source/RigVMDeveloper/Public/RigVMModel/RigVMNode.h
+++ b/Engine/Plugins/Runtime/RigVM/Source/RigVMDeveloper/Public/RigVMModel/RigVMNode.h
@@ -172,6 +172,9 @@ public:
 	// Warning : The raw title might be ignored by some nodes overrides to GetNodeTitle
 	UE_API const FString& GetNodeTitleRaw() const;
 
+	// Returns true if this node can be renamed by the user
+	UE_API virtual bool SupportsRenaming() const;
+
 	// Returns the 2d position of this node - used for UI.
 	UFUNCTION(BlueprintCallable, Category = RigVMNode)
 	UE_API FVector2D GetPosition() const;
diff --git a/Engine/Plugins/Runtime/RigVM/Source/RigVMEditor/Private/EdGraph/NodeSpawners/RigVMEdGraphTemplateNodeSpawner.cpp b/Engine/Plugins/Runtime/RigVM/Source/RigVMEditor/Private/EdGraph/NodeSpawners/RigVMEdGraphTemplateNodeSpawner.cpp
--- a/Engine/Plugins/Runtime/RigVM/Source/RigVMEditor/Private/EdGraph/NodeSpawners/RigVMEdGraphTemplateNodeSpawner.cpp
+++ b/Engine/Plugins/Runtime/RigVM/Source/RigVMEditor/Private/EdGraph/NodeSpawners/RigVMEdGraphTemplateNodeSpawner.cpp
@@ -168,7 +168,7 @@ UEdGraphNode* URigVMEdGraphTemplateNodeSpawner::Invoke(UEdGraph* ParentGraph, FB
 
 			TArray<FPinInfo> Pins;
 
-			static const FRigVMDispatchContext Context;
+			const FRigVMDispatchContext Context;
 			for (int32 Index = 0; Index < Template->NumExecuteArguments(Context); Index++)
 			{
 				const FRigVMExecuteArgument* Argument = Template->GetExecuteArgument(Index, Context);
@@ -234,6 +234,11 @@ URigVMEdGraphNode* URigVMEdGraphTemplateNodeSpawner::SpawnNode(UEdGraph* ParentG
 				Controller->SelectNode(ModelNode, true, true);
 
 				URigVMEdGraphUnitNodeSpawner::HookupMutableNode(ModelNode, RigBlueprint);
+
+				if(NewNode->GetCanRenameNode())
+				{
+					NewNode->RequestRename();
+				}
 			}
 
 			Controller->CloseUndoBracket();
diff --git a/Engine/Plugins/Runtime/RigVM/Source/RigVMEditor/Private/EdGraph/NodeSpawners/RigVMEdGraphUnitNodeSpawner.cpp b/Engine/Plugins/Runtime/RigVM/Source/RigVMEditor/Private/EdGraph/NodeSpawners/RigVMEdGraphUnitNodeSpawner.cpp
--- a/Engine/Plugins/Runtime/RigVM/Source/RigVMEditor/Private/EdGraph/NodeSpawners/RigVMEdGraphUnitNodeSpawner.cpp
+++ b/Engine/Plugins/Runtime/RigVM/Source/RigVMEditor/Private/EdGraph/NodeSpawners/RigVMEdGraphUnitNodeSpawner.cpp
@@ -358,6 +358,11 @@ URigVMEdGraphNode* URigVMEdGraphUnitNodeSpawner::SpawnNode(UEdGraph* ParentGraph
 #endif
 
 			Controller->CloseUndoBracket();
+
+			if(NewNode->GetCanRenameNode())
+			{
+				NewNode->RequestRename();
+			}
 		}
 		else
 		{
diff --git a/Engine/Plugins/Runtime/RigVM/Source/RigVMEditor/Private/Editor/RigVMEditor.cpp b/Engine/Plugins/Runtime/RigVM/Source/RigVMEditor/Private/Editor/RigVMEditor.cpp
--- a/Engine/Plugins/Runtime/RigVM/Source/RigVMEditor/Private/Editor/RigVMEditor.cpp
+++ b/Engine/Plugins/Runtime/RigVM/Source/RigVMEditor/Private/Editor/RigVMEditor.cpp
@@ -1371,6 +1371,23 @@ void FRigVMEditorBase::OnNodeTitleCommitted(const FText& NewText, ETextCommit::T
 			GetFocusedController()->SetCommentTextByName(CommentBeingChanged->GetFName(), NewText.ToString(), CommentBeingChanged->FontSize, CommentBeingChanged->bCommentBubbleVisible, CommentBeingChanged->bColorCommentBubble, true, true);
 		}
 	}
+	else if (const URigVMEdGraphNode* RigVMNodeBeingChanged = Cast<URigVMEdGraphNode>(NodeBeingChanged))
+	{
+		if(URigVMNode* ModelNode = RigVMNodeBeingChanged->GetModelNode())
+		{
+			if(URigVMController* Controller = GetFocusedController())
+			{
+				const FString NewNodeName = NewText.ToString();
+				if(!ModelNode->GetName().Equals(NewNodeName, ESearchCase::CaseSensitive))
+				{
+					if(Controller->RenameNode(ModelNode, *NewText.ToString(), true, true))
+					{
+						RigVMNodeBeingChanged->InvalidateNodeTitle();
+					}
+				}
+			}
+		}
+	}
 }
 
 void FRigVMEditorBase::FocusInspectorOnGraphSelection(const TSet<UObject*>& NewSelection, bool bForceRefresh)
diff --git a/Engine/Plugins/Runtime/RigVM/Source/RigVMEditor/Private/RigVMEditorModule.cpp b/Engine/Plugins/Runtime/RigVM/Source/RigVMEditor/Private/RigVMEditorModule.cpp
--- a/Engine/Plugins/Runtime/RigVM/Source/RigVMEditor/Private/RigVMEditorModule.cpp
+++ b/Engine/Plugins/Runtime/RigVM/Source/RigVMEditor/Private/RigVMEditorModule.cpp
@@ -1068,6 +1068,23 @@ void FRigVMEditorModule::GetNodeOrganizationContextMenuActions(IRigVMClientHost*
 	URigVMController* Controller = RigVMClientHost->GetRigVMClient()->GetController(Model);
 
 	FToolMenuSection& OrganizationSection = Menu->AddSection("RigVMEditorContextMenuOrganization", LOCTEXT("OrganizationHeader", "Organization"));
+
+	check(EdGraphNode);
+	check(ModelNode);
+
+	if(ModelNode->SupportsRenaming())
+	{
+		OrganizationSection.AddMenuEntry(
+			"Rename",
+			LOCTEXT("RenameNode", "Rename Node"),
+			LOCTEXT("RenameNode_Tooltip", "Constant nodes can be renamed"),
+			FSlateIcon(),
+			FUIAction(FExecuteAction::CreateLambda([EdGraphNode]() {
+				const_cast<URigVMEdGraphNode*>(EdGraphNode)->RequestRename();
+			})
+		));
+	}
+	
 	OrganizationSection.AddMenuEntry(
 		"Collapse Nodes",
 		LOCTEXT("CollapseNodes", "Collapse Nodes"),
@@ -1105,7 +1122,7 @@ void FRigVMEditorModule::GetNodeOrganizationContextMenuActions(IRigVMClientHost*
 		})
 	));
 
-	if (URigVMCollapseNode* CollapseNode = Cast<URigVMCollapseNode>(EdGraphNode->GetModelNode()))
+	if (URigVMCollapseNode* CollapseNode = Cast<URigVMCollapseNode>(ModelNode))
 	{
 		OrganizationSection.AddMenuEntry(
 			"Promote To Function",
@@ -1120,7 +1137,7 @@ void FRigVMEditorModule::GetNodeOrganizationContextMenuActions(IRigVMClientHost*
 
 	if(RigVMBlueprint)
 	{
-		if (URigVMFunctionReferenceNode* FunctionReferenceNode = Cast<URigVMFunctionReferenceNode>(EdGraphNode->GetModelNode()))
+		if (URigVMFunctionReferenceNode* FunctionReferenceNode = Cast<URigVMFunctionReferenceNode>(ModelNode))
 		{
 			TSoftObjectPtr<URigVMFunctionReferenceNode> RefPtr(FunctionReferenceNode);
 			if(RefPtr.GetLongPackageName() != FunctionReferenceNode->GetReferencedFunctionHeader().LibraryPointer.GetNodeSoftPath().GetLongPackageName())
@@ -1137,7 +1154,7 @@ void FRigVMEditorModule::GetNodeOrganizationContextMenuActions(IRigVMClientHost*
 			}
 		}
 
-		if (URigVMFunctionReferenceNode* FunctionRefNode = Cast<URigVMFunctionReferenceNode>(EdGraphNode->GetModelNode()))
+		if (URigVMFunctionReferenceNode* FunctionRefNode = Cast<URigVMFunctionReferenceNode>(ModelNode))
 		{
 			OrganizationSection.AddMenuEntry(
 				"Promote To Collapse Node",
@@ -1151,7 +1168,7 @@ void FRigVMEditorModule::GetNodeOrganizationContextMenuActions(IRigVMClientHost*
 		}
 	}
 
-	if (URigVMLibraryNode* LibraryNode = Cast<URigVMLibraryNode>(EdGraphNode->GetModelNode()))
+	if (URigVMLibraryNode* LibraryNode = Cast<URigVMLibraryNode>(ModelNode))
 	{
 		OrganizationSection.AddMenuEntry(
 			"Expand Node",
diff --git a/Engine/Plugins/Runtime/RigVM/Source/RigVMEditor/Private/Widgets/SRigVMGraphNode.cpp b/Engine/Plugins/Runtime/RigVM/Source/RigVMEditor/Private/Widgets/SRigVMGraphNode.cpp
--- a/Engine/Plugins/Runtime/RigVM/Source/RigVMEditor/Private/Widgets/SRigVMGraphNode.cpp
+++ b/Engine/Plugins/Runtime/RigVM/Source/RigVMEditor/Private/Widgets/SRigVMGraphNode.cpp
@@ -128,6 +128,20 @@ void SRigVMGraphNode::Construct( const FArguments& InArgs )
 	EdGraphNode->OnNodeBeginRemoval().AddSP(this, &SRigVMGraphNode::HandleNodeBeginRemoval);
 
 	LastHighDetailSize = FVector2D::ZeroVector;
+	
+	EdGraphNode->RequestRenameDelegate.BindSP(this, &SRigVMGraphNode::RequestRename);
+
+	if(EdGraphNode->bRenameIsPending)
+	{
+		TWeakObjectPtr<URigVMEdGraphNode> WeakNode = EdGraphNode;
+		FFunctionGraphTask::CreateAndDispatchWhenReady([WeakNode]()
+		{
+			if(WeakNode.IsValid())
+			{
+				(void)WeakNode->RequestRename();
+			}
+		}, TStatId(), NULL, ENamedThreads::GameThread);
+	}
 }
 
 TSharedRef<SWidget> SRigVMGraphNode::CreateNodeContentArea()
@@ -1095,6 +1109,12 @@ TArray<FOverlayWidgetInfo> SRigVMGraphNode::GetOverlayWidgets(bool bSelected, co
 	return Widgets;
 }
 
+void SRigVMGraphNode::OnFocusLost(const FFocusEvent& InFocusEvent)
+{
+	SGraphNode::OnFocusLost(InFocusEvent);
+	CancelRename();
+}
+
 void SRigVMGraphNode::RefreshErrorInfo()
 {
 	if (GraphNode)
diff --git a/Engine/Plugins/Runtime/RigVM/Source/RigVMEditor/Private/Widgets/SRigVMGraphNodeKnot.cpp b/Engine/Plugins/Runtime/RigVM/Source/RigVMEditor/Private/Widgets/SRigVMGraphNodeKnot.cpp
--- a/Engine/Plugins/Runtime/RigVM/Source/RigVMEditor/Private/Widgets/SRigVMGraphNodeKnot.cpp
+++ b/Engine/Plugins/Runtime/RigVM/Source/RigVMEditor/Private/Widgets/SRigVMGraphNodeKnot.cpp
@@ -1,7 +1,9 @@
 // Copyright Epic Games, Inc. All Rights Reserved.
 
 #include "Widgets/SRigVMGraphNodeKnot.h"
+#include "Widgets/SToolTip.h"
 #include "EdGraph/RigVMEdGraphSchema.h"
+#include "IDocumentation.h"
 
 #if WITH_EDITOR
 #include "Editor.h"
@@ -75,4 +77,15 @@ void SRigVMGraphNodeKnot::HandleNodeBeginRemoval()
 	InvalidateGraphData();
 }
 
+void SRigVMGraphNodeKnot::UpdateGraphNode()
+{
+	SGraphNodeKnot::UpdateGraphNode();
+
+	if (!SWidget::GetToolTip().IsValid())
+	{
+		TSharedRef<SToolTip> DefaultToolTip = IDocumentation::Get()->CreateToolTip( TAttribute< FText >( this, &SGraphNode::GetNodeTooltip ), NULL, GraphNode->GetDocumentationLink(), GraphNode->GetDocumentationExcerptName() );
+		SetToolTip(DefaultToolTip);
+	}
+}
+
 #undef LOCTEXT_NAMESPACE
diff --git a/Engine/Plugins/Runtime/RigVM/Source/RigVMEditor/Public/Widgets/SRigVMGraphNode.h b/Engine/Plugins/Runtime/RigVM/Source/RigVMEditor/Public/Widgets/SRigVMGraphNode.h
--- a/Engine/Plugins/Runtime/RigVM/Source/RigVMEditor/Public/Widgets/SRigVMGraphNode.h
+++ b/Engine/Plugins/Runtime/RigVM/Source/RigVMEditor/Public/Widgets/SRigVMGraphNode.h
@@ -51,6 +51,7 @@ public:
 	UE_API virtual void GetOverlayBrushes(bool bSelected, const FVector2f& WidgetSize, TArray<FOverlayBrushInfo>& Brushes) const override;
 	UE_API virtual void GetNodeInfoPopups(FNodeInfoContext* Context, TArray<FGraphInformationPopupInfo>& Popups) const override;
 	UE_API virtual TArray<FOverlayWidgetInfo> GetOverlayWidgets(bool bSelected, const FVector2f& WidgetSize) const override;
+	UE_API virtual void OnFocusLost(const FFocusEvent& InFocusEvent) override;
 
 	UE_API virtual void RefreshErrorInfo() override;
 	UE_API virtual void Tick(const FGeometry& AllottedGeometry, const double InCurrentTime, const float InDeltaTime) override;
diff --git a/Engine/Plugins/Runtime/RigVM/Source/RigVMEditor/Public/Widgets/SRigVMGraphNodeKnot.h b/Engine/Plugins/Runtime/RigVM/Source/RigVMEditor/Public/Widgets/SRigVMGraphNodeKnot.h
--- a/Engine/Plugins/Runtime/RigVM/Source/RigVMEditor/Public/Widgets/SRigVMGraphNodeKnot.h
+++ b/Engine/Plugins/Runtime/RigVM/Source/RigVMEditor/Public/Widgets/SRigVMGraphNodeKnot.h
@@ -20,6 +20,7 @@ class SRigVMGraphNodeKnot : public SGraphNodeKnot
 	UE_API virtual void MoveTo(const FVector2f& NewPosition, FNodeSet& NodeFilter, bool bMarkDirty = true) override;
 
 	UE_API void HandleNodeBeginRemoval();
+	UE_API void UpdateGraphNode() override;
 };
 
 #undef UE_API
diff --git a/Engine/Source/Editor/GraphEditor/Private/SGraphNode.cpp b/Engine/Source/Editor/GraphEditor/Private/SGraphNode.cpp
--- a/Engine/Source/Editor/GraphEditor/Private/SGraphNode.cpp
+++ b/Engine/Source/Editor/GraphEditor/Private/SGraphNode.cpp
@@ -1580,6 +1580,17 @@ void SGraphNode::ApplyRename()
 	}
 }
 
+void SGraphNode::CancelRename()
+{
+	if (InlineEditableText.IsValid())
+	{
+		if(InlineEditableText->IsInEditMode())
+		{
+			InlineEditableText->ExitEditingMode();
+		}
+	}
+}
+
 FSlateRect SGraphNode::GetTitleRect() const
 {
 	const FVector2f NodePosition = GetPosition2f();
diff --git a/Engine/Source/Editor/GraphEditor/Public/SGraphNode.h b/Engine/Source/Editor/GraphEditor/Public/SGraphNode.h
--- a/Engine/Source/Editor/GraphEditor/Public/SGraphNode.h
+++ b/Engine/Source/Editor/GraphEditor/Public/SGraphNode.h
@@ -260,6 +260,9 @@ public:
 	/** Sets node into rename state if supported */
 	UE_API void ApplyRename();
 
+	/** Cancels the rename editing mode */
+	UE_API void CancelRename();
+
 	/** return rect of the title area */
 	UE_API virtual FSlateRect GetTitleRect() const;
 
