RigVM: Follow up to 43263230

Avoid renaming of Break Struct nodes

#rnx

diff --git a/Engine/Plugins/Runtime/RigVM/Source/RigVM/Public/RigVMFunctions/RigVMDispatch_MakeStruct.h b/Engine/Plugins/Runtime/RigVM/Source/RigVM/Public/RigVMFunctions/RigVMDispatch_MakeStruct.h
--- a/Engine/Plugins/Runtime/RigVM/Source/RigVM/Public/RigVMFunctions/RigVMDispatch_MakeStruct.h
+++ b/Engine/Plugins/Runtime/RigVM/Source/RigVM/Public/RigVMFunctions/RigVMDispatch_MakeStruct.h
@@ -56,5 +56,6 @@ public:
 	RIGVM_API virtual FText GetNodeTooltip(const FRigVMTemplateTypeMap& InTypes) const override;
 	RIGVM_API virtual FString GetKeywords() const override;
 #endif
+	virtual bool SupportsRenaming() const override { return false; }
 };
 
