RigVM: Follow up to 43263230, fixing upgrade functionality for math nodes

#rnx

diff --git a/Engine/Plugins/Animation/ControlRig/Source/ControlRig/Private/Units/Hierarchy/RigUnit_BoneName.cpp b/Engine/Plugins/Animation/ControlRig/Source/ControlRig/Private/Units/Hierarchy/RigUnit_BoneName.cpp
--- a/Engine/Plugins/Animation/ControlRig/Source/ControlRig/Private/Units/Hierarchy/RigUnit_BoneName.cpp
+++ b/Engine/Plugins/Animation/ControlRig/Source/ControlRig/Private/Units/Hierarchy/RigUnit_BoneName.cpp
@@ -17,9 +17,10 @@ FRigUnit_ItemArray_Execute()
 
 FRigVMStructUpgradeInfo FRigUnit_ItemArray::GetUpgradeInfo() const
 {
-	FRigVMStructUpgradeInfo Info = FRigVMStructUpgradeInfo::MakeFromStructToFactory(StaticStruct(), FRigVMDispatch_Constant::StaticStruct());
-	Info.AddRemappedPin(TEXT("Items"), TEXT("Value"));
-	return Info;
+	return FRigVMStructUpgradeInfo::MakeFromStructToFactory(
+		StaticStruct(),
+		FRigVMDispatch_Constant::StaticStruct(),
+		{{TEXT("Items"), TEXT("Value")}});
 }
 
 FRigUnit_BoneName_Execute()
diff --git a/Engine/Plugins/Runtime/RigVM/Source/RigVM/Private/RigVMCore/RigVMStruct.cpp b/Engine/Plugins/Runtime/RigVM/Source/RigVM/Private/RigVMCore/RigVMStruct.cpp
--- a/Engine/Plugins/Runtime/RigVM/Source/RigVM/Private/RigVMCore/RigVMStruct.cpp
+++ b/Engine/Plugins/Runtime/RigVM/Source/RigVM/Private/RigVMCore/RigVMStruct.cpp
@@ -95,7 +95,7 @@ FRigVMStructUpgradeInfo::FRigVMStructUpgradeInfo()
 {
 }
 
-FRigVMStructUpgradeInfo FRigVMStructUpgradeInfo::MakeFromStructToFactory(UScriptStruct* InRigVMStruct, UScriptStruct* InFactoryStruct)
+FRigVMStructUpgradeInfo FRigVMStructUpgradeInfo::MakeFromStructToFactory(UScriptStruct* InRigVMStruct, UScriptStruct* InFactoryStruct, const TMap<FName,FName>& InArgumentNameMap)
 {
 	check(InRigVMStruct);
 	check(InFactoryStruct);
@@ -103,13 +103,67 @@ FRigVMStructUpgradeInfo FRigVMStructUpgradeInfo::MakeFromStructToFactory(UScript
 	FRigVMRegistry& Registry = FRigVMRegistry::Get();
 	const FRigVMDispatchFactory* Factory = Registry.FindOrAddDispatchFactory(InFactoryStruct);
 	check(Factory);
-	const FRigVMTemplateTypeMap Types = GetTypeMapFromStruct(InRigVMStruct);
+	FRigVMTemplateTypeMap Types = GetTypeMapFromStruct(InRigVMStruct);
+	{
+		TArray<FRigVMTemplateArgumentInfo> Infos;
+		{
+			FRigVMRegistryReadLock ReadLock(Registry);
+			Infos = Factory->GetArgumentInfos(ReadLock);
+		}
+
+		// if we have only one argument - it's quite possible we'll have to remap the argument name.
+		if(Types.Num() == 1 && Infos.Num() == 1)
+		{
+			FRigVMTemplateTypeMap NewTypes;
+			NewTypes.Add(Infos[0].Name, Types.begin().Value());
+			Swap(NewTypes, Types);
+		}
+
+		// if we have more arguments or a mismatch - we need to rely on the name map
+		else
+		{
+			FRigVMTemplateTypeMap NewTypes;
+			for(const FRigVMTemplateArgumentInfo& Info : Infos)
+			{
+				FName ArgumentName = Info.Name;
+				if(const FName* MappedArgumentName = InArgumentNameMap.FindKey(ArgumentName))
+				{
+					ArgumentName = *MappedArgumentName;
+				}
+				
+				if(const TRigVMTypeIndex* TypeIndex = Types.Find(ArgumentName))
+				{
+					NewTypes.Add(Info.Name, *TypeIndex);
+				}
+				else
+				{
+					UE_LOG(
+						LogRigVM,
+						Error, 
+						TEXT("Cannot upgrade node using '%s' to dispatch '%s' - type for argument '%s' cannot be determined.\nYou may need to provide the InArgumentNameMap."),
+						*InRigVMStruct->GetName(),
+						*InFactoryStruct->GetName(),
+						*Info.Name.ToString()
+					);
+					return FRigVMStructUpgradeInfo();
+				}
+			}
+			Swap(NewTypes, Types);
+		}
+	}
+	
 	const FString PermutationName = Factory->GetPermutationName(Types);
 
 	FRigVMStructUpgradeInfo Info;
 	Info.OldStruct = InRigVMStruct;
 	Info.NewStruct = InFactoryStruct;
 	Info.NewDispatchFunction = *PermutationName;
+
+	for(const TPair<FName, FName>& Pair : InArgumentNameMap)
+	{
+		Info.AddRemappedPin(Pair.Key.ToString(), Pair.Value.ToString());
+	}
+	
 	return Info;
 }
 
diff --git a/Engine/Plugins/Runtime/RigVM/Source/RigVM/Private/RigVMFunctions/Math/RigVMFunction_MathColor.cpp b/Engine/Plugins/Runtime/RigVM/Source/RigVM/Private/RigVMFunctions/Math/RigVMFunction_MathColor.cpp
--- a/Engine/Plugins/Runtime/RigVM/Source/RigVM/Private/RigVMFunctions/Math/RigVMFunction_MathColor.cpp
+++ b/Engine/Plugins/Runtime/RigVM/Source/RigVM/Private/RigVMFunctions/Math/RigVMFunction_MathColor.cpp
@@ -12,12 +12,15 @@ FRigVMFunction_MathColorMake_Execute()
 
 FRigVMStructUpgradeInfo FRigVMFunction_MathColorMake::GetUpgradeInfo() const
 {
-	FRigVMStructUpgradeInfo Info = FRigVMStructUpgradeInfo::MakeFromStructToFactory(StaticStruct(), FRigVMDispatch_Constant::StaticStruct());
+	FRigVMStructUpgradeInfo Info = FRigVMStructUpgradeInfo::MakeFromStructToFactory(
+		StaticStruct(),
+		FRigVMDispatch_Constant::StaticStruct(),
+		{{TEXT("Result"), TEXT("Value")}}
+	);
 	Info.AddRemappedPin(TEXT("R"), TEXT("Value.R"));
 	Info.AddRemappedPin(TEXT("G"), TEXT("Value.G"));
 	Info.AddRemappedPin(TEXT("B"), TEXT("Value.B"));
 	Info.AddRemappedPin(TEXT("A"), TEXT("Value.A"));
-	Info.AddRemappedPin(TEXT("Result"), TEXT("Value"));
 	return Info;
 }
 
diff --git a/Engine/Plugins/Runtime/RigVM/Source/RigVM/Private/RigVMFunctions/Math/RigVMFunction_MathTransform.cpp b/Engine/Plugins/Runtime/RigVM/Source/RigVM/Private/RigVMFunctions/Math/RigVMFunction_MathTransform.cpp
--- a/Engine/Plugins/Runtime/RigVM/Source/RigVM/Private/RigVMFunctions/Math/RigVMFunction_MathTransform.cpp
+++ b/Engine/Plugins/Runtime/RigVM/Source/RigVM/Private/RigVMFunctions/Math/RigVMFunction_MathTransform.cpp
@@ -14,11 +14,14 @@ FRigVMFunction_MathTransformMake_Execute()
 
 FRigVMStructUpgradeInfo FRigVMFunction_MathTransformMake::GetUpgradeInfo() const
 {
-	FRigVMStructUpgradeInfo Info = FRigVMStructUpgradeInfo::MakeFromStructToFactory(StaticStruct(), FRigVMDispatch_Constant::StaticStruct());
+	FRigVMStructUpgradeInfo Info = FRigVMStructUpgradeInfo::MakeFromStructToFactory(
+		StaticStruct(),
+		FRigVMDispatch_Constant::StaticStruct(),
+		{{TEXT("Result"), TEXT("Value")}}
+	);
 	Info.AddRemappedPin(TEXT("Translation"), TEXT("Value.Translation"));
 	Info.AddRemappedPin(TEXT("Rotation"), TEXT("Value.Rotation"));
 	Info.AddRemappedPin(TEXT("Scale"), TEXT("Value.Scale"));
-	Info.AddRemappedPin(TEXT("Result"), TEXT("Value"));
 	return Info;
 }
 
diff --git a/Engine/Plugins/Runtime/RigVM/Source/RigVM/Private/RigVMFunctions/Math/RigVMFunction_MathVector.cpp b/Engine/Plugins/Runtime/RigVM/Source/RigVM/Private/RigVMFunctions/Math/RigVMFunction_MathVector.cpp
--- a/Engine/Plugins/Runtime/RigVM/Source/RigVM/Private/RigVMFunctions/Math/RigVMFunction_MathVector.cpp
+++ b/Engine/Plugins/Runtime/RigVM/Source/RigVM/Private/RigVMFunctions/Math/RigVMFunction_MathVector.cpp
@@ -16,11 +16,14 @@ FRigVMFunction_MathVectorMake_Execute()
 
 FRigVMStructUpgradeInfo FRigVMFunction_MathVectorMake::GetUpgradeInfo() const
 {
-	FRigVMStructUpgradeInfo Info = FRigVMStructUpgradeInfo::MakeFromStructToFactory(StaticStruct(), FRigVMDispatch_Constant::StaticStruct());
+	FRigVMStructUpgradeInfo Info = FRigVMStructUpgradeInfo::MakeFromStructToFactory(
+		StaticStruct(),
+		FRigVMDispatch_Constant::StaticStruct(),
+		{{TEXT("Result"), TEXT("Value")}}
+	);
 	Info.AddRemappedPin(TEXT("X"), TEXT("Value.X"));
 	Info.AddRemappedPin(TEXT("Y"), TEXT("Value.Y"));
 	Info.AddRemappedPin(TEXT("Z"), TEXT("Value.Z"));
-	Info.AddRemappedPin(TEXT("Result"), TEXT("Value"));
 	return Info;
 }
 
diff --git a/Engine/Plugins/Runtime/RigVM/Source/RigVM/Public/RigVMCore/RigVMStructUpgradeInfo.h b/Engine/Plugins/Runtime/RigVM/Source/RigVM/Public/RigVMCore/RigVMStructUpgradeInfo.h
--- a/Engine/Plugins/Runtime/RigVM/Source/RigVM/Public/RigVMCore/RigVMStructUpgradeInfo.h
+++ b/Engine/Plugins/Runtime/RigVM/Source/RigVM/Public/RigVMCore/RigVMStructUpgradeInfo.h
@@ -34,7 +34,7 @@ public:
 		NewDispatchFunction = InNewDispatchFunction;
 	}
 
-	static UE_API FRigVMStructUpgradeInfo MakeFromStructToFactory(UScriptStruct* InRigVMStruct, UScriptStruct* InFactoryStruct);
+	static UE_API FRigVMStructUpgradeInfo MakeFromStructToFactory(UScriptStruct* InRigVMStruct, UScriptStruct* InFactoryStruct, const TMap<FName,FName>& InArgumentNameMap = TMap<FName,FName>());
 
 	// returns true if this upgrade info can be applied
 	UE_API bool IsValid() const;
