RigVM Compiler to collapse constants

#rnx

diff --git a/Engine/Plugins/Runtime/RigVM/Source/RigVMDeveloper/Private/RigVMCompiler/RigVMAST.cpp b/Engine/Plugins/Runtime/RigVM/Source/RigVMDeveloper/Private/RigVMCompiler/RigVMAST.cpp
--- a/Engine/Plugins/Runtime/RigVM/Source/RigVMDeveloper/Private/RigVMCompiler/RigVMAST.cpp
+++ b/Engine/Plugins/Runtime/RigVM/Source/RigVMDeveloper/Private/RigVMCompiler/RigVMAST.cpp
@@ -20,6 +20,8 @@
 #include "RigVMModel/Nodes/RigVMDispatchNode.h"
 #include "UObject/FieldIterator.h"
 #include "Algo/Sort.h"
+#include "RigVMFunctions/RigVMDispatch_Constant.h"
+#include "RigVMFunctions/RigVMDispatch_MakeStruct.h"
 
 #include UE_INLINE_GENERATED_CPP_BY_NAME(RigVMAST)
 
@@ -536,6 +538,7 @@ void FRigVMExprAST::AddParent(FRigVMExprAST* InParent)
 	ensure(IsValid());
 	ensure(InParent->IsValid());
 	ensure(InParent != this);
+	ensure(!InParent->IsA(FRigVMExprAST::EType::Literal));
 	
 	if (Parents.Contains(InParent))
 	{
@@ -1365,7 +1368,10 @@ FRigVMExprAST* FRigVMParserAST::CreateExpressionForNode(const FRigVMASTProxy& In
 			InNodeProxy.IsA<URigVMVariableNode>() ||
 			InNodeProxy.IsA<URigVMEnumNode>() ||
 			InNodeProxy.IsA<URigVMLibraryNode>() ||
-			InNodeProxy.IsA<URigVMFunctionInterfaceNode>())
+			InNodeProxy.IsA<URigVMFunctionInterfaceNode>() ||
+			IsConstantDispatchNode(Node) ||
+			IsMakeStructDispatchNode(Node) ||
+			IsBreakStructDispatchNode(Node))
 		{
 			NodeExpr = MakeExpr<FRigVMNoOpExprAST>(InNodeProxy);
 		}
@@ -1485,6 +1491,29 @@ FRigVMExprAST* FRigVMParserAST::TraversePin(const FRigVMASTProxy& InPinProxy, FR
 	ensure(!SubjectToExpression.Contains(InPinProxy));
 
 	URigVMPin* Pin = InPinProxy.GetSubjectChecked<URigVMPin>();
+
+	if(IsMakeStructDispatchStructPin(Pin))
+	{
+		URigVMPin* ElementsPin = Pin->GetNode()->FindRootPinByName(FRigVMDispatch_MakeStruct::ElementsName.Resolve());
+		check(ElementsPin);
+		
+		FRigVMASTProxy ElementsPinProxy = InPinProxy.GetSibling(ElementsPin);
+		FRigVMExprAST* ElementsPinExpression = SubjectToExpression.FindChecked(ElementsPinProxy);
+		SubjectToExpression.Add(InPinProxy, ElementsPinExpression);
+		return ElementsPinExpression;
+	}
+
+	if(IsBreakStructDispatchElementsPin(Pin))
+	{
+		URigVMPin* StructPin = Pin->GetNode()->FindRootPinByName(FRigVMDispatch_MakeStruct::StructName.Resolve());
+		check(StructPin);
+		
+		FRigVMASTProxy StructPinProxy = InPinProxy.GetSibling(StructPin);
+		FRigVMExprAST* StructPinExpression = SubjectToExpression.FindChecked(StructPinProxy);
+		SubjectToExpression.Add(InPinProxy, StructPinExpression);
+		return StructPinExpression;
+	}
+
 	URigVMPin::FPinOverride PinOverride(InPinProxy, PinOverrides);
 
 	TArray<int32> LinkIndices = GetSourceLinkIndices(InPinProxy, true);
@@ -1498,6 +1527,17 @@ FRigVMExprAST* FRigVMParserAST::TraversePin(const FRigVMASTProxy& InPinProxy, FR
 		);
 	}
 
+	bool bForceUseOfLiteralExpression = false;
+	if(IsConstantDispatchValuePin(Pin) || IsMakeStructDispatchElementsPin(Pin) || IsBreakStructDispatchStructPin(Pin))
+	{
+		// if this is a constant with no input links at all - just use a literal for it
+		const TArray<URigVMPin*> AllLinkedPins = Pin->GetRootPin()->GetLinkedSourcePins(true);
+		if(AllLinkedPins.IsEmpty())
+		{
+			bForceUseOfLiteralExpression = true;
+		}
+	}
+
 	FRigVMExprAST* PinExpr = nullptr;
 
 	if (Cast<URigVMVariableNode>(Pin->GetNode()))
@@ -1536,6 +1576,10 @@ FRigVMExprAST* FRigVMParserAST::TraversePin(const FRigVMASTProxy& InPinProxy, FR
 			}
 		}
 	}
+	else if (bForceUseOfLiteralExpression)
+	{
+		PinExpr = MakeExpr<FRigVMLiteralExprAST>(InPinProxy);
+	}
 	else if ((Pin->GetDirection() == ERigVMPinDirection::Input ||
 		Pin->GetDirection() == ERigVMPinDirection::Visible) &&
 		LinkIndices.Num() == 0)
@@ -1581,7 +1625,8 @@ FRigVMExprAST* FRigVMParserAST::TraversePin(const FRigVMASTProxy& InPinProxy, FR
 
 	if ((Pin->GetDirection() == ERigVMPinDirection::IO ||
 		Pin->GetDirection() == ERigVMPinDirection::Input)
-		&& !Pin->IsExecuteContext())
+		&& !Pin->IsExecuteContext()
+		&& !PinExpr->IsA(FRigVMExprAST::EType::Literal))
 	{
 		bool bHasSourceLinkToRoot = false;
 		URigVMPin* RootPin = Pin->GetRootPin();
@@ -2020,10 +2065,28 @@ void FRigVMParserAST::FoldAssignments()
 				continue;
 			}
 		}
+
+		// if we copy onto a sub pin - let's not fold this copy expression.
+		// we also don't allow to assign onto variable nodes
+		if(AssignExpr->GetType() == FRigVMExprAST::EType::Copy)
+		{
+			if(!TargetPin->IsRootPin())
+			{
+				continue;
+			}
+			if(TargetPin->GetNode()->IsA<URigVMVariableNode>())
+			{
+				continue;
+			}
+		}
 		
 		// non-input pins on anything but a reroute (passthrough or literal) or array node should be skipped
 		const bool bIsReroute = TargetPin->GetNode()->IsA<URigVMRerouteNode>();
-		if (TargetPin->GetDirection() != ERigVMPinDirection::Input && !bIsReroute)
+
+		// a copy / assign to a constant node root pin can be folded 
+		const bool bIsConstantValueTarget = IsConstantDispatchValuePin(TargetPin);
+		
+		if (TargetPin->GetDirection() != ERigVMPinDirection::Input && !bIsReroute && !bIsConstantValueTarget)
 		{
 			continue;
 		}
@@ -2094,7 +2157,34 @@ void FRigVMParserAST::FoldAssignments()
 			if (GrandParent->IsA(FRigVMExprAST::EType::Assign))
 			{
 				FRigVMAssignExprAST* GrandParentAssign = GrandParent->To<FRigVMAssignExprAST>();
-				GrandParentAssign->Link = FRigVMASTLinkDescription(AssignExpr->GetSourceProxy(),
+
+				const URigVMPin* OldSourcePin = GrandParentAssign->GetSourcePin();
+				const FString SourceSegmentPath = OldSourcePin->GetSegmentPath();
+				URigVMPin* NewSourcePin = AssignExpr->GetSourcePin();
+
+				if(IsMakeStructDispatchStructPin(NewSourcePin))
+				{
+					NewSourcePin = NewSourcePin->GetNode()->FindRootPinByName(FRigVMDispatch_MakeStruct::ElementsName.Resolve());
+					check(NewSourcePin);
+				}
+				
+				if(!SourceSegmentPath.IsEmpty())
+				{
+					NewSourcePin = NewSourcePin->FindSubPin(SourceSegmentPath);
+					if(NewSourcePin == nullptr)
+					{
+						Settings.Reportf(
+							EMessageSeverity::Error,
+							AssignExpr->GetSourcePin()->GetNode(),
+							TEXT("Cannot resolve segment path '%s' on pin '%s'."),
+							*SourceSegmentPath,
+							*AssignExpr->GetSourcePin()->GetPinPath());
+						continue;
+					}
+				}
+				
+				GrandParentAssign->Link = FRigVMASTLinkDescription(
+					GrandParentAssign->GetSourceProxy().GetSibling(NewSourcePin),
 					GrandParentAssign->GetTargetProxy(), FString());
 				GrandParentAssign->Name = *GetLinkAsString(GrandParentAssign->GetLink());
 			}
@@ -3574,3 +3664,85 @@ void FRigVMParserAST::IncrementCacheVersion() const
 {
 	CacheVersion++;
 }
+
+bool FRigVMParserAST::IsConstantDispatchNode(const URigVMNode* InNode)
+{
+	if(const URigVMDispatchNode* DispatchNode = Cast<URigVMDispatchNode>(InNode))
+	{
+		if(const FRigVMDispatchFactory* Factory = DispatchNode->GetFactory())
+		{
+			return Factory->GetScriptStruct() == FRigVMDispatch_Constant::StaticStruct();
+		}
+	}
+	return false;
+}
+
+bool FRigVMParserAST::IsConstantDispatchValuePin(const URigVMPin* InPin)
+{
+	if(InPin && InPin->IsRootPin() && IsConstantDispatchNode(InPin->GetNode()))
+	{
+		return InPin->GetFName() == FRigVMDispatch_Constant::ValueName.Resolve();
+	}
+	return false;
+}
+
+bool FRigVMParserAST::IsMakeStructDispatchNode(const URigVMNode* InNode)
+{
+	if(const URigVMDispatchNode* DispatchNode = Cast<URigVMDispatchNode>(InNode))
+	{
+		if(const FRigVMDispatchFactory* Factory = DispatchNode->GetFactory())
+		{
+			return Factory->GetScriptStruct() == FRigVMDispatch_MakeStruct::StaticStruct();
+		}
+	}
+	return false;
+}
+
+bool FRigVMParserAST::IsMakeStructDispatchElementsPin(const URigVMPin* InPin)
+{
+	if(InPin && InPin->IsRootPin() && IsMakeStructDispatchNode(InPin->GetNode()))
+	{
+		return InPin->GetFName() == FRigVMDispatch_MakeStruct::ElementsName.Resolve();
+	}
+	return false;
+}
+
+bool FRigVMParserAST::IsMakeStructDispatchStructPin(const URigVMPin* InPin)
+{
+	if(InPin && InPin->IsRootPin() && IsMakeStructDispatchNode(InPin->GetNode()))
+	{
+		return InPin->GetFName() == FRigVMDispatch_MakeStruct::StructName.Resolve();
+	}
+	return false;
+}
+
+bool FRigVMParserAST::IsBreakStructDispatchNode(const URigVMNode* InNode)
+{
+	if(const URigVMDispatchNode* DispatchNode = Cast<URigVMDispatchNode>(InNode))
+	{
+		if(const FRigVMDispatchFactory* Factory = DispatchNode->GetFactory())
+		{
+			return Factory->GetScriptStruct() == FRigVMDispatch_BreakStruct::StaticStruct();
+		}
+	}
+	return false;
+}
+
+bool FRigVMParserAST::IsBreakStructDispatchStructPin(const URigVMPin* InPin)
+{
+	if(InPin && InPin->IsRootPin() && IsBreakStructDispatchNode(InPin->GetNode()))
+	{
+		return InPin->GetFName() == FRigVMDispatch_MakeStruct::StructName.Resolve();
+	}
+	return false;
+}
+
+bool FRigVMParserAST::IsBreakStructDispatchElementsPin(const URigVMPin* InPin)
+{
+	if(InPin && InPin->IsRootPin() && IsBreakStructDispatchNode(InPin->GetNode()))
+	{
+		return InPin->GetFName() == FRigVMDispatch_MakeStruct::ElementsName.Resolve();
+	}
+	return false;
+
+}
diff --git a/Engine/Plugins/Runtime/RigVM/Source/RigVMDeveloper/Public/RigVMCompiler/RigVMAST.h b/Engine/Plugins/Runtime/RigVM/Source/RigVMDeveloper/Public/RigVMCompiler/RigVMAST.h
--- a/Engine/Plugins/Runtime/RigVM/Source/RigVMDeveloper/Public/RigVMCompiler/RigVMAST.h
+++ b/Engine/Plugins/Runtime/RigVM/Source/RigVMDeveloper/Public/RigVMCompiler/RigVMAST.h
@@ -1,4 +1,4 @@
-// Copyright Epic Games, Inc. All Rights Reserved.
+ // Copyright Epic Games, Inc. All Rights Reserved.
 
 #pragma once
 
@@ -1148,10 +1148,10 @@ struct FRigVMParserASTSettings
 {
 	GENERATED_BODY()
 
-	// fold assignments / copies
+	// fold root pin assignments / copies
 	UPROPERTY(EditAnywhere, Category = "AST")
 	bool bFoldAssignments = false;
-
+	
 	// fold literals and share memory
 	UPROPERTY(EditAnywhere, Category = "AST")
 	bool bFoldLiterals = false;
@@ -1426,6 +1426,15 @@ private:
 
 	UE_API void IncrementCacheVersion() const;
 
+	static bool IsConstantDispatchNode(const URigVMNode* InNode);
+	static bool IsConstantDispatchValuePin(const URigVMPin* InPin);
+ 	static bool IsMakeStructDispatchNode(const URigVMNode* InNode);
+	static bool IsMakeStructDispatchElementsPin(const URigVMPin* InPin);
+	static bool IsMakeStructDispatchStructPin(const URigVMPin* InPin);
+	static bool IsBreakStructDispatchNode(const URigVMNode* InNode);
+	static bool IsBreakStructDispatchStructPin(const URigVMPin* InPin);
+	static bool IsBreakStructDispatchElementsPin(const URigVMPin* InPin);
+
 	TMap<FRigVMASTProxy, FRigVMExprAST*> SubjectToExpression;
 	TMap<FRigVMASTProxy, int32> NodeExpressionIndex;
 	TArray<FRigVMExprAST*> Expressions;
