RigVM: Fix folding issue with make struct + break struct within functions

#rnx

diff --git a/Engine/Plugins/Runtime/RigVM/Source/RigVMDeveloper/Private/EdGraph/RigVMEdGraphNode.cpp b/Engine/Plugins/Runtime/RigVM/Source/RigVMDeveloper/Private/EdGraph/RigVMEdGraphNode.cpp
--- a/Engine/Plugins/Runtime/RigVM/Source/RigVMDeveloper/Private/EdGraph/RigVMEdGraphNode.cpp
+++ b/Engine/Plugins/Runtime/RigVM/Source/RigVMDeveloper/Private/EdGraph/RigVMEdGraphNode.cpp
@@ -73,7 +73,11 @@ FText URigVMEdGraphNode::GetNodeTitle(ENodeTitleType::Type TitleType) const
 		{
 			if(bCanRenameNode)
 			{
-				NodeTitle = FText::FromString(ModelNode->GetName());
+				FString NodeName = ModelNode->GetName();
+				NodeName.RemoveFromStart(TEXT("DISPATCH_"));
+				NodeName.RemoveFromStart(TEXT("RigVMFunction_"));
+				NodeName.RemoveFromStart(TEXT("RigVMDispatch_"));
+				NodeTitle = FText::FromString(NodeName);
 			}
 
 			else if (const URigVMUnitNode* UnitNode = Cast<URigVMUnitNode>(ModelNode))
diff --git a/Engine/Plugins/Runtime/RigVM/Source/RigVMDeveloper/Private/RigVMCompiler/RigVMAST.cpp b/Engine/Plugins/Runtime/RigVM/Source/RigVMDeveloper/Private/RigVMCompiler/RigVMAST.cpp
--- a/Engine/Plugins/Runtime/RigVM/Source/RigVMDeveloper/Private/RigVMCompiler/RigVMAST.cpp
+++ b/Engine/Plugins/Runtime/RigVM/Source/RigVMDeveloper/Private/RigVMCompiler/RigVMAST.cpp
@@ -882,7 +882,7 @@ const FRigVMExprAST* FRigVMNodeExprAST::FindExprWithPinName(const FName& InPinNa
 const FRigVMVarExprAST* FRigVMNodeExprAST::FindVarWithPinName(const FName& InPinName) const
 {
 	const FRigVMExprAST* Child = FindExprWithPinName(InPinName);
-	if (Child->IsA(FRigVMExprAST::Var))
+	if (Child && Child->IsA(FRigVMExprAST::Var))
 	{
 		return Child->To<FRigVMVarExprAST>();
 	}
@@ -1492,28 +1492,6 @@ FRigVMExprAST* FRigVMParserAST::TraversePin(const FRigVMASTProxy& InPinProxy, FR
 
 	URigVMPin* Pin = InPinProxy.GetSubjectChecked<URigVMPin>();
 
-	if(IsMakeStructDispatchStructPin(Pin))
-	{
-		URigVMPin* ElementsPin = Pin->GetNode()->FindRootPinByName(FRigVMDispatch_MakeStruct::ElementsName.Resolve());
-		check(ElementsPin);
-		
-		FRigVMASTProxy ElementsPinProxy = InPinProxy.GetSibling(ElementsPin);
-		FRigVMExprAST* ElementsPinExpression = SubjectToExpression.FindChecked(ElementsPinProxy);
-		SubjectToExpression.Add(InPinProxy, ElementsPinExpression);
-		return ElementsPinExpression;
-	}
-
-	if(IsBreakStructDispatchElementsPin(Pin))
-	{
-		URigVMPin* StructPin = Pin->GetNode()->FindRootPinByName(FRigVMDispatch_MakeStruct::StructName.Resolve());
-		check(StructPin);
-		
-		FRigVMASTProxy StructPinProxy = InPinProxy.GetSibling(StructPin);
-		FRigVMExprAST* StructPinExpression = SubjectToExpression.FindChecked(StructPinProxy);
-		SubjectToExpression.Add(InPinProxy, StructPinExpression);
-		return StructPinExpression;
-	}
-
 	URigVMPin::FPinOverride PinOverride(InPinProxy, PinOverrides);
 
 	TArray<int32> LinkIndices = GetSourceLinkIndices(InPinProxy, true);
@@ -1994,7 +1972,27 @@ void FRigVMParserAST::FoldNoOps()
 						}
 					}
 				}
+
+				// for make struct and break struct nodes we'll always refer to the input
+				// pin's expression since the nodes are considered no-op.
+				if(IsMakeStructDispatchNode(Node) || IsBreakStructDispatchNode(Node))
+				{
+					check(Expression->Children.Num() == 2);
+
+					FRigVMExprAST* ChildToKeep = Expression->Children[0];
+					FRigVMExprAST* ChildToRemove = Expression->Children[1];
+
+					TArray<FRigVMExprAST*> Parents = ChildToRemove->Parents;
+					for(FRigVMExprAST* Parent : Parents)
+					{
+						if(Parent != Expression)
+						{
+							Parent->ReplaceChild(ChildToRemove, ChildToKeep);
+						}
+					}
+				}
 			}
+
 			// copy since we are changing the content during iteration below
 			TArray<FRigVMExprAST*> Children = Expression->Children;
 			TArray<FRigVMExprAST*> Parents = Expression->Parents;
@@ -2162,12 +2160,6 @@ void FRigVMParserAST::FoldAssignments()
 				const FString SourceSegmentPath = OldSourcePin->GetSegmentPath();
 				URigVMPin* NewSourcePin = AssignExpr->GetSourcePin();
 
-				if(IsMakeStructDispatchStructPin(NewSourcePin))
-				{
-					NewSourcePin = NewSourcePin->GetNode()->FindRootPinByName(FRigVMDispatch_MakeStruct::ElementsName.Resolve());
-					check(NewSourcePin);
-				}
-				
 				if(!SourceSegmentPath.IsEmpty())
 				{
 					NewSourcePin = NewSourcePin->FindSubPin(SourceSegmentPath);
