RigVM: Upgrade all occurences of node workflow

This also has the copy instruction counting console command that I am going to remove again once done with my client asset refactor

#rnx

diff --git a/Engine/Plugins/Runtime/RigVM/Source/RigVMDeveloper/Private/RigVMModel/RigVMClient.cpp b/Engine/Plugins/Runtime/RigVM/Source/RigVMDeveloper/Private/RigVMModel/RigVMClient.cpp
--- a/Engine/Plugins/Runtime/RigVM/Source/RigVMDeveloper/Private/RigVMModel/RigVMClient.cpp
+++ b/Engine/Plugins/Runtime/RigVM/Source/RigVMDeveloper/Private/RigVMModel/RigVMClient.cpp
@@ -2,6 +2,7 @@
 
 #include "RigVMModel/RigVMClient.h"
 #include "Misc/TransactionObjectEvent.h"
+#include "Misc/ScopedSlowTask.h"
 #include "Exporters/Exporter.h"
 #include "UObject/ObjectSaveContext.h"
 #include "RigVMModel/RigVMControllerActions.h"
@@ -1881,3 +1882,101 @@ bool FRigVMClient::IsFunctionPublic(URigVMLibraryNode* InLibraryNode)
 	}
 	return false;
 }
+
+#if WITH_EDITOR
+
+bool FRigVMClient::UpgradeAllOccurencesOfNodes(const TArray<URigVMNode*>& InNodesToUpgrade, UObject* InOuter, bool bSetupUndoRedo)
+{
+	TArray<UScriptStruct*> ScriptStructsToFind;
+	for(const URigVMNode* NodeToUpgrade : InNodesToUpgrade)
+	{
+		if(!NodeToUpgrade || !IsValid(NodeToUpgrade))
+		{
+			return false;
+		}
+		
+		if(!NodeToUpgrade->CanBeUpgraded())
+		{
+			continue;
+		}
+		
+		if(const URigVMUnitNode* UnitNode = Cast<URigVMUnitNode>(NodeToUpgrade))
+		{
+			ScriptStructsToFind.AddUnique(UnitNode->GetScriptStruct());
+		}
+		else if(const URigVMDispatchNode* DispatchNode = Cast<URigVMDispatchNode>(NodeToUpgrade))
+		{
+			if(const FRigVMDispatchFactory* Factory = DispatchNode->GetFactory())
+			{
+				ScriptStructsToFind.AddUnique(Factory->GetScriptStruct());
+			}
+		}
+	}
+
+	if(ScriptStructsToFind.IsEmpty())
+	{
+		return false;
+	}
+
+	TArray<URigVMGraph*> AllGraphs = GetAllModelsLeavesFirst(true);
+	TArray<TArray<URigVMNode*>> NodesPerGraph;
+	NodesPerGraph.AddDefaulted(AllGraphs.Num());
+
+	for(int32 GraphIndex = 0; GraphIndex < AllGraphs.Num(); GraphIndex++)
+	{
+		const URigVMGraph* Graph = AllGraphs[GraphIndex];
+		
+		for(URigVMNode* Node : Graph->GetNodes())
+		{
+			if(const URigVMUnitNode* UnitNode = Cast<URigVMUnitNode>(Node))
+			{
+				if(ScriptStructsToFind.Contains(UnitNode->GetScriptStruct()))
+				{
+					NodesPerGraph[GraphIndex].Add(Node);
+				}
+			}
+			else if(const URigVMDispatchNode* DispatchNode = Cast<URigVMDispatchNode>(Node))
+			{
+				if(const FRigVMDispatchFactory* Factory = DispatchNode->GetFactory())
+				{
+					if(ScriptStructsToFind.Contains(Factory->GetScriptStruct()))
+					{
+						NodesPerGraph[GraphIndex].Add(Node);
+					}
+				}
+			}
+		}
+	}
+
+	TSharedPtr<FScopedTransaction> Transaction;
+	if(bSetupUndoRedo)
+	{
+		Transaction = MakeShared<FScopedTransaction>(NSLOCTEXT("RigVMClient", "UpgradeAllOccurencesOfNodes", "Upgrade Nodes"));
+		InOuter->Modify();
+	}
+	
+	FScopedSlowTask UpgradeNodesSlowTask((float)NodesPerGraph.Num(), NSLOCTEXT("RigVMClient", "UpgradeAllOccurencesOfNodes", "Upgrade Nodes"));
+	UpgradeNodesSlowTask.MakeDialog(true);
+
+	bool bResult = true;
+	for(int32 GraphIndex = 0; GraphIndex < AllGraphs.Num(); GraphIndex++)
+	{
+		UpgradeNodesSlowTask.EnterProgressFrame(1.0f, FText::FromString(FString::Printf(TEXT("Processing nodes of graph '%s'"), *AllGraphs[GraphIndex]->GetGraphName())));
+		UpgradeNodesSlowTask.ForceRefresh();
+		
+		if(NodesPerGraph[GraphIndex].IsEmpty())
+		{
+			continue;
+		}
+		URigVMController* Controller = GetOrCreateController((AllGraphs[GraphIndex]));
+		const TArray<URigVMNode*> UpgradedNodes = Controller->UpgradeNodes(NodesPerGraph[GraphIndex], true, bSetupUndoRedo);
+		if(UpgradedNodes.Num() != NodesPerGraph[GraphIndex].Num())
+		{
+			bResult = false;
+		}
+	}
+
+	return bResult;
+}
+
+#endif
\ No newline at end of file
diff --git a/Engine/Plugins/Runtime/RigVM/Source/RigVMDeveloper/Private/RigVMModel/RigVMController.cpp b/Engine/Plugins/Runtime/RigVM/Source/RigVMDeveloper/Private/RigVMModel/RigVMController.cpp
--- a/Engine/Plugins/Runtime/RigVM/Source/RigVMDeveloper/Private/RigVMModel/RigVMController.cpp
+++ b/Engine/Plugins/Runtime/RigVM/Source/RigVMDeveloper/Private/RigVMModel/RigVMController.cpp
@@ -21,6 +21,7 @@
 #include "UObject/PropertyPortFlags.h"
 #include "UObject/Package.h"
 #include "Misc/CoreMisc.h"
+#include "Misc/ScopedSlowTask.h"
 #include "Misc/StringOutputDevice.h"
 #include "Algo/Count.h"
 #include "Algo/Reverse.h"
@@ -2807,6 +2808,12 @@ TArray<URigVMNode*> URigVMController::UpgradeNodes(const TArray<URigVMNode*>& In
 		GetActionStack()->BeginAction(Action);
 	}
 
+	FScopedSlowTask UpgradeNodesSlowTask((float)(InNodes.Num() + 2), NSLOCTEXT("RigVMController", "UpgradeNodes", "Upgrade Nodes"));
+	UpgradeNodesSlowTask.MakeDialog(true);
+
+	UpgradeNodesSlowTask.EnterProgressFrame(1.0f, NSLOCTEXT("RigVMController", "DetachingLinks", "Detaching Links"));
+	UpgradeNodesSlowTask.ForceRefresh();
+
 	// find all links affecting the nodes to upgrade
 	TArray<FLinkedPath> LinkedPaths = GetLinkedPaths(InNodes, true);
 	if(!FastBreakLinkedPaths(LinkedPaths, bSetupUndoRedo))
@@ -2825,6 +2832,10 @@ TArray<URigVMNode*> URigVMController::UpgradeNodes(const TArray<URigVMNode*>& In
 	for(URigVMNode* Node : InNodes)
 	{
 		const FName OriginalName = Node->GetFName();
+
+		UpgradeNodesSlowTask.EnterProgressFrame(1.0f, FText::FromString(FString::Printf(TEXT("Upgrading node '%s'"), *OriginalName.ToString())));
+		UpgradeNodesSlowTask.ForceRefresh();
+		
 		FRigVMController_PinPathRemapDelegate RemapPinDelegate;
 		URigVMNode* UpgradedNode = UpgradeNode(Node, bSetupUndoRedo, &RemapPinDelegate);
 		if(UpgradedNode)
@@ -2843,6 +2854,9 @@ TArray<URigVMNode*> URigVMController::UpgradeNodes(const TArray<URigVMNode*>& In
 		}
 	}
 
+	UpgradeNodesSlowTask.EnterProgressFrame(1.0f, NSLOCTEXT("RigVMController", "RestoringLinks", "Restoring Links"));
+	UpgradeNodesSlowTask.ForceRefresh();
+
 	Settings.RemapDelegates = RemapPinDelegates;
 	RestoreLinkedPaths(LinkedPaths, Settings, bSetupUndoRedo);
 
@@ -19600,12 +19614,15 @@ void URigVMController::ApplyPinState(URigVMPin* InPin, const FPinState& InPinSta
 
 	if (!InPinState.DefaultValue.IsEmpty() || InPin->IsWildCard())
 	{
-		FString DefaultValue = InPinState.DefaultValue;
-		PostProcessDefaultValue(InPin, DefaultValue);
-		if(!DefaultValue.IsEmpty() || InPin->IsWildCard())
+		if(InPinState.Direction != ERigVMPinDirection::Output)
 		{
-			FRigVMDefaultValueTypeGuard _(this, InPinState.DefaultValueType, true);
-			SetPinDefaultValue(InPin, DefaultValue, true, bSetupUndoRedo, false);
+			FString DefaultValue = InPinState.DefaultValue;
+			PostProcessDefaultValue(InPin, DefaultValue);
+			if(!DefaultValue.IsEmpty() || InPin->IsWildCard())
+			{
+				FRigVMDefaultValueTypeGuard _(this, InPinState.DefaultValueType, true);
+				SetPinDefaultValue(InPin, DefaultValue, true, bSetupUndoRedo, false);
+			}
 		}
 	}
 
@@ -22082,25 +22099,26 @@ TArray<URigVMController::FLinkedPath> URigVMController::RemapLinkedPaths(
 			return InPinPath;
 		}
 
-		FString PinPath = InPinPath;
-		FString RemappedPinPath;
+		FString RemappedPinPath = InPinPath;
 
 		if (InSettings.bFollowCoreRedirectors)
 		{
 			FString RedirectedSourcePinPath;
-			if (ShouldRedirectPin(PinPath, RedirectedSourcePinPath))
+			if (ShouldRedirectPin(RemappedPinPath, RedirectedSourcePinPath))
 			{
 				RemappedPinPath = RedirectedSourcePinPath;
 			}
 		}
 
-		if(const FRigVMController_PinPathRemapDelegate* RemapDelegate = InSettings.RemapDelegates.Find(NodeName))
+		if(const FString* RemappedNodeName = InSettings.NodeNameMap.Find(NodeName))
 		{
-			RemappedPinPath = RemapDelegate->Execute(PinPath, bAsInput);
+			NodeName = *RemappedNodeName;
+			RemappedPinPath = URigVMPin::JoinPinPath(*RemappedNodeName, SegmentPath);
 		}
-		else if(const FString* RemappedNodeName = InSettings.NodeNameMap.Find(NodeName))
+
+		if(const FRigVMController_PinPathRemapDelegate* RemapDelegate = InSettings.RemapDelegates.Find(NodeName))
 		{
-			RemappedPinPath = URigVMPin::JoinPinPath(*RemappedNodeName, SegmentPath);
+			RemappedPinPath = RemapDelegate->Execute(RemappedPinPath, bAsInput);
 		}
 
 		// if the pin cannot be found based on this path the pin may be a sub pin on something that
@@ -22132,10 +22150,10 @@ TArray<URigVMController::FLinkedPath> URigVMController::RemapLinkedPaths(
 
 		if(!RemappedPinPath.IsEmpty() && GetGraph()->FindPin(RemappedPinPath))
 		{
-			PinPath = RemappedPinPath;
+			return RemappedPinPath;
 		}
 
-		return PinPath;
+		return InPinPath;
 	};
 
 	TArray<FLinkedPath> RemappedLinkedPaths;
diff --git a/Engine/Plugins/Runtime/RigVM/Source/RigVMDeveloper/Public/RigVMModel/RigVMClient.h b/Engine/Plugins/Runtime/RigVM/Source/RigVMDeveloper/Public/RigVMModel/RigVMClient.h
--- a/Engine/Plugins/Runtime/RigVM/Source/RigVMDeveloper/Public/RigVMModel/RigVMClient.h
+++ b/Engine/Plugins/Runtime/RigVM/Source/RigVMDeveloper/Public/RigVMModel/RigVMClient.h
@@ -330,6 +330,10 @@ public:
 	UE_API bool UpdateGraphFunctionSerializedGraph(URigVMLibraryNode* InLibraryNode);
 	UE_API bool IsFunctionPublic(URigVMLibraryNode* InLibraryNode);
 
+#if WITH_EDITOR
+	UE_API bool UpgradeAllOccurencesOfNodes(const TArray<URigVMNode*>& InNodesToUpgrade,  UObject* InOuter, bool bSetupUndoRedo = true);
+#endif
+
 	static constexpr TCHAR RigVMModelPrefix[] = TEXT("RigVMModel");
 
 private:
diff --git a/Engine/Plugins/Runtime/RigVM/Source/RigVMEditor/Private/Editor/RigVMEditorTools.cpp b/Engine/Plugins/Runtime/RigVM/Source/RigVMEditor/Private/Editor/RigVMEditorTools.cpp
--- a/Engine/Plugins/Runtime/RigVM/Source/RigVMEditor/Private/Editor/RigVMEditorTools.cpp
+++ b/Engine/Plugins/Runtime/RigVM/Source/RigVMEditor/Private/Editor/RigVMEditorTools.cpp
@@ -9,6 +9,7 @@
 #include "RigVMCore/RigVMGraphFunctionDefinition.h"
 #include "Widgets/SRigVMGraphFunctionLocalizationWidget.h"
 #include "Subsystems/EditorAssetSubsystem.h"
+#include "UObject/UObjectIterator.h"
 
 namespace UE::RigVM::Editor::Tools
 {
@@ -203,3 +204,66 @@ bool FFilterByAssetTag::PassesFilter(const FContentBrowserItem& InItem) const
 
 
 } // end namespace UE::RigVM::Editor::Tools
+
+FAutoConsoleCommand RigVMCountCopyInstructions(TEXT("RigVM.CountCopyInstructions"),
+	TEXT("Iterates all open RigVM VMs and counts the copy instructions"),
+	FConsoleCommandDelegate::CreateLambda([]() {
+
+		for (TObjectIterator<URigVM> RigVMIt; RigVMIt; ++RigVMIt)
+		{
+			const URigVM* VM = *RigVMIt;
+			if(!IsValid(VM))
+			{
+				continue;
+			}
+			
+			int32 CopiesRootToRoot = 0;
+			int32 CopiesSubPinToRoot = 0;
+			int32 CopiesToSubPin = 0;
+
+			const FRigVMByteCode& ByteCode = VM->GetByteCode();
+			const FRigVMInstructionArray Instructions = ByteCode.GetInstructions();
+
+			for(const FRigVMInstruction& Instruction : Instructions)
+			{
+				if(Instruction.OpCode != ERigVMOpCode::Copy)
+				{
+					continue;
+				}
+
+				FRigVMOperandArray Operands = ByteCode.GetOperandsForOp(Instruction);
+				check(Operands.Num() == 2);
+
+				const FRigVMOperand& Source = Operands[0];
+				const FRigVMOperand& Target = Operands[1];
+
+				if(Source.GetRegisterOffset() == INDEX_NONE && Target.GetRegisterOffset() == INDEX_NONE)
+				{
+					CopiesRootToRoot++;
+				}
+				else if(Source.GetRegisterOffset() != INDEX_NONE && Target.GetRegisterOffset() == INDEX_NONE)
+				{
+					CopiesSubPinToRoot++;
+				}
+				else
+				{
+					CopiesToSubPin++;
+				}
+			}
+
+			const FString PathName = VM->GetPathName();
+			if(CopiesRootToRoot > 0)
+			{
+				UE_LOG(LogRigVMDeveloper, Display, TEXT("%s: Copies from root to root: %d out of %d instructions."), *PathName, CopiesRootToRoot, Instructions.Num());
+			}
+			if(CopiesSubPinToRoot > 0)
+			{
+				UE_LOG(LogRigVMDeveloper, Display, TEXT("%s: Copies from sub pin to root: %d out of %d instructions."), *PathName, CopiesSubPinToRoot, Instructions.Num());
+			}
+			if(CopiesToSubPin > 0)
+			{
+				UE_LOG(LogRigVMDeveloper, Display, TEXT("%s: Copies to sub pin (composition): %d out of %d instructions."), *PathName, CopiesToSubPin, Instructions.Num());
+			}
+		}
+	})
+);
diff --git a/Engine/Plugins/Runtime/RigVM/Source/RigVMEditor/Private/RigVMEditorModule.cpp b/Engine/Plugins/Runtime/RigVM/Source/RigVMEditor/Private/RigVMEditorModule.cpp
--- a/Engine/Plugins/Runtime/RigVM/Source/RigVMEditor/Private/RigVMEditorModule.cpp
+++ b/Engine/Plugins/Runtime/RigVM/Source/RigVMEditor/Private/RigVMEditorModule.cpp
@@ -1192,6 +1192,79 @@ void FRigVMEditorModule::GetNodeOrganizationContextMenuActions(IRigVMClientHost*
 		));
 	}
 
+	OrganizationSection.AddMenuEntry(
+		"Select Similar Nodes",
+		LOCTEXT("SelectSimilarNodes", "Select Similar Nodes"),
+		LOCTEXT("SelectSimilarNodes_Tooltip", "Selects all matching nodes to the selection in this graph"),
+		FSlateIcon(),
+		FUIAction(FExecuteAction::CreateLambda([Model, Controller]() {
+			TArray<FName> Nodes = Model->GetSelectNodes();
+
+			TArray<const FRigVMFunction*> RigVMFunctions;
+			TArray<FRigVMGraphFunctionIdentifier> FunctionIdentifiers;
+
+			for(const FName& NodeName : Nodes)
+			{
+				const URigVMNode* Node = Model->FindNodeByName(NodeName);
+				if(Node == nullptr)
+				{
+					continue;
+				}
+
+				if(const URigVMFunctionReferenceNode* FunctionRefNode = Cast<URigVMFunctionReferenceNode>(Node))
+				{
+					const FRigVMGraphFunctionIdentifier Identifier = FunctionRefNode->GetFunctionIdentifier();
+					if(Identifier.IsValid())
+					{
+						FunctionIdentifiers.AddUnique(Identifier);
+					}
+				}
+				else if(const URigVMTemplateNode* TemplateNode = Cast<URigVMTemplateNode>(Node))
+				{
+					if(const FRigVMFunction* Function = TemplateNode->GetResolvedFunction())
+					{
+						RigVMFunctions.AddUnique(Function);
+					}
+				}
+			}
+
+			if(RigVMFunctions.IsEmpty() && FunctionIdentifiers.IsEmpty())
+			{
+				return;
+			}
+			
+			TArray<FName> SimilarNodes;
+			SimilarNodes.Append(Nodes);
+
+			for(const URigVMNode* Node : Model->GetNodes())
+			{
+				if(const URigVMFunctionReferenceNode* FunctionRefNode = Cast<URigVMFunctionReferenceNode>(Node))
+				{
+					const FRigVMGraphFunctionIdentifier Identifier = FunctionRefNode->GetFunctionIdentifier();
+					if(Identifier.IsValid())
+					{
+						if(FunctionIdentifiers.Contains(Identifier))
+						{
+							SimilarNodes.AddUnique(Node->GetFName());
+						}
+					}
+				}
+				else if(const URigVMTemplateNode* TemplateNode = Cast<URigVMTemplateNode>(Node))
+				{
+					if(const FRigVMFunction* Function = TemplateNode->GetResolvedFunction())
+					{
+						if(RigVMFunctions.Contains(Function))
+						{
+							SimilarNodes.AddUnique(Node->GetFName());
+						}
+					}
+				}
+			}
+
+			Controller->SetNodeSelection(SimilarNodes, true);
+		})
+	));
+
 	OrganizationSection.AddSubMenu("Alignment", LOCTEXT("AlignmentHeader", "Alignment"), FText(), FNewToolMenuDelegate::CreateLambda([](UToolMenu* AlignmentMenu)
 	{
 		{
@@ -1240,13 +1313,33 @@ void FRigVMEditorModule::GetNodeVariantContextMenuActions(IRigVMClientHost* RigV
 			VariantSection.AddMenuEntry(
 				"Upgrade Nodes",
 				LOCTEXT("UpgradeNodes", "Upgrade Nodes"),
-				LOCTEXT("UpgradeNodes_Tooltip", "Upgrades deprecated nodes to their current implementation"),
+				LOCTEXT("UpgradeNodes_Tooltip", "Upgrades selected outdated nodes to their current implementation"),
 				FSlateIcon(),
 				FUIAction(FExecuteAction::CreateLambda([Model, Controller]() {
 					TArray<FName> Nodes = Model->GetSelectNodes();
 					Controller->UpgradeNodes(Nodes, true, true);
 				})
 			));
+
+			VariantSection.AddMenuEntry(
+				"Upgrade Nodes All Occurences",
+				LOCTEXT("UpgradeNodesAllOccurences", "Upgrade Nodes (All Occurences)"),
+				LOCTEXT("UpgradeNodesAllOccurences_Tooltip", "Upgrades all occurences of selected outdated nodes to their current implementation"),
+				FSlateIcon(),
+				FUIAction(FExecuteAction::CreateLambda([RigVMClientHost, Model]() {
+					if(FRigVMClient* Client = RigVMClientHost->GetRigVMClient())
+					{
+						TArray<FName> NodeNames = Model->GetSelectNodes();
+						TArray<URigVMNode*> Nodes;
+						Nodes.Reserve(NodeNames.Num());
+						for(const FName& NodeName : NodeNames)
+						{
+							Nodes.Add(Model->FindNodeByName(NodeName));
+						}
+						Client->UpgradeAllOccurencesOfNodes(Nodes, Model->GetOutermost(), true);
+					}
+				})
+			));
 		}
 
 		if(bIsFunctionReference)
