RigVM: Compiler / AST - Fold sub pin copies

Move from Handle.GetData to GetInputData, GetOutputData and GetPrivateData
Make GetData protected (to cause compiler errors in dispatch implementations)
Adapt all dispatch implementations
Change UHT accordingly
Allow AST to fold copies from a subpin onto a root pin
Disable new foldsubpincopies by default (and only enable for single unit test)

#rnx

diff --git a/Engine/Plugins/Animation/ControlRig/Source/ControlRig/Private/Units/Animation/RigUnit_AnimAttribute.cpp b/Engine/Plugins/Animation/ControlRig/Source/ControlRig/Private/Units/Animation/RigUnit_AnimAttribute.cpp
--- a/Engine/Plugins/Animation/ControlRig/Source/ControlRig/Private/Units/Animation/RigUnit_AnimAttribute.cpp
+++ b/Engine/Plugins/Animation/ControlRig/Source/ControlRig/Private/Units/Animation/RigUnit_AnimAttribute.cpp
@@ -156,16 +156,16 @@ void FRigDispatch_GetAnimAttribute::GetAnimAttributeDispatch(FRigVMExtendedExecu
 #endif
 
 	// unpack the memory
-	const FName& Name = *(const FName*)Handles[Factory->NameArgIndex].GetData();
-	const FName& BoneName = *(const FName*)Handles[Factory->BoneNameArgIndex].GetData();
-	const uint8* Default = (const uint8*)Handles[Factory->DefaultArgIndex].GetData();
+	const FName& Name = *(const FName*)Handles[Factory->NameArgIndex].GetInputData();
+	const FName& BoneName = *(const FName*)Handles[Factory->BoneNameArgIndex].GetInputData();
+	const uint8* Default = (const uint8*)Handles[Factory->DefaultArgIndex].GetInputData();
 	
-	uint8* Value = (uint8*)Handles[Factory->ValueArgIndex].GetData();
-	bool& Found = *(bool*)Handles[Factory->FoundArgIndex].GetData();
+	uint8* Value = (uint8*)Handles[Factory->ValueArgIndex].GetOutputData();
+	bool& Found = *(bool*)Handles[Factory->FoundArgIndex].GetOutputData();
 	Found = false;
 	
-	FName& CachedBoneName = *(FName*)Handles[Factory->CachedBoneNameArgIndex].GetData(false, InContext.GetSlice().GetIndex());
-	int32& CachedBoneIndex = *(int32*)Handles[Factory->CachedBoneIndexArgIndex].GetData(false, InContext.GetSlice().GetIndex());
+	FName& CachedBoneName = *(FName*)Handles[Factory->CachedBoneNameArgIndex].GetPrivateData(InContext.GetSlice().GetIndex());
+	int32& CachedBoneIndex = *(int32*)Handles[Factory->CachedBoneIndexArgIndex].GetPrivateData(InContext.GetSlice().GetIndex());
 	
 	if (const FStructProperty* StructProperty = CastFieldChecked<FStructProperty>(Handles[Factory->ValueArgIndex].GetResolvedProperty()))
 	{
@@ -219,15 +219,15 @@ void FRigDispatch_SetAnimAttribute::SetAnimAttributeDispatch(FRigVMExtendedExecu
 #endif
 
 	// unpack the memory
-	const FName& Name = *(const FName*)Handles[Factory->NameArgIndex].GetData();
-	const FName& BoneName = *(const FName*)Handles[Factory->BoneNameArgIndex].GetData();
-	const uint8* Value = (uint8*)Handles[Factory->ValueArgIndex].GetData();
+	const FName& Name = *(const FName*)Handles[Factory->NameArgIndex].GetInputData();
+	const FName& BoneName = *(const FName*)Handles[Factory->BoneNameArgIndex].GetInputData();
+	const uint8* Value = (uint8*)Handles[Factory->ValueArgIndex].GetInputData();
 	
-	bool& Success = *(bool*)Handles[Factory->SuccessArgIndex].GetData();
+	bool& Success = *(bool*)Handles[Factory->SuccessArgIndex].GetOutputData();
 	Success = false;
 	
-	FName& CachedBoneName = *(FName*)Handles[Factory->CachedBoneNameArgIndex].GetData(false, InContext.GetSlice().GetIndex());
-	int32& CachedBoneIndex = *(int32*)Handles[Factory->CachedBoneIndexArgIndex].GetData(false, InContext.GetSlice().GetIndex());
+	FName& CachedBoneName = *(FName*)Handles[Factory->CachedBoneNameArgIndex].GetPrivateData(InContext.GetSlice().GetIndex());
+	int32& CachedBoneIndex = *(int32*)Handles[Factory->CachedBoneIndexArgIndex].GetPrivateData(InContext.GetSlice().GetIndex());
 	
 	if (const FStructProperty* StructProperty = CastFieldChecked<FStructProperty>(Handles[Factory->ValueArgIndex].GetResolvedProperty()))
 	{
diff --git a/Engine/Plugins/Animation/ControlRig/Source/ControlRig/Private/Units/Core/RigUnit_UserData.cpp b/Engine/Plugins/Animation/ControlRig/Source/ControlRig/Private/Units/Core/RigUnit_UserData.cpp
--- a/Engine/Plugins/Animation/ControlRig/Source/ControlRig/Private/Units/Core/RigUnit_UserData.cpp
+++ b/Engine/Plugins/Animation/ControlRig/Source/ControlRig/Private/Units/Core/RigUnit_UserData.cpp
@@ -96,11 +96,11 @@ void FRigDispatch_GetUserData::Execute(FRigVMExtendedExecuteContext& InContext,
 	check(Handles[4].IsBool());
 #endif
 
-	const FString& NameSpace = *(FString*)Handles[0].GetData();
-	const FString& Path = *(FString*)Handles[1].GetData();
-	const uint8* Default = Handles[2].GetData();
-	uint8* Result = Handles[3].GetData();
-	bool& bFound = *(bool*)Handles[4].GetData();
+	const FString& NameSpace = *(FString*)Handles[0].GetInputData();
+	const FString& Path = *(FString*)Handles[1].GetInputData();
+	const uint8* Default = Handles[2].GetInputData();
+	uint8* Result = Handles[3].GetOutputData();
+	bool& bFound = *(bool*)Handles[4].GetOutputData();
 	bFound = false;
 
 	auto Log = [&ControlRigContext](const FString& Message)
diff --git a/Engine/Plugins/Animation/ControlRig/Source/ControlRig/Private/Units/Hierarchy/RigUnit_Component.cpp b/Engine/Plugins/Animation/ControlRig/Source/ControlRig/Private/Units/Hierarchy/RigUnit_Component.cpp
--- a/Engine/Plugins/Animation/ControlRig/Source/ControlRig/Private/Units/Hierarchy/RigUnit_Component.cpp
+++ b/Engine/Plugins/Animation/ControlRig/Source/ControlRig/Private/Units/Hierarchy/RigUnit_Component.cpp
@@ -142,11 +142,11 @@ void FRigDispatch_SpawnComponent::Execute(FRigVMExtendedExecuteContext& InContex
 	const FRigDispatch_SpawnComponent* Factory = static_cast<const FRigDispatch_SpawnComponent*>(InContext.Factory);
 		
 	// unpack the memory
-	const FRigElementKey& ElementKey = *(const FRigElementKey*)Handles[Factory->ItemArgIndex].GetData();
-	const FName& Name = *(const FName*)Handles[Factory->NameArgIndex].GetData();
-	const uint8* SourceComponent = Handles[Factory->ComponentArgIndex].GetData();
-	FRigComponentKey& ComponentKey = *(FRigComponentKey*)Handles[Factory->KeyArgIndex].GetData();
-	bool& bSuccess = *(bool*)Handles[Factory->SuccessArgIndex].GetData();
+	const FRigElementKey& ElementKey = *(const FRigElementKey*)Handles[Factory->ItemArgIndex].GetInputData();
+	const FName& Name = *(const FName*)Handles[Factory->NameArgIndex].GetInputData();
+	const uint8* SourceComponent = Handles[Factory->ComponentArgIndex].GetInputData();
+	FRigComponentKey& ComponentKey = *(FRigComponentKey*)Handles[Factory->KeyArgIndex].GetOutputData();
+	bool& bSuccess = *(bool*)Handles[Factory->SuccessArgIndex].GetOutputData();
 
 	const FStructProperty* StructProperty = CastField<FStructProperty>(Handles[Factory->ComponentArgIndex].GetResolvedProperty());
 	check(StructProperty);
@@ -207,9 +207,9 @@ void FRigDispatch_GetComponentContent::Execute(FRigVMExtendedExecuteContext& InC
 	const FRigDispatch_GetComponentContent* Factory = static_cast<const FRigDispatch_GetComponentContent*>(InContext.Factory);
 			
 	// unpack the memory
-	const FRigComponentKey& ComponentKey = *(const FRigComponentKey*)Handles[Factory->KeyArgIndex].GetData();
-	uint8* TargetComponent = Handles[Factory->ComponentArgIndex].GetData();
-	bool& bSuccess = *(bool*)Handles[Factory->SuccessArgIndex].GetData();
+	const FRigComponentKey& ComponentKey = *(const FRigComponentKey*)Handles[Factory->KeyArgIndex].GetInputData();
+	uint8* TargetComponent = Handles[Factory->ComponentArgIndex].GetOutputData();
+	bool& bSuccess = *(bool*)Handles[Factory->SuccessArgIndex].GetOutputData();
 
 	const FStructProperty* StructProperty = CastField<FStructProperty>(Handles[Factory->ComponentArgIndex].GetResolvedProperty());
 	check(StructProperty);
@@ -263,9 +263,9 @@ void FRigDispatch_SetComponentContent::Execute(FRigVMExtendedExecuteContext& InC
 	const FRigDispatch_SetComponentContent* Factory = static_cast<const FRigDispatch_SetComponentContent*>(InContext.Factory);
 			
 	// unpack the memory
-	const FRigComponentKey& ComponentKey = *(const FRigComponentKey*)Handles[Factory->KeyArgIndex].GetData();
-	const  uint8* SourceComponent = Handles[Factory->ComponentArgIndex].GetData();
-	bool& bSuccess = *(bool*)Handles[Factory->SuccessArgIndex].GetData();
+	const FRigComponentKey& ComponentKey = *(const FRigComponentKey*)Handles[Factory->KeyArgIndex].GetInputData();
+	const  uint8* SourceComponent = Handles[Factory->ComponentArgIndex].GetInputData();
+	bool& bSuccess = *(bool*)Handles[Factory->SuccessArgIndex].GetOutputData();
 
 	const FStructProperty* StructProperty = CastField<FStructProperty>(Handles[Factory->ComponentArgIndex].GetResolvedProperty());
 	check(StructProperty);
diff --git a/Engine/Plugins/Animation/ControlRig/Source/ControlRig/Public/Units/Animation/RigUnit_AnimAttribute.h b/Engine/Plugins/Animation/ControlRig/Source/ControlRig/Public/Units/Animation/RigUnit_AnimAttribute.h
--- a/Engine/Plugins/Animation/ControlRig/Source/ControlRig/Public/Units/Animation/RigUnit_AnimAttribute.h
+++ b/Engine/Plugins/Animation/ControlRig/Source/ControlRig/Public/Units/Animation/RigUnit_AnimAttribute.h
@@ -243,15 +243,15 @@ protected:
 #endif
 
 		// unpack the memory
-		const FName& Name = *(const FName*)Handles[Factory->NameArgIndex].GetData();
-		const FName& BoneName = *(const FName*)Handles[Factory->BoneNameArgIndex].GetData();
-		const ValueType& Default = *(const ValueType*)Handles[Factory->DefaultArgIndex].GetData();
+		const FName& Name = *(const FName*)Handles[Factory->NameArgIndex].GetInputData();
+		const FName& BoneName = *(const FName*)Handles[Factory->BoneNameArgIndex].GetInputData();
+		const ValueType& Default = *(const ValueType*)Handles[Factory->DefaultArgIndex].GetInputData();
 		
-		ValueType& Value = *(ValueType*)Handles[Factory->ValueArgIndex].GetData();
-		bool& Found = *(bool*)Handles[Factory->FoundArgIndex].GetData();
+		ValueType& Value = *(ValueType*)Handles[Factory->ValueArgIndex].GetOutputData();
+		bool& Found = *(bool*)Handles[Factory->FoundArgIndex].GetOutputData();
 		
-		FName& CachedBoneName = *(FName*)Handles[Factory->CachedBoneNameArgIndex].GetData(false, InContext.GetSlice().GetIndex());
-		int32& CachedBoneIndex = *(int32*)Handles[Factory->CachedBoneIndexArgIndex].GetData(false, InContext.GetSlice().GetIndex());
+		FName& CachedBoneName = *(FName*)Handles[Factory->CachedBoneNameArgIndex].GetPrivateData(InContext.GetSlice().GetIndex());
+		int32& CachedBoneIndex = *(int32*)Handles[Factory->CachedBoneIndexArgIndex].GetPrivateData(InContext.GetSlice().GetIndex());
 
 		// extract the animation attribute
 		const FControlRigExecuteContext& Context = InContext.GetPublicDataSafe<FControlRigExecuteContext>();
@@ -322,15 +322,15 @@ protected:
 #endif
 
 		// unpack the memory
-		const FName& Name = *(const FName*)Handles[Factory->NameArgIndex].GetData();
-		const FName& BoneName = *(const FName*)Handles[Factory->BoneNameArgIndex].GetData();
-		const ValueType& Value = *(ValueType*)Handles[Factory->ValueArgIndex].GetData();
+		const FName& Name = *(const FName*)Handles[Factory->NameArgIndex].GetInputData();
+		const FName& BoneName = *(const FName*)Handles[Factory->BoneNameArgIndex].GetInputData();
+		const ValueType& Value = *(ValueType*)Handles[Factory->ValueArgIndex].GetInputData();
 		
-		bool& Success = *(bool*)Handles[Factory->SuccessArgIndex].GetData();
+		bool& Success = *(bool*)Handles[Factory->SuccessArgIndex].GetOutputData();
 		Success = false;
 		
-		FName& CachedBoneName = *(FName*)Handles[Factory->CachedBoneNameArgIndex].GetData(false, InContext.GetSlice().GetIndex());
-		int32& CachedBoneIndex = *(int32*)Handles[Factory->CachedBoneIndexArgIndex].GetData(false, InContext.GetSlice().GetIndex());
+		FName& CachedBoneName = *(FName*)Handles[Factory->CachedBoneNameArgIndex].GetPrivateData(InContext.GetSlice().GetIndex());
+		int32& CachedBoneIndex = *(int32*)Handles[Factory->CachedBoneIndexArgIndex].GetPrivateData(InContext.GetSlice().GetIndex());
 
 		// extract the animation attribute
 		const FControlRigExecuteContext& Context = InContext.GetPublicDataSafe<FControlRigExecuteContext>();
diff --git a/Engine/Plugins/Animation/ControlRig/Source/ControlRig/Public/Units/Hierarchy/RigUnit_Metadata.h b/Engine/Plugins/Animation/ControlRig/Source/ControlRig/Public/Units/Hierarchy/RigUnit_Metadata.h
--- a/Engine/Plugins/Animation/ControlRig/Source/ControlRig/Public/Units/Hierarchy/RigUnit_Metadata.h
+++ b/Engine/Plugins/Animation/ControlRig/Source/ControlRig/Public/Units/Hierarchy/RigUnit_Metadata.h
@@ -222,13 +222,13 @@ protected:
 #endif
 
 		// unpack the memory
-		const FRigElementKey& Item = *reinterpret_cast<const FRigElementKey*>(Handles[Factory->ItemArgIndex].GetData());
-		const FName& Name = *reinterpret_cast<const FName*>(Handles[Factory->NameArgIndex].GetData());
-		const ERigMetaDataNameSpace NameSpace = *reinterpret_cast<const ERigMetaDataNameSpace*>(Handles[Factory->NameSpaceArgIndex].GetData());
-		FCachedRigElement& Cache = *reinterpret_cast<FCachedRigElement*>(Handles[Factory->CacheArgIndex].GetData(false, InContext.GetSlice().GetIndex()));
-		const ValueType& Default = *reinterpret_cast<const ValueType*>(Handles[Factory->DefaultArgIndex].GetData());
-		ValueType& Value = *reinterpret_cast<ValueType*>(Handles[Factory->ValueArgIndex].GetData());
-		bool& Found = *reinterpret_cast<bool*>(Handles[Factory->FoundArgIndex].GetData());
+		const FRigElementKey& Item = *reinterpret_cast<const FRigElementKey*>(Handles[Factory->ItemArgIndex].GetInputData());
+		const FName& Name = *reinterpret_cast<const FName*>(Handles[Factory->NameArgIndex].GetInputData());
+		const ERigMetaDataNameSpace NameSpace = *reinterpret_cast<const ERigMetaDataNameSpace*>(Handles[Factory->NameSpaceArgIndex].GetInputData());
+		FCachedRigElement& Cache = *reinterpret_cast<FCachedRigElement*>(Handles[Factory->CacheArgIndex].GetPrivateData(InContext.GetSlice().GetIndex()));
+		const ValueType& Default = *reinterpret_cast<const ValueType*>(Handles[Factory->DefaultArgIndex].GetInputData());
+		ValueType& Value = *reinterpret_cast<ValueType*>(Handles[Factory->ValueArgIndex].GetOutputData());
+		bool& Found = *reinterpret_cast<bool*>(Handles[Factory->FoundArgIndex].GetOutputData());
 
 		// extract the metadata
 		if (const MetadataType* Md = Cast<MetadataType>(FindMetadata(InContext, Item, Name, EnumValue, NameSpace, Cache)))
@@ -288,12 +288,12 @@ protected:
 
 		// unpack the memory
 		FControlRigExecuteContext& ExecuteContext = InContext.GetPublicData<FControlRigExecuteContext>();
-		const FRigElementKey& Item = *(const FRigElementKey*)Handles[Factory->ItemArgIndex].GetData();
-		const FName& Name = *(const FName*)Handles[Factory->NameArgIndex].GetData();
-		const ERigMetaDataNameSpace NameSpace = *(const ERigMetaDataNameSpace*)Handles[Factory->NameSpaceArgIndex].GetData();
-		FCachedRigElement& Cache = *(FCachedRigElement*)Handles[Factory->CacheArgIndex].GetData(false, InContext.GetSlice().GetIndex());
-		ValueType& Value = *(ValueType*)Handles[Factory->ValueArgIndex].GetData();
-		bool& Success = *(bool*)Handles[Factory->SuccessArgIndex].GetData();
+		const FRigElementKey& Item = *(const FRigElementKey*)Handles[Factory->ItemArgIndex].GetInputData();
+		const FName& Name = *(const FName*)Handles[Factory->NameArgIndex].GetInputData();
+		const ERigMetaDataNameSpace NameSpace = *(const ERigMetaDataNameSpace*)Handles[Factory->NameSpaceArgIndex].GetInputData();
+		FCachedRigElement& Cache = *(FCachedRigElement*)Handles[Factory->CacheArgIndex].GetPrivateData(InContext.GetSlice().GetIndex());
+		ValueType& Value = *(ValueType*)Handles[Factory->ValueArgIndex].GetOutputData();
+		bool& Success = *(bool*)Handles[Factory->SuccessArgIndex].GetOutputData();
 
 		// extract the metadata
 		if (MetadataType* Md = Cast<MetadataType>(FindOrAddMetadata(ExecuteContext, Item, Name, EnumValue, NameSpace, Cache)))
@@ -909,11 +909,11 @@ protected:
 #endif
 
 		// unpack the memory
-		const FName& Name = *reinterpret_cast<const FName*>(Handles[Factory->NameArgIndex].GetData());
-		const ERigMetaDataNameSpace NameSpace = *reinterpret_cast<const ERigMetaDataNameSpace*>(Handles[Factory->NameSpaceArgIndex].GetData());
-		const ValueType& Default = *reinterpret_cast<const ValueType*>(Handles[Factory->DefaultArgIndex].GetData());
-		ValueType& Value = *reinterpret_cast<ValueType*>(Handles[Factory->ValueArgIndex].GetData());
-		bool& Found = *reinterpret_cast<bool*>(Handles[Factory->FoundArgIndex].GetData());
+		const FName& Name = *reinterpret_cast<const FName*>(Handles[Factory->NameArgIndex].GetInputData());
+		const ERigMetaDataNameSpace NameSpace = *reinterpret_cast<const ERigMetaDataNameSpace*>(Handles[Factory->NameSpaceArgIndex].GetInputData());
+		const ValueType& Default = *reinterpret_cast<const ValueType*>(Handles[Factory->DefaultArgIndex].GetInputData());
+		ValueType& Value = *reinterpret_cast<ValueType*>(Handles[Factory->ValueArgIndex].GetOutputData());
+		bool& Found = *reinterpret_cast<bool*>(Handles[Factory->FoundArgIndex].GetOutputData());
 
 		// extract the metadata
 		if (const MetadataType* Md = Cast<MetadataType>(FindMetadata(InContext, Name, EnumValue, NameSpace)))
@@ -972,10 +972,10 @@ protected:
 
 		// unpack the memory
 		FControlRigExecuteContext& ExecuteContext = InContext.GetPublicData<FControlRigExecuteContext>();
-		const FName& Name = *(const FName*)Handles[Factory->NameArgIndex].GetData();
-		const ERigMetaDataNameSpace NameSpace = *(const ERigMetaDataNameSpace*)Handles[Factory->NameSpaceArgIndex].GetData();
-		ValueType& Value = *(ValueType*)Handles[Factory->ValueArgIndex].GetData();
-		bool& Success = *(bool*)Handles[Factory->SuccessArgIndex].GetData();
+		const FName& Name = *(const FName*)Handles[Factory->NameArgIndex].GetInputData();
+		const ERigMetaDataNameSpace NameSpace = *(const ERigMetaDataNameSpace*)Handles[Factory->NameSpaceArgIndex].GetInputData();
+		ValueType& Value = *(ValueType*)Handles[Factory->ValueArgIndex].GetOutputData();
+		bool& Success = *(bool*)Handles[Factory->SuccessArgIndex].GetOutputData();
 
 		// extract the metadata
 		if (MetadataType* Md = Cast<MetadataType>(FindOrAddMetadata(ExecuteContext, Name, EnumValue, NameSpace)))
diff --git a/Engine/Plugins/Experimental/AnimNext/Source/AnimNext/Private/Variables/RigVMDispatch_CallHoistedAccessorFunction.cpp b/Engine/Plugins/Experimental/AnimNext/Source/AnimNext/Private/Variables/RigVMDispatch_CallHoistedAccessorFunction.cpp
--- a/Engine/Plugins/Experimental/AnimNext/Source/AnimNext/Private/Variables/RigVMDispatch_CallHoistedAccessorFunction.cpp
+++ b/Engine/Plugins/Experimental/AnimNext/Source/AnimNext/Private/Variables/RigVMDispatch_CallHoistedAccessorFunction.cpp
@@ -81,14 +81,14 @@ FRigVMDispatch_CallHoistedAccessorFunctionNative::FRigVMDispatch_CallHoistedAcce
 
 void FRigVMDispatch_CallHoistedAccessorFunctionNative::Execute(FRigVMExtendedExecuteContext& InContext, FRigVMMemoryHandleArray Handles, FRigVMPredicateBranchArray RigVMBranches)
 {
-	UObject* ObjectPtr = *reinterpret_cast<UObject**>(Handles[0].GetData());
+	UObject* ObjectPtr = *reinterpret_cast<UObject**>(Handles[0].GetOutputData());
 	if(ObjectPtr == nullptr)
 	{
 		// Something failed to resolve upstream, OK to just skip this work
 		return;
 	}
 
-	FAnimNextSoftFunctionPtr* SoftFunctionPtr = reinterpret_cast<FAnimNextSoftFunctionPtr*>(Handles[1].GetData());
+	FAnimNextSoftFunctionPtr* SoftFunctionPtr = reinterpret_cast<FAnimNextSoftFunctionPtr*>(Handles[1].GetOutputData());
 	UFunction* Function = SoftFunctionPtr->SoftObjectPtr.Get();
 	if(Function == nullptr || Function->NumParms != 2)
 	{
@@ -104,7 +104,7 @@ void FRigVMDispatch_CallHoistedAccessorFunctionNative::Execute(FRigVMExtendedExe
 	const FProperty* TargetProperty = Handles[2].GetResolvedProperty();
 	checkSlow(TargetProperty);
 	check(TargetProperty->GetClass() == ReturnValueProperty->GetClass());
-	uint8* TargetAddress = Handles[2].GetData();
+	uint8* TargetAddress = Handles[2].GetOutputData();
 	checkSlow(TargetAddress);
 
 	UObject* CDO = Function->GetOuterUClass()->GetDefaultObject();
@@ -119,14 +119,14 @@ FRigVMDispatch_CallHoistedAccessorFunctionScript::FRigVMDispatch_CallHoistedAcce
 
 void FRigVMDispatch_CallHoistedAccessorFunctionScript::Execute(FRigVMExtendedExecuteContext& InContext, FRigVMMemoryHandleArray Handles, FRigVMPredicateBranchArray RigVMBranches)
 {
-	UObject* ObjectPtr = *reinterpret_cast<UObject**>(Handles[0].GetData());
+	UObject* ObjectPtr = *reinterpret_cast<UObject**>(Handles[0].GetOutputData());
 	if(ObjectPtr == nullptr)
 	{
 		// Something failed to resolve upstream, OK to just skip this work
 		return;
 	}
 
-	FAnimNextSoftFunctionPtr* SoftFunctionPtr = reinterpret_cast<FAnimNextSoftFunctionPtr*>(Handles[1].GetData());
+	FAnimNextSoftFunctionPtr* SoftFunctionPtr = reinterpret_cast<FAnimNextSoftFunctionPtr*>(Handles[1].GetOutputData());
 	UFunction* Function = SoftFunctionPtr->SoftObjectPtr.Get();
 	if(Function == nullptr || Function->NumParms != 3)
 	{
@@ -156,7 +156,7 @@ void FRigVMDispatch_CallHoistedAccessorFunctionScript::Execute(FRigVMExtendedExe
 	const FProperty* TargetProperty = Handles[2].GetResolvedProperty();
 	checkSlow(TargetProperty);
 	check(TargetProperty->GetClass() == ReturnValueProperty->GetClass());
-	uint8* TargetAddress = Handles[2].GetData();
+	uint8* TargetAddress = Handles[2].GetOutputData();
 	checkSlow(TargetAddress);
 
 	// Allocate space for function params (script calls require multiple params like the input object and world context to be in one contiguous struct)
diff --git a/Engine/Plugins/Experimental/AnimNext/Source/AnimNext/Private/Variables/RigVMDispatch_CallObjectAccessorFunction.cpp b/Engine/Plugins/Experimental/AnimNext/Source/AnimNext/Private/Variables/RigVMDispatch_CallObjectAccessorFunction.cpp
--- a/Engine/Plugins/Experimental/AnimNext/Source/AnimNext/Private/Variables/RigVMDispatch_CallObjectAccessorFunction.cpp
+++ b/Engine/Plugins/Experimental/AnimNext/Source/AnimNext/Private/Variables/RigVMDispatch_CallObjectAccessorFunction.cpp
@@ -81,14 +81,14 @@ FRigVMDispatch_CallObjectAccessorFunctionNative::FRigVMDispatch_CallObjectAccess
 
 void FRigVMDispatch_CallObjectAccessorFunctionNative::Execute(FRigVMExtendedExecuteContext& InContext, FRigVMMemoryHandleArray Handles, FRigVMPredicateBranchArray RigVMBranches)
 {
-	UObject* ObjectPtr = *reinterpret_cast<UObject**>(Handles[0].GetData());
+	UObject* ObjectPtr = *reinterpret_cast<UObject**>(Handles[0].GetOutputData());
 	if(ObjectPtr == nullptr)
 	{
 		// Something failed to resolve upstream, OK to just skip this work
 		return;
 	}
 
-	FAnimNextSoftFunctionPtr* SoftFunctionPtr = reinterpret_cast<FAnimNextSoftFunctionPtr*>(Handles[1].GetData());
+	FAnimNextSoftFunctionPtr* SoftFunctionPtr = reinterpret_cast<FAnimNextSoftFunctionPtr*>(Handles[1].GetOutputData());
 	UFunction* Function = SoftFunctionPtr->SoftObjectPtr.Get();
 	if(Function == nullptr || Function->NumParms != 1)
 	{
@@ -104,7 +104,7 @@ void FRigVMDispatch_CallObjectAccessorFunctionNative::Execute(FRigVMExtendedExec
 	const FProperty* TargetProperty = Handles[2].GetResolvedProperty();
 	checkSlow(TargetProperty);
 	check(TargetProperty->GetClass() == ReturnValueProperty->GetClass());
-	uint8* TargetAddress = Handles[2].GetData();
+	uint8* TargetAddress = Handles[2].GetOutputData();
 	checkSlow(TargetAddress);
 
 	FFrame Stack(ObjectPtr, Function, nullptr, nullptr, Function->ChildProperties);
@@ -118,14 +118,14 @@ FRigVMDispatch_CallObjectAccessorFunctionScript::FRigVMDispatch_CallObjectAccess
 
 void FRigVMDispatch_CallObjectAccessorFunctionScript::Execute(FRigVMExtendedExecuteContext& InContext, FRigVMMemoryHandleArray Handles, FRigVMPredicateBranchArray RigVMBranches)
 {
-	UObject* ObjectPtr = *reinterpret_cast<UObject**>(Handles[0].GetData());
+	UObject* ObjectPtr = *reinterpret_cast<UObject**>(Handles[0].GetOutputData());
 	if(ObjectPtr == nullptr)
 	{
 		// Something failed to resolve upstream, OK to just skip this work
 		return;
 	}
 
-	FAnimNextSoftFunctionPtr* SoftFunctionPtr = reinterpret_cast<FAnimNextSoftFunctionPtr*>(Handles[1].GetData());
+	FAnimNextSoftFunctionPtr* SoftFunctionPtr = reinterpret_cast<FAnimNextSoftFunctionPtr*>(Handles[1].GetOutputData());
 	UFunction* Function = SoftFunctionPtr->SoftObjectPtr.Get();
 	if(Function == nullptr || Function->NumParms != 1)
 	{
@@ -141,7 +141,7 @@ void FRigVMDispatch_CallObjectAccessorFunctionScript::Execute(FRigVMExtendedExec
 	const FProperty* TargetProperty = Handles[2].GetResolvedProperty();
 	checkSlow(TargetProperty);
 	check(TargetProperty->GetClass() == ReturnValueProperty->GetClass());
-	uint8* TargetAddress = Handles[2].GetData();
+	uint8* TargetAddress = Handles[2].GetOutputData();
 	checkSlow(TargetAddress);
 
 	check(ObjectPtr->GetClass()->IsChildOf(Function->GetOuterUClass()));
diff --git a/Engine/Plugins/Experimental/AnimNext/Source/AnimNext/Private/Variables/RigVMDispatch_GetObjectProperty.cpp b/Engine/Plugins/Experimental/AnimNext/Source/AnimNext/Private/Variables/RigVMDispatch_GetObjectProperty.cpp
--- a/Engine/Plugins/Experimental/AnimNext/Source/AnimNext/Private/Variables/RigVMDispatch_GetObjectProperty.cpp
+++ b/Engine/Plugins/Experimental/AnimNext/Source/AnimNext/Private/Variables/RigVMDispatch_GetObjectProperty.cpp
@@ -76,14 +76,14 @@ FRigVMTemplateTypeMap FRigVMDispatch_GetObjectProperty::OnNewArgumentType(const
 
 void FRigVMDispatch_GetObjectProperty::Execute(FRigVMExtendedExecuteContext& InContext, FRigVMMemoryHandleArray Handles, FRigVMPredicateBranchArray RigVMBranches)
 {
-	const UObject* ObjectPtr = *reinterpret_cast<UObject**>(Handles[0].GetData());
+	const UObject* ObjectPtr = *reinterpret_cast<UObject**>(Handles[0].GetOutputData());
 	if(ObjectPtr == nullptr)
 	{
 		// Something failed to resolve upstream, OK to just skip this work
 		return;
 	}
 
-	FAnimNextFieldPath* FieldPathPtr = reinterpret_cast<FAnimNextFieldPath*>(Handles[1].GetData());
+	const FAnimNextFieldPath* FieldPathPtr = reinterpret_cast<const FAnimNextFieldPath*>(Handles[1].GetInputData());
 	const FProperty* SourceProperty = FieldPathPtr->FieldPath.Get();
 	if(SourceProperty == nullptr)
 	{
@@ -96,7 +96,7 @@ void FRigVMDispatch_GetObjectProperty::Execute(FRigVMExtendedExecuteContext& InC
 	const FProperty* TargetProperty = Handles[2].GetResolvedProperty();
 	checkSlow(TargetProperty);
 	checkSlow(TargetProperty->GetClass() == SourceProperty->GetClass());
-	uint8* TargetAddress = Handles[2].GetData();
+	uint8* TargetAddress = Handles[2].GetOutputData();
 	checkSlow(TargetAddress);
 
 	// TODO: add a specialization for bool to GetDispatchFunctionImpl rather than branching here
diff --git a/Engine/Plugins/Experimental/AnimNext/Source/AnimNext/Private/Variables/RigVMTrait_OverrideVariables.cpp b/Engine/Plugins/Experimental/AnimNext/Source/AnimNext/Private/Variables/RigVMTrait_OverrideVariables.cpp
--- a/Engine/Plugins/Experimental/AnimNext/Source/AnimNext/Private/Variables/RigVMTrait_OverrideVariables.cpp
+++ b/Engine/Plugins/Experimental/AnimNext/Source/AnimNext/Private/Variables/RigVMTrait_OverrideVariables.cpp
@@ -76,7 +76,7 @@ void FRigVMTrait_OverrideVariablesForAsset::GenerateOverrides(const FRigVMTraitS
 			{
 				check(OverrideInfo.Type == FAnimNextParamType::FromProperty(Desc->CachedProperty));
 				check(OverrideInfo.Type == FAnimNextParamType::FromProperty(MemoryHandle.GetProperty()));
-				AssetOverrides.Emplace(Desc->Name, OverrideInfo.Type, const_cast<uint8*>(MemoryHandle.GetData()));
+				AssetOverrides.Emplace(Desc->Name, OverrideInfo.Type, const_cast<uint8*>(MemoryHandle.GetInputData()));
 			}
 		}
 
@@ -113,7 +113,7 @@ void FRigVMTrait_OverrideVariablesForStruct::GenerateOverrides(const FRigVMTrait
 			{
 				check(OverrideInfo.Type == FAnimNextParamType::FromProperty(Property));
 				check(OverrideInfo.Type == FAnimNextParamType::FromProperty(MemoryHandle.GetProperty()));
-				StructOverrides.Emplace(Property->GetFName(), OverrideInfo.Type, const_cast<uint8*>(MemoryHandle.GetData()));
+				StructOverrides.Emplace(Property->GetFName(), OverrideInfo.Type, const_cast<uint8*>(MemoryHandle.GetInputData()));
 			}
 		}
 
diff --git a/Engine/Plugins/Experimental/AnimNextAnimGraph/Source/AnimNextAnimGraph/Private/Graph/RigUnit_AnimNextGraphEvaluator.cpp b/Engine/Plugins/Experimental/AnimNextAnimGraph/Source/AnimNextAnimGraph/Private/Graph/RigUnit_AnimNextGraphEvaluator.cpp
--- a/Engine/Plugins/Experimental/AnimNextAnimGraph/Source/AnimNextAnimGraph/Private/Graph/RigUnit_AnimNextGraphEvaluator.cpp
+++ b/Engine/Plugins/Experimental/AnimNextAnimGraph/Source/AnimNextAnimGraph/Private/Graph/RigUnit_AnimNextGraphEvaluator.cpp
@@ -64,7 +64,7 @@ void FRigUnit_AnimNextGraphEvaluator::StaticExecute(FRigVMExtendedExecuteContext
 			MemoryHandle.ComputeLazyValueIfNecessary(RigVMExecuteContext, RigVMExecuteContext.GetSliceHash());
 		}
 
-		const uint8* SourcePtr = MemoryHandle.GetData();
+		const uint8* SourcePtr = MemoryHandle.GetInputData();
 		uint8* DestinationPtr = DestinationBasePtr + Handle.GetLatentPropertyOffset();
 
 		// Copy from our source into our destination
diff --git a/Engine/Plugins/Experimental/AnimNextChooser/Source/AnimNextChooser/Private/RigUnit_EvaluateChooser.cpp b/Engine/Plugins/Experimental/AnimNextChooser/Source/AnimNextChooser/Private/RigUnit_EvaluateChooser.cpp
--- a/Engine/Plugins/Experimental/AnimNextChooser/Source/AnimNextChooser/Private/RigUnit_EvaluateChooser.cpp
+++ b/Engine/Plugins/Experimental/AnimNextChooser/Source/AnimNextChooser/Private/RigUnit_EvaluateChooser.cpp
@@ -93,8 +93,8 @@ void FRigVMDispatch_EvaluateChooser::Execute(FRigVMExtendedExecuteContext& InCon
 
 	const FProperty* ResultProperty = Handles[2].GetResolvedProperty();
 	
-	TObjectPtr<UChooserTable>* Chooser = reinterpret_cast<TObjectPtr<UChooserTable>*>(Handles[0].GetData());
-	TObjectPtr<UObject>* ContextObject = reinterpret_cast<TObjectPtr<UObject>*>(Handles[1].GetData());
+	TObjectPtr<UChooserTable>* Chooser = reinterpret_cast<TObjectPtr<UChooserTable>*>(Handles[0].GetOutputData());
+	TObjectPtr<UObject>* ContextObject = reinterpret_cast<TObjectPtr<UObject>*>(Handles[1].GetOutputData());
 
 
 	FChooserEvaluationContext ChooserContext;
@@ -116,7 +116,7 @@ void FRigVMDispatch_EvaluateChooser::Execute(FRigVMExtendedExecuteContext& InCon
 		if (const FObjectProperty* ObjectResultProperty = CastField<FObjectProperty>(ArrayResultProperty->Inner))
 		{
 			UClass* ObjectClass = ObjectResultProperty->PropertyClass;
-			FScriptArrayHelper ArrayHelper(ArrayResultProperty, Handles[2].GetData());
+			FScriptArrayHelper ArrayHelper(ArrayResultProperty, Handles[2].GetOutputData());
 			ArrayHelper.Resize(0);
 			
 			UChooserTable::EvaluateChooser(ChooserContext, *Chooser, FObjectChooserBase::FObjectChooserIteratorCallback::CreateLambda([ObjectClass, &ArrayHelper](UObject* InResult)
@@ -138,7 +138,7 @@ void FRigVMDispatch_EvaluateChooser::Execute(FRigVMExtendedExecuteContext& InCon
 	else if (const FObjectProperty* ObjectResultProperty = CastField<FObjectProperty>(ResultProperty))
 	{
 		UClass* ObjectClass = ObjectResultProperty->PropertyClass;
-		TObjectPtr<UObject>* Result = (TObjectPtr<UObject>*)Handles[2].GetData();
+		TObjectPtr<UObject>* Result = (TObjectPtr<UObject>*)Handles[2].GetOutputData();
 		UChooserTable::EvaluateChooser(ChooserContext, *Chooser, FObjectChooserBase::FObjectChooserIteratorCallback::CreateLambda([ObjectClass, &Result](UObject* InResult)
 		{
 			if (InResult)
diff --git a/Engine/Plugins/Runtime/RigVM/Source/RigVM/Private/RigVMCore/RigVM.cpp b/Engine/Plugins/Runtime/RigVM/Source/RigVM/Private/RigVMCore/RigVM.cpp
--- a/Engine/Plugins/Runtime/RigVM/Source/RigVM/Private/RigVMCore/RigVM.cpp
+++ b/Engine/Plugins/Runtime/RigVM/Source/RigVM/Private/RigVMCore/RigVM.cpp
@@ -1999,11 +1999,11 @@ ERigVMExecuteResult URigVM::ExecuteInstructions(FRigVMExtendedExecuteContext& Co
 				const FRigVMMemoryHandle& Handle = Context.CachedMemoryHandles[FirstHandleForInstruction[ContextPublicData.InstructionIndex]];
 				if(Handle.GetProperty()->IsA<FIntProperty>())
 				{
-					*((int32*)Context.CachedMemoryHandles[FirstHandleForInstruction[ContextPublicData.InstructionIndex]].GetData()) = 0;
+					*((int32*)Context.CachedMemoryHandles[FirstHandleForInstruction[ContextPublicData.InstructionIndex]].GetInputData()) = 0;
 				}
 				else if(Handle.GetProperty()->IsA<FNameProperty>())
 				{
-					*((FName*)Context.CachedMemoryHandles[FirstHandleForInstruction[ContextPublicData.InstructionIndex]].GetData()) = NAME_None;
+					*((FName*)Context.CachedMemoryHandles[FirstHandleForInstruction[ContextPublicData.InstructionIndex]].GetInputData()) = NAME_None;
 				}
 #if WITH_EDITOR
 				if(GetDebugMemory(Context) && GetDebugMemory(Context)->Num() > 0)
@@ -2018,13 +2018,13 @@ ERigVMExecuteResult URigVM::ExecuteInstructions(FRigVMExtendedExecuteContext& Co
 			}
 			case ERigVMOpCode::BoolFalse:
 			{
-				*((bool*)Context.CachedMemoryHandles[FirstHandleForInstruction[ContextPublicData.InstructionIndex]].GetData()) = false;
+				*((bool*)Context.CachedMemoryHandles[FirstHandleForInstruction[ContextPublicData.InstructionIndex]].GetInputData()) = false;
 				ContextPublicData.InstructionIndex++;
 				break;
 			}
 			case ERigVMOpCode::BoolTrue:
 			{
-				*((bool*)Context.CachedMemoryHandles[FirstHandleForInstruction[ContextPublicData.InstructionIndex]].GetData()) = true;
+				*((bool*)Context.CachedMemoryHandles[FirstHandleForInstruction[ContextPublicData.InstructionIndex]].GetInputData()) = true;
 				ContextPublicData.InstructionIndex++;
 				break;
 			}
@@ -2048,7 +2048,7 @@ ERigVMExecuteResult URigVM::ExecuteInstructions(FRigVMExtendedExecuteContext& Co
 			}
 			case ERigVMOpCode::Increment:
 			{
-				(*((int32*)Context.CachedMemoryHandles[FirstHandleForInstruction[ContextPublicData.InstructionIndex]].GetData()))++;
+				(*((int32*)Context.CachedMemoryHandles[FirstHandleForInstruction[ContextPublicData.InstructionIndex]].GetInputData()))++;
 #if WITH_EDITOR
 				if(GetDebugMemory(Context) && GetDebugMemory(Context)->Num() > 0)
 				{
@@ -2061,7 +2061,7 @@ ERigVMExecuteResult URigVM::ExecuteInstructions(FRigVMExtendedExecuteContext& Co
 			}
 			case ERigVMOpCode::Decrement:
 			{
-				(*((int32*)Context.CachedMemoryHandles[FirstHandleForInstruction[ContextPublicData.InstructionIndex]].GetData()))--;
+				(*((int32*)Context.CachedMemoryHandles[FirstHandleForInstruction[ContextPublicData.InstructionIndex]].GetInputData()))--;
 #if WITH_EDITOR
 				if(GetDebugMemory(Context) && GetDebugMemory(Context)->Num() > 0)
 				{
@@ -2079,9 +2079,9 @@ ERigVMExecuteResult URigVM::ExecuteInstructions(FRigVMExtendedExecuteContext& Co
 
 				FRigVMMemoryHandle& HandleA = Context.CachedMemoryHandles[FirstHandleForInstruction[ContextPublicData.InstructionIndex]];
 				FRigVMMemoryHandle& HandleB = Context.CachedMemoryHandles[FirstHandleForInstruction[ContextPublicData.InstructionIndex] + 1];
-				const bool Result = HandleA.GetProperty()->Identical(HandleA.GetData(true), HandleB.GetData(true));
+				const bool Result = HandleA.GetProperty()->Identical(HandleA.GetInputData(), HandleB.GetInputData());
 
-				*((bool*)Context.CachedMemoryHandles[FirstHandleForInstruction[ContextPublicData.InstructionIndex]+2].GetData()) = Result;
+				*((bool*)Context.CachedMemoryHandles[FirstHandleForInstruction[ContextPublicData.InstructionIndex]+2].GetOutputData()) = Result;
 				ContextPublicData.InstructionIndex++;
 				break;
 			}
@@ -2106,7 +2106,7 @@ ERigVMExecuteResult URigVM::ExecuteInstructions(FRigVMExtendedExecuteContext& Co
 			case ERigVMOpCode::JumpAbsoluteIf:
 			{
 				const FRigVMJumpIfOp& Op = ByteCode.GetOpAt<FRigVMJumpIfOp>(Instruction);
-				const bool Condition = *(bool*)Context.CachedMemoryHandles[FirstHandleForInstruction[ContextPublicData.InstructionIndex]].GetData();
+				const bool Condition = *(bool*)Context.CachedMemoryHandles[FirstHandleForInstruction[ContextPublicData.InstructionIndex]].GetInputData();
 				if (Condition == Op.Condition)
 				{
 					ContextPublicData.InstructionIndex = Op.InstructionIndex;
@@ -2120,7 +2120,7 @@ ERigVMExecuteResult URigVM::ExecuteInstructions(FRigVMExtendedExecuteContext& Co
 			case ERigVMOpCode::JumpForwardIf:
 			{
 				const FRigVMJumpIfOp& Op = ByteCode.GetOpAt<FRigVMJumpIfOp>(Instruction);
-				const bool Condition = *(bool*)Context.CachedMemoryHandles[FirstHandleForInstruction[ContextPublicData.InstructionIndex]].GetData();
+				const bool Condition = *(bool*)Context.CachedMemoryHandles[FirstHandleForInstruction[ContextPublicData.InstructionIndex]].GetInputData();
 				if (Condition == Op.Condition)
 				{
 					ContextPublicData.InstructionIndex += Op.InstructionIndex;
@@ -2134,7 +2134,7 @@ ERigVMExecuteResult URigVM::ExecuteInstructions(FRigVMExtendedExecuteContext& Co
 			case ERigVMOpCode::JumpBackwardIf:
 			{
 				const FRigVMJumpIfOp& Op = ByteCode.GetOpAt<FRigVMJumpIfOp>(Instruction);
-				const bool Condition = *(bool*)Context.CachedMemoryHandles[FirstHandleForInstruction[ContextPublicData.InstructionIndex]].GetData();
+				const bool Condition = *(bool*)Context.CachedMemoryHandles[FirstHandleForInstruction[ContextPublicData.InstructionIndex]].GetInputData();
 				if (Condition == Op.Condition)
 				{
 					ContextPublicData.InstructionIndex -= Op.InstructionIndex;
@@ -2183,8 +2183,8 @@ ERigVMExecuteResult URigVM::ExecuteInstructions(FRigVMExtendedExecuteContext& Co
 			}
 			case ERigVMOpCode::BeginBlock:
 			{
-				const int32 Count = (*((int32*)Context.CachedMemoryHandles[FirstHandleForInstruction[ContextPublicData.InstructionIndex]].GetData()));
-				const int32 Index = (*((int32*)Context.CachedMemoryHandles[FirstHandleForInstruction[ContextPublicData.InstructionIndex] + 1].GetData()));
+				const int32 Count = (*((int32*)Context.CachedMemoryHandles[FirstHandleForInstruction[ContextPublicData.InstructionIndex]].GetInputData()));
+				const int32 Index = (*((int32*)Context.CachedMemoryHandles[FirstHandleForInstruction[ContextPublicData.InstructionIndex] + 1].GetInputData()));
 				Context.BeginSlice(Count, Index);
 				ContextPublicData.InstructionIndex++;
 				break;
@@ -2229,7 +2229,7 @@ ERigVMExecuteResult URigVM::ExecuteInstructions(FRigVMExtendedExecuteContext& Co
 			{
 				const FRigVMJumpToBranchOp& Op = ByteCode.GetOpAt<FRigVMJumpToBranchOp>(Instruction);
 				// BranchLabel = Op.Arg
-				FName BranchLabel = *(FName*)Context.CachedMemoryHandles[FirstHandleForInstruction[ContextPublicData.InstructionIndex]].GetData();
+				FName BranchLabel = *(FName*)Context.CachedMemoryHandles[FirstHandleForInstruction[ContextPublicData.InstructionIndex]].GetInputData();
 
 				// iterate over the branches stored in the bytecode,
 				// starting at the first branch index stored in the operator.
@@ -2281,7 +2281,7 @@ ERigVMExecuteResult URigVM::ExecuteInstructions(FRigVMExtendedExecuteContext& Co
 			case ERigVMOpCode::RunInstructions:
 			{
 				const FRigVMRunInstructionsOp& Op = ByteCode.GetOpAt<FRigVMRunInstructionsOp>(Instruction);
-				FRigVMInstructionSetExecuteState& ExecutionState = *(FRigVMInstructionSetExecuteState*)Context.CachedMemoryHandles[FirstHandleForInstruction[ContextPublicData.InstructionIndex]].GetData();
+				FRigVMInstructionSetExecuteState& ExecutionState = *(FRigVMInstructionSetExecuteState*)Context.CachedMemoryHandles[FirstHandleForInstruction[ContextPublicData.InstructionIndex]].GetPrivateData(INDEX_NONE);
 
 				if((Op.StartInstruction != INDEX_NONE) &&
 					(Op.EndInstruction != INDEX_NONE) &&
@@ -2306,7 +2306,7 @@ ERigVMExecuteResult URigVM::ExecuteInstructions(FRigVMExtendedExecuteContext& Co
 				ContextPublicData.Traits.Reset();
 				ContextPublicData.AdditionalTraitMemoryHandles.Reset();
 
-				const TArray<int32>& TraitList = *(TArray<int32>*)Context.CachedMemoryHandles[FirstHandleForInstruction[ContextPublicData.InstructionIndex]].GetData();
+				const TArray<int32>& TraitList = *(TArray<int32>*)Context.CachedMemoryHandles[FirstHandleForInstruction[ContextPublicData.InstructionIndex]].GetInputData();
 				ContextPublicData.Traits.Reserve(TraitList.Num());
 				ContextPublicData.AdditionalTraitMemoryHandles.Reserve(TraitList.Num());
 				int32 AdditionalStartIndex = INDEX_NONE;
@@ -3032,7 +3032,7 @@ void URigVM::CacheSingleMemoryHandle(FRigVMExtendedExecuteContext& Context, cons
 		Context.CachedMemoryHandles[InHandleIndex] = {ExternalVariableRuntimeData.Memory, ExternalVariable.Property, PropertyPath};
 #if UE_BUILD_DEBUG
 		// make sure the handle points to valid memory
-		check(Context.CachedMemoryHandles[InHandleIndex].GetData(false) != nullptr);
+		check(Context.CachedMemoryHandles[InHandleIndex].GetOutputData() != nullptr);
 #endif
 		return;
 	}
@@ -3066,10 +3066,10 @@ void URigVM::CacheSingleMemoryHandle(FRigVMExtendedExecuteContext& Context, cons
 
 #if UE_BUILD_DEBUG
 	// make sure the handle points to valid memory
-	check(Handle.GetData(false) != nullptr);
+	check(Handle.GetOutputData() != nullptr);
 	if(PropertyPath)
 	{
-		uint8* MemoryPtr = Handle.GetData(false);
+		uint8* MemoryPtr = Handle.GetOutputData();
 		// don't check the result - since it may be an array element
 		// that doesn't exist yet. 
 		PropertyPath->GetData<uint8>(MemoryPtr, Property);
diff --git a/Engine/Plugins/Runtime/RigVM/Source/RigVM/Private/RigVMCore/RigVMMemoryStorage.cpp b/Engine/Plugins/Runtime/RigVM/Source/RigVM/Private/RigVMCore/RigVMMemoryStorage.cpp
--- a/Engine/Plugins/Runtime/RigVM/Source/RigVM/Private/RigVMCore/RigVMMemoryStorage.cpp
+++ b/Engine/Plugins/Runtime/RigVM/Source/RigVM/Private/RigVMCore/RigVMMemoryStorage.cpp
@@ -124,7 +124,7 @@ const uint8* TRigVMLazyValueBase::GetData(const FRigVMExecuteContext& Context) c
 			check(Context.ExtendedExecuteContext != nullptr);
 			MemoryHandle->ComputeLazyValueIfNecessary(*Context.ExtendedExecuteContext, SliceHash);
 		}
-		return MemoryHandle->GetData(bFollowPropertyPath, INDEX_NONE);
+		return MemoryHandle->GetData_Internal(bFollowPropertyPath, INDEX_NONE);
 	}
 	return nullptr;
 }
@@ -1209,10 +1209,10 @@ bool URigVMMemoryStorage::CopyProperty(
 {
 	return CopyProperty(
 		InTargetHandle.GetProperty(),
-		InTargetHandle.GetData(false),
+		InTargetHandle.GetData_Internal(false, INDEX_NONE),
 		InTargetHandle.GetPropertyPathRef(),
 		InSourceHandle.GetProperty(),
-		InSourceHandle.GetData(false),
+		InSourceHandle.GetData_Internal(false, INDEX_NONE),
 		InSourceHandle.GetPropertyPathRef());
 }
 
diff --git a/Engine/Plugins/Runtime/RigVM/Source/RigVM/Private/RigVMCore/RigVMMemoryStorageStruct.cpp b/Engine/Plugins/Runtime/RigVM/Source/RigVM/Private/RigVMCore/RigVMMemoryStorageStruct.cpp
--- a/Engine/Plugins/Runtime/RigVM/Source/RigVM/Private/RigVMCore/RigVMMemoryStorageStruct.cpp
+++ b/Engine/Plugins/Runtime/RigVM/Source/RigVM/Private/RigVMCore/RigVMMemoryStorageStruct.cpp
@@ -484,10 +484,10 @@ bool FRigVMMemoryStorageStruct::CopyProperty(
 {
 	return CopyProperty(
 		InTargetHandle.GetProperty(),
-		InTargetHandle.GetData(false),
+		InTargetHandle.GetData_Internal(false, INDEX_NONE),
 		InTargetHandle.GetPropertyPathRef(),
 		InSourceHandle.GetProperty(),
-		InSourceHandle.GetData(false),
+		InSourceHandle.GetData_Internal(false, INDEX_NONE),
 		InSourceHandle.GetPropertyPathRef());
 }
 
diff --git a/Engine/Plugins/Runtime/RigVM/Source/RigVM/Private/RigVMFunctions/RigVMDispatch_Array.cpp b/Engine/Plugins/Runtime/RigVM/Source/RigVM/Private/RigVMFunctions/RigVMDispatch_Array.cpp
--- a/Engine/Plugins/Runtime/RigVM/Source/RigVM/Private/RigVMFunctions/RigVMDispatch_Array.cpp
+++ b/Engine/Plugins/Runtime/RigVM/Source/RigVM/Private/RigVMFunctions/RigVMDispatch_Array.cpp
@@ -275,7 +275,7 @@ FName FRigVMDispatch_ArrayMake::GetArgumentNameForOperandIndex(int32 InOperandIn
 void FRigVMDispatch_ArrayMake::Execute(FRigVMExtendedExecuteContext& InContext, FRigVMMemoryHandleArray Handles, FRigVMPredicateBranchArray Predicates)
 {
 	const FArrayProperty* ArrayProperty = CastFieldChecked<FArrayProperty>(Handles.Last().GetProperty());
-	FScriptArrayHelper ArrayHelper(ArrayProperty, Handles.Last().GetData());
+	FScriptArrayHelper ArrayHelper(ArrayProperty, Handles.Last().GetOutputData());
 
 	const int32 Num = Handles.Num() - 1;
 	if(InContext.IsValidArraySize(Num))
@@ -284,7 +284,7 @@ void FRigVMDispatch_ArrayMake::Execute(FRigVMExtendedExecuteContext& InContext,
 
 		for(int32 Index = 0; Index < Num; Index++)
 		{
-			const uint8* SourceMemory = Handles[Index].GetData();
+			const uint8* SourceMemory = Handles[Index].GetInputData();
 			uint8* TargetMemory = ArrayHelper.GetRawPtr(Index);
 			URigVMMemoryStorage::CopyProperty(ArrayProperty->Inner, TargetMemory, Handles[Index].GetProperty(), SourceMemory);
 		}
@@ -323,7 +323,7 @@ FRigVMTemplateTypeMap FRigVMDispatch_ArrayReset::OnNewArgumentType(const FName&
 void FRigVMDispatch_ArrayReset::Execute(FRigVMExtendedExecuteContext& InContext, FRigVMMemoryHandleArray Handles, FRigVMPredicateBranchArray Predicates)
 {
 	const FArrayProperty* ArrayProperty = CastFieldChecked<FArrayProperty>(Handles[0].GetProperty());
-	FScriptArrayHelper ArrayHelper(ArrayProperty, Handles[0].GetData());
+	FScriptArrayHelper ArrayHelper(ArrayProperty, Handles[0].GetOutputData());
 	ArrayHelper.Resize(0);
 }
 
@@ -379,8 +379,8 @@ FRigVMTemplateTypeMap FRigVMDispatch_ArrayGetNum::OnNewArgumentType(const FName&
 void FRigVMDispatch_ArrayGetNum::Execute(FRigVMExtendedExecuteContext& InContext, FRigVMMemoryHandleArray Handles, FRigVMPredicateBranchArray Predicates)
 {
 	const FArrayProperty* ArrayProperty = CastFieldChecked<FArrayProperty>(Handles[0].GetProperty());
-	const FScriptArrayHelper ArrayHelper(ArrayProperty, Handles[0].GetData());
-	int32* Num = (int32*)Handles[1].GetData();
+	const FScriptArrayHelper ArrayHelper(ArrayProperty, Handles[0].GetInputData());
+	int32* Num = (int32*)Handles[1].GetOutputData();
 	*Num = ArrayHelper.Num();
 }
 
@@ -440,8 +440,8 @@ FRigVMTemplateTypeMap FRigVMDispatch_ArraySetNum::OnNewArgumentType(const FName&
 void FRigVMDispatch_ArraySetNum::Execute(FRigVMExtendedExecuteContext& InContext, FRigVMMemoryHandleArray Handles, FRigVMPredicateBranchArray Predicates)
 {
 	const FArrayProperty* ArrayProperty = CastFieldChecked<FArrayProperty>(Handles[0].GetProperty());
-	FScriptArrayHelper ArrayHelper(ArrayProperty, Handles[0].GetData());
-	const int32* Num = (int32*)Handles[1].GetData();
+	FScriptArrayHelper ArrayHelper(ArrayProperty, Handles[0].GetOutputData());
+	const int32* Num = (const int32*)Handles[1].GetInputData();
 	if(InContext.IsValidArraySize(*Num))
 	{
 		ArrayHelper.Resize(*Num);
@@ -507,10 +507,10 @@ FRigVMTemplateTypeMap FRigVMDispatch_ArrayInit::OnNewArgumentType(const FName& I
 void FRigVMDispatch_ArrayInit::Execute(FRigVMExtendedExecuteContext& InContext, FRigVMMemoryHandleArray Handles, FRigVMPredicateBranchArray Predicates)
 {
 	const FArrayProperty* ArrayProperty = CastFieldChecked<FArrayProperty>(Handles[0].GetProperty());
-	FScriptArrayHelper ArrayHelper(ArrayProperty, Handles[0].GetData());
-	const int32* Num = (int32*)Handles[1].GetData();
+	FScriptArrayHelper ArrayHelper(ArrayProperty, Handles[0].GetOutputData());
+	const int32* Num = (const int32*)Handles[1].GetInputData();
 	const FProperty* ValueProperty = Handles[2].GetProperty();
-	const uint8* ValueData = Handles[2].GetData();
+	const uint8* ValueData = Handles[2].GetInputData();
 
 	if(InContext.IsValidArraySize(*Num))
 	{
@@ -595,13 +595,14 @@ FRigVMTemplateTypeMap FRigVMDispatch_ArrayGetAtIndex::OnNewArgumentType(const FN
 void FRigVMDispatch_ArrayGetAtIndex::Execute(FRigVMExtendedExecuteContext& InContext, FRigVMMemoryHandleArray Handles, FRigVMPredicateBranchArray Predicates)
 {
 	const FArrayProperty* ArrayProperty = CastFieldChecked<FArrayProperty>(Handles[0].GetProperty());
-	FScriptArrayHelper ArrayHelper(ArrayProperty, Handles[0].GetData());
-	int32 Index = *(int32*)Handles[1].GetData();
-	uint8* TargetMemory = Handles[2].GetData();
+	FScriptArrayHelper ArrayHelper(ArrayProperty, Handles[0].GetInputData());
+	const int32 Index = *(const int32*)Handles[1].GetInputData();
+	int32 OutIndex = Index;
+	uint8* TargetMemory = Handles[2].GetOutputData();
 
-	if(InContext.IsValidArrayIndex(Index, ArrayHelper))
+	if(InContext.IsValidArrayIndex(OutIndex, ArrayHelper))
 	{
-		const uint8* SourceMemory = ArrayHelper.GetRawPtr(Index);
+		const uint8* SourceMemory = ArrayHelper.GetRawPtr(OutIndex);
 		URigVMMemoryStorage::CopyProperty(Handles[2].GetProperty(), TargetMemory, ArrayProperty->Inner, SourceMemory);
 	}
 }
@@ -679,13 +680,14 @@ FRigVMTemplateTypeMap FRigVMDispatch_ArraySetAtIndex::OnNewArgumentType(const FN
 void FRigVMDispatch_ArraySetAtIndex::Execute(FRigVMExtendedExecuteContext& InContext, FRigVMMemoryHandleArray Handles, FRigVMPredicateBranchArray Predicates)
 {
 	const FArrayProperty* ArrayProperty = CastFieldChecked<FArrayProperty>(Handles[0].GetProperty());
-	FScriptArrayHelper ArrayHelper(ArrayProperty, Handles[0].GetData());
-	int32 Index = *(int32*)Handles[1].GetData();
-	const uint8* SourceMemory = Handles[2].GetData();
+	FScriptArrayHelper ArrayHelper(ArrayProperty, Handles[0].GetOutputData());
+	const int32 Index = *(const int32*)Handles[1].GetInputData();
+	int32 OutIndex = Index;
+	const uint8* SourceMemory = Handles[2].GetInputData();
 
-	if(InContext.IsValidArrayIndex(Index, ArrayHelper))
+	if(InContext.IsValidArrayIndex(OutIndex, ArrayHelper))
 	{
-		uint8* TargetMemory = ArrayHelper.GetRawPtr(Index);
+		uint8* TargetMemory = ArrayHelper.GetRawPtr(OutIndex);
 		URigVMMemoryStorage::CopyProperty(ArrayProperty->Inner, TargetMemory, Handles[2].GetProperty(), SourceMemory);
 	}
 }
@@ -743,9 +745,9 @@ const TArray<FRigVMTemplateArgumentInfo>& FRigVMDispatch_ArrayAdd::GetArgumentIn
 void FRigVMDispatch_ArrayAdd::Execute(FRigVMExtendedExecuteContext& InContext, FRigVMMemoryHandleArray Handles, FRigVMPredicateBranchArray Predicates)
 {
 	const FArrayProperty* ArrayProperty = CastFieldChecked<FArrayProperty>(Handles[0].GetProperty());
-	FScriptArrayHelper ArrayHelper(ArrayProperty, Handles[0].GetData());
-	const uint8* SourceMemory = Handles[1].GetData();
-	int32& Index = *(int32*)Handles[2].GetData();
+	FScriptArrayHelper ArrayHelper(ArrayProperty, Handles[0].GetOutputData());
+	const uint8* SourceMemory = Handles[1].GetInputData();
+	int32& Index = *(int32*)Handles[2].GetOutputData();
 
 	if(InContext.IsValidArraySize(ArrayHelper.Num() + 1))
 	{
@@ -790,9 +792,9 @@ FText FRigVMDispatch_ArrayAdd::GetArgumentTooltip(const FName& InArgumentName, T
 void FRigVMDispatch_ArrayInsert::Execute(FRigVMExtendedExecuteContext& InContext, FRigVMMemoryHandleArray Handles, FRigVMPredicateBranchArray Predicates)
 {
 	const FArrayProperty* ArrayProperty = CastFieldChecked<FArrayProperty>(Handles[0].GetProperty());
-	FScriptArrayHelper ArrayHelper(ArrayProperty, Handles[0].GetData());
-	int32 Index = *(int32*)Handles[1].GetData();
-	const uint8* SourceMemory = Handles[2].GetData();
+	FScriptArrayHelper ArrayHelper(ArrayProperty, Handles[0].GetOutputData());
+	int32 Index = *(int32*)Handles[1].GetOutputData();
+	const uint8* SourceMemory = Handles[2].GetInputData();
 
 	if(InContext.IsValidArraySize(ArrayHelper.Num() + 1))
 	{
@@ -871,12 +873,13 @@ FRigVMTemplateTypeMap FRigVMDispatch_ArrayRemove::OnNewArgumentType(const FName&
 void FRigVMDispatch_ArrayRemove::Execute(FRigVMExtendedExecuteContext& InContext, FRigVMMemoryHandleArray Handles, FRigVMPredicateBranchArray Predicates)
 {
 	const FArrayProperty* ArrayProperty = CastFieldChecked<FArrayProperty>(Handles[0].GetProperty());
-	FScriptArrayHelper ArrayHelper(ArrayProperty, Handles[0].GetData());
-	int32 Index = *(int32*)Handles[1].GetData();
+	FScriptArrayHelper ArrayHelper(ArrayProperty, Handles[0].GetOutputData());
+	const int32 Index = *(const int32*)Handles[1].GetInputData();
+	int32 OutIndex = Index;
 
-	if(InContext.IsValidArrayIndex(Index, ArrayHelper))
+	if(InContext.IsValidArrayIndex(OutIndex, ArrayHelper))
 	{
-		ArrayHelper.RemoveValues(Index, 1);
+		ArrayHelper.RemoveValues(OutIndex, 1);
 	}
 }
 
@@ -907,7 +910,7 @@ FText FRigVMDispatch_ArrayRemove::GetArgumentTooltip(const FName& InArgumentName
 void FRigVMDispatch_ArrayReverse::Execute(FRigVMExtendedExecuteContext& InContext, FRigVMMemoryHandleArray Handles, FRigVMPredicateBranchArray Predicates)
 {
 	const FArrayProperty* ArrayProperty = CastFieldChecked<FArrayProperty>(Handles[0].GetProperty());
-	FScriptArrayHelper ArrayHelper(ArrayProperty, Handles[0].GetData());
+	FScriptArrayHelper ArrayHelper(ArrayProperty, Handles[0].GetOutputData());
 	for(int32 A=0, B=ArrayHelper.Num()-1; A<B; A++, B--)
 	{
 		ArrayHelper.SwapValues(A, B);
@@ -983,9 +986,9 @@ FRigVMTemplateTypeMap FRigVMDispatch_ArrayFind::OnNewArgumentType(const FName& I
 void FRigVMDispatch_ArrayFind::Execute(FRigVMExtendedExecuteContext& InContext, FRigVMMemoryHandleArray Handles, FRigVMPredicateBranchArray Predicates)
 {
 	const FArrayProperty* ArrayProperty = CastFieldChecked<FArrayProperty>(Handles[0].GetProperty());
-	FScriptArrayHelper ArrayHelper(ArrayProperty, Handles[0].GetData());
-	int32& Index = *(int32*)Handles[2].GetData();
-	bool& bSuccess = *(bool*)Handles[3].GetData();
+	FScriptArrayHelper ArrayHelper(ArrayProperty, Handles[0].GetInputData());
+	int32& Index = *(int32*)Handles[2].GetOutputData();
+	bool& bSuccess = *(bool*)Handles[3].GetOutputData();
 
 	Index = INDEX_NONE;
 	bSuccess = false;
@@ -995,7 +998,7 @@ void FRigVMDispatch_ArrayFind::Execute(FRigVMExtendedExecuteContext& InContext,
 
 	if(PropertyA->SameType(PropertyB))
 	{
-		const uint8* MemoryA = Handles[1].GetData();
+		const uint8* MemoryA = Handles[1].GetInputData();
 
 		for(int32 ElementIndex = 0; ElementIndex < ArrayHelper.Num(); ElementIndex++)
 		{
@@ -1081,12 +1084,12 @@ FRigVMTemplateTypeMap FRigVMDispatch_ArrayAppend::OnNewArgumentType(const FName&
 void FRigVMDispatch_ArrayAppend::Execute(FRigVMExtendedExecuteContext& InContext, FRigVMMemoryHandleArray Handles, FRigVMPredicateBranchArray Predicates)
 {
 	const FArrayProperty* SourceArrayProperty = CastFieldChecked<FArrayProperty>(Handles[1].GetProperty());
-	FScriptArrayHelper SourceArrayHelper(SourceArrayProperty, Handles[1].GetData());
+	FScriptArrayHelper SourceArrayHelper(SourceArrayProperty, Handles[1].GetInputData());
 
 	if(SourceArrayHelper.Num() > 0)
 	{
 		const FArrayProperty* TargetArrayProperty = CastFieldChecked<FArrayProperty>(Handles[0].GetProperty());
-		FScriptArrayHelper TargetArrayHelper(TargetArrayProperty, Handles[0].GetData());
+		FScriptArrayHelper TargetArrayHelper(TargetArrayProperty, Handles[0].GetOutputData());
 
 		if(InContext.IsValidArraySize(TargetArrayHelper.Num() + SourceArrayHelper.Num()))
 		{
@@ -1162,9 +1165,9 @@ FRigVMTemplateTypeMap FRigVMDispatch_ArrayClone::OnNewArgumentType(const FName&
 void FRigVMDispatch_ArrayClone::Execute(FRigVMExtendedExecuteContext& InContext, FRigVMMemoryHandleArray Handles, FRigVMPredicateBranchArray Predicates)
 {
 	const FArrayProperty* ArrayProperty = CastFieldChecked<FArrayProperty>(Handles[0].GetProperty());
-	FScriptArrayHelper ArrayHelper(ArrayProperty, Handles[0].GetData());
+	FScriptArrayHelper ArrayHelper(ArrayProperty, Handles[0].GetInputData());
 	const FArrayProperty* CloneProperty = CastFieldChecked<FArrayProperty>(Handles[1].GetProperty());
-	FScriptArrayHelper CloneHelper(CloneProperty, Handles[1].GetData());
+	FScriptArrayHelper CloneHelper(CloneProperty, Handles[1].GetOutputData());
 
 	CloneHelper.Resize(ArrayHelper.Num());
 	if(ArrayHelper.Num() > 0)
@@ -1207,12 +1210,12 @@ FText FRigVMDispatch_ArrayClone::GetArgumentTooltip(const FName& InArgumentName,
 void FRigVMDispatch_ArrayUnion::Execute(FRigVMExtendedExecuteContext& InContext, FRigVMMemoryHandleArray Handles, FRigVMPredicateBranchArray Predicates)
 {
 	const FArrayProperty* OtherProperty = CastFieldChecked<FArrayProperty>(Handles[1].GetProperty());
-	FScriptArrayHelper OtherHelper(OtherProperty, Handles[1].GetData());
+	FScriptArrayHelper OtherHelper(OtherProperty, Handles[1].GetInputData());
 
 	if(OtherHelper.Num() > 0)
 	{
 		const FArrayProperty* ArrayProperty = CastFieldChecked<FArrayProperty>(Handles[0].GetProperty());
-		FScriptArrayHelper ArrayHelper(ArrayProperty, Handles[0].GetData());
+		FScriptArrayHelper ArrayHelper(ArrayProperty, Handles[0].GetOutputData());
 
 		const FProperty* ArrayElementProperty = ArrayProperty->Inner;
 		const FProperty* OtherElementProperty = OtherProperty->Inner;
@@ -1349,11 +1352,11 @@ FRigVMTemplateTypeMap FRigVMDispatch_ArrayDifference::OnNewArgumentType(const FN
 void FRigVMDispatch_ArrayDifference::Execute(FRigVMExtendedExecuteContext& InContext, FRigVMMemoryHandleArray Handles, FRigVMPredicateBranchArray Predicates)
 {
 	const FArrayProperty* ArrayProperty = CastFieldChecked<FArrayProperty>(Handles[0].GetProperty());
-	FScriptArrayHelper ArrayHelper(ArrayProperty, Handles[0].GetData());
+	FScriptArrayHelper ArrayHelper(ArrayProperty, Handles[0].GetInputData());
 	const FArrayProperty* OtherProperty = CastFieldChecked<FArrayProperty>(Handles[1].GetProperty());
-	FScriptArrayHelper OtherHelper(OtherProperty, Handles[1].GetData());
+	FScriptArrayHelper OtherHelper(OtherProperty, Handles[1].GetInputData());
 	const FArrayProperty* ResultProperty = CastFieldChecked<FArrayProperty>(Handles[2].GetProperty());
-	FScriptArrayHelper ResultHelper(ResultProperty, Handles[2].GetData());
+	FScriptArrayHelper ResultHelper(ResultProperty, Handles[2].GetOutputData());
 
 	const FProperty* ArrayElementProperty = ArrayProperty->Inner;
 	const FProperty* OtherElementProperty = OtherProperty->Inner;
@@ -1423,11 +1426,11 @@ FText FRigVMDispatch_ArrayDifference::GetArgumentTooltip(const FName& InArgument
 void FRigVMDispatch_ArrayIntersection::Execute(FRigVMExtendedExecuteContext& InContext, FRigVMMemoryHandleArray Handles, FRigVMPredicateBranchArray Predicates)
 {
 	const FArrayProperty* ArrayProperty = CastFieldChecked<FArrayProperty>(Handles[0].GetProperty());
-	FScriptArrayHelper ArrayHelper(ArrayProperty, Handles[0].GetData());
+	FScriptArrayHelper ArrayHelper(ArrayProperty, Handles[0].GetInputData());
 	const FArrayProperty* OtherProperty = CastFieldChecked<FArrayProperty>(Handles[1].GetProperty());
-	FScriptArrayHelper OtherHelper(OtherProperty, Handles[1].GetData());
+	FScriptArrayHelper OtherHelper(OtherProperty, Handles[1].GetInputData());
 	const FArrayProperty* ResultProperty = CastFieldChecked<FArrayProperty>(Handles[2].GetProperty());
-	FScriptArrayHelper ResultHelper(ResultProperty, Handles[2].GetData());
+	FScriptArrayHelper ResultHelper(ResultProperty, Handles[2].GetOutputData());
 
 	const FProperty* ArrayElementProperty = ArrayProperty->Inner;
 	const FProperty* ResultElementProperty = ResultProperty->Inner;
@@ -1561,11 +1564,11 @@ FRigVMTemplateTypeMap FRigVMDispatch_ArrayIterator::OnNewArgumentType(const FNam
 void FRigVMDispatch_ArrayIterator::Execute(FRigVMExtendedExecuteContext& InContext, FRigVMMemoryHandleArray Handles, FRigVMPredicateBranchArray Predicates)
 {
 	const FArrayProperty* ArrayProperty = CastFieldChecked<FArrayProperty>(Handles[0].GetProperty());
-	FScriptArrayHelper ArrayHelper(ArrayProperty, Handles[0].GetData());
-	int32& Index = *(int32*)Handles[2].GetData();
-	int32& Count = *(int32*)Handles[3].GetData();
-	float& Ratio = *(float*)Handles[4].GetData();
-	FName& Block = *(FName*)Handles[5].GetData();
+	FScriptArrayHelper ArrayHelper(ArrayProperty, Handles[0].GetInputData());
+	int32& Index = *(int32*)Handles[2].GetOutputData();
+	int32& Count = *(int32*)Handles[3].GetOutputData();
+	float& Ratio = *(float*)Handles[4].GetOutputData();
+	FName& Block = *(FName*)Handles[5].GetOutputData();
 
 	if(Block.IsNone())
 	{
@@ -1587,7 +1590,7 @@ void FRigVMDispatch_ArrayIterator::Execute(FRigVMExtendedExecuteContext& InConte
 		return;
 	}
 	
-	uint8* TargetMemory = Handles[1].GetData();
+	uint8* TargetMemory = Handles[1].GetOutputData();
 	const uint8* SourceMemory = ArrayHelper.GetRawPtr(Index);
 	URigVMMemoryStorage::CopyProperty(Handles[1].GetProperty(), TargetMemory, ArrayProperty->Inner, SourceMemory);
 }
diff --git a/Engine/Plugins/Runtime/RigVM/Source/RigVM/Private/RigVMFunctions/RigVMDispatch_CastEnum.cpp b/Engine/Plugins/Runtime/RigVM/Source/RigVM/Private/RigVMFunctions/RigVMDispatch_CastEnum.cpp
--- a/Engine/Plugins/Runtime/RigVM/Source/RigVM/Private/RigVMFunctions/RigVMDispatch_CastEnum.cpp
+++ b/Engine/Plugins/Runtime/RigVM/Source/RigVM/Private/RigVMFunctions/RigVMDispatch_CastEnum.cpp
@@ -58,8 +58,8 @@ void FRigVMDispatch_CastEnumToInt::Execute(FRigVMExtendedExecuteContext& InConte
 		return;
 	}
 
-	const uint8* ValuePtr = Handles[0].GetData();
-	int32* ResultPtr = (int32*)Handles[1].GetData();
+	const uint8* ValuePtr = Handles[0].GetInputData();
+	int32* ResultPtr = (int32*)Handles[1].GetOutputData();
 	if (ValuePtr == nullptr || ResultPtr == nullptr)
 	{
 		return;
@@ -152,8 +152,8 @@ void FRigVMDispatch_CastIntToEnum::Execute(FRigVMExtendedExecuteContext& InConte
 		return;
 	}
 
-	const int32* ValuePtr = (int32*)Handles[0].GetData();
-	uint8* ResultPtr = Handles[1].GetData();
+	const int32* ValuePtr = (int32*)Handles[0].GetInputData();
+	uint8* ResultPtr = Handles[1].GetOutputData();
 	if (ValuePtr == nullptr || ResultPtr == nullptr)
 	{
 		return;
diff --git a/Engine/Plugins/Runtime/RigVM/Source/RigVM/Private/RigVMFunctions/RigVMDispatch_CastObject.cpp b/Engine/Plugins/Runtime/RigVM/Source/RigVM/Private/RigVMFunctions/RigVMDispatch_CastObject.cpp
--- a/Engine/Plugins/Runtime/RigVM/Source/RigVM/Private/RigVMFunctions/RigVMDispatch_CastObject.cpp
+++ b/Engine/Plugins/Runtime/RigVM/Source/RigVM/Private/RigVMFunctions/RigVMDispatch_CastObject.cpp
@@ -104,8 +104,8 @@ FText FRigVMDispatch_CastObject::GetNodeTooltip(const FRigVMTemplateTypeMap& InT
 
 void FRigVMDispatch_CastObject::Execute(FRigVMExtendedExecuteContext& InContext, FRigVMMemoryHandleArray Handles, FRigVMPredicateBranchArray RigVMBranches)
 {
-	UObject** ValueObjectPtr = (UObject**)Handles[0].GetData();
-	UObject** ResultObjectPtr = (UObject**)Handles[1].GetData();
+	UObject** ValueObjectPtr = (UObject**)Handles[0].GetInputData();
+	UObject** ResultObjectPtr = (UObject**)Handles[1].GetOutputData();
 	const FObjectPropertyBase* ValueObjectProperty = CastFieldChecked<FObjectPropertyBase>(Handles[0].GetProperty());
 	const FObjectPropertyBase* ResultObjectProperty = CastFieldChecked<FObjectPropertyBase>(Handles[1].GetProperty());
 	if(ValueObjectPtr && ResultObjectPtr && ValueObjectProperty && ResultObjectProperty)
diff --git a/Engine/Plugins/Runtime/RigVM/Source/RigVM/Private/RigVMFunctions/RigVMDispatch_Core.cpp b/Engine/Plugins/Runtime/RigVM/Source/RigVM/Private/RigVMFunctions/RigVMDispatch_Core.cpp
--- a/Engine/Plugins/Runtime/RigVM/Source/RigVM/Private/RigVMFunctions/RigVMDispatch_Core.cpp
+++ b/Engine/Plugins/Runtime/RigVM/Source/RigVM/Private/RigVMFunctions/RigVMDispatch_Core.cpp
@@ -114,9 +114,9 @@ void FRigVMDispatch_CoreEquals::Execute(FRigVMExtendedExecuteContext& InContext,
 	check(PropertyA->SameType(PropertyB));
 	check(Handles[2].IsBool());
 
-	const uint8* A = Handles[0].GetData();
-	const uint8* B = Handles[1].GetData();
-	bool& Result = *reinterpret_cast<bool*>(Handles[2].GetData());
+	const uint8* A = Handles[0].GetInputData();
+	const uint8* B = Handles[1].GetInputData();
+	bool& Result = *reinterpret_cast<bool*>(Handles[2].GetOutputData());
 	Result = PropertyA->Identical(A, B);
 	Result = AdaptResult(Result, InContext);
 }
diff --git a/Engine/Plugins/Runtime/RigVM/Source/RigVM/Private/RigVMFunctions/RigVMDispatch_If.cpp b/Engine/Plugins/Runtime/RigVM/Source/RigVM/Private/RigVMFunctions/RigVMDispatch_If.cpp
--- a/Engine/Plugins/Runtime/RigVM/Source/RigVM/Private/RigVMFunctions/RigVMDispatch_If.cpp
+++ b/Engine/Plugins/Runtime/RigVM/Source/RigVM/Private/RigVMFunctions/RigVMDispatch_If.cpp
@@ -75,7 +75,7 @@ void FRigVMDispatch_If::Execute(FRigVMExtendedExecuteContext& InContext, FRigVMM
 	check(Handles[0].IsBool());
 #endif
 
-	const bool& Condition = *(const bool*)Handles[0].GetData();
+	const bool& Condition = *(const bool*)Handles[0].GetInputData();
 
 	FRigVMMemoryHandle& InputHandle = Condition ? Handles[1] : Handles[2];
 	if(InputHandle.IsLazy())
@@ -83,8 +83,8 @@ void FRigVMDispatch_If::Execute(FRigVMExtendedExecuteContext& InContext, FRigVMM
 		InputHandle.ComputeLazyValueIfNecessary(InContext, InContext.GetSliceHash());
 	}
 
-	const uint8* Input = InputHandle.GetData();
-	uint8* Result = Handles[3].GetData();
+	const uint8* Input = InputHandle.GetInputData();
+	uint8* Result = Handles[3].GetOutputData();
 
 	(void)CopyProperty(
 		Handles[3].GetProperty(), Result,
diff --git a/Engine/Plugins/Runtime/RigVM/Source/RigVM/Private/RigVMFunctions/RigVMDispatch_MakeStruct.cpp b/Engine/Plugins/Runtime/RigVM/Source/RigVM/Private/RigVMFunctions/RigVMDispatch_MakeStruct.cpp
--- a/Engine/Plugins/Runtime/RigVM/Source/RigVM/Private/RigVMFunctions/RigVMDispatch_MakeStruct.cpp
+++ b/Engine/Plugins/Runtime/RigVM/Source/RigVM/Private/RigVMFunctions/RigVMDispatch_MakeStruct.cpp
@@ -96,8 +96,8 @@ void FRigVMDispatch_MakeStruct::Execute(FRigVMExtendedExecuteContext& InContext,
 	const FProperty* SourceProperty = Handles[0].GetProperty();
 	const FProperty* TargetProperty = Handles[1].GetProperty();
 	check(SourceProperty->SameType(TargetProperty));
-	const uint8* SourceMemory = Handles[0].GetData();
-	uint8* TargetMemory = Handles[1].GetData();
+	const uint8* SourceMemory = Handles[0].GetInputData();
+	uint8* TargetMemory = Handles[1].GetOutputData();
 	URigVMMemoryStorage::CopyProperty(TargetProperty, TargetMemory, SourceProperty, SourceMemory);
 }
 
diff --git a/Engine/Plugins/Runtime/RigVM/Source/RigVM/Private/RigVMFunctions/RigVMDispatch_Print.cpp b/Engine/Plugins/Runtime/RigVM/Source/RigVM/Private/RigVMFunctions/RigVMDispatch_Print.cpp
--- a/Engine/Plugins/Runtime/RigVM/Source/RigVM/Private/RigVMFunctions/RigVMDispatch_Print.cpp
+++ b/Engine/Plugins/Runtime/RigVM/Source/RigVM/Private/RigVMFunctions/RigVMDispatch_Print.cpp
@@ -97,11 +97,11 @@ void FRigVMDispatch_Print::Execute(FRigVMExtendedExecuteContext& InContext, FRig
 	check(Handles[3].IsFloat());
 	check(Handles[4].IsType<FLinearColor>());
 
-	const FString& Prefix = *(const FString*)Handles[0].GetData();
-	const bool bEnabled = *(const bool*)Handles[2].GetData();
-	const float& ScreenDuration = *(const float*)Handles[3].GetData();
-	const FLinearColor& ScreenColor = *(const FLinearColor*)Handles[4].GetData();
-	const uint8* Value = Handles[1].GetData();
+	const FString& Prefix = *(const FString*)Handles[0].GetInputData();
+	const bool bEnabled = *(const bool*)Handles[2].GetInputData();
+	const float& ScreenDuration = *(const float*)Handles[3].GetInputData();
+	const FLinearColor& ScreenColor = *(const FLinearColor*)Handles[4].GetInputData();
+	const uint8* Value = Handles[1].GetInputData();
 	
 	if(!bEnabled)
 	{
diff --git a/Engine/Plugins/Runtime/RigVM/Source/RigVM/Private/RigVMFunctions/RigVMDispatch_Select.cpp b/Engine/Plugins/Runtime/RigVM/Source/RigVM/Private/RigVMFunctions/RigVMDispatch_Select.cpp
--- a/Engine/Plugins/Runtime/RigVM/Source/RigVM/Private/RigVMFunctions/RigVMDispatch_Select.cpp
+++ b/Engine/Plugins/Runtime/RigVM/Source/RigVM/Private/RigVMFunctions/RigVMDispatch_Select.cpp
@@ -123,7 +123,7 @@ void FRigVMDispatch_SelectInt32::Execute(FRigVMExtendedExecuteContext& InContext
 	check(Handles[0].IsInt32());
 #endif
 
-	int32 Index = *(const int32*)Handles[0].GetData();
+	int32 Index = *(const int32*)Handles[0].GetInputData();
 	if(!FMath::IsWithin(Index, 0, NumCases))
 	{
 		static constexpr TCHAR InvalidIndexMessage[] = TEXT("SelectInt32: Invalid index %d - num cases is %d. Falling back to index 0.");
@@ -137,8 +137,8 @@ void FRigVMDispatch_SelectInt32::Execute(FRigVMExtendedExecuteContext& InContext
 		InputHandle.ComputeLazyValueIfNecessary(InContext, InContext.GetSliceHash());
 	}
 
-	const uint8* Input = InputHandle.GetData();
-	uint8* Result = Handles.Last().GetData();
+	const uint8* Input = InputHandle.GetInputData();
+	uint8* Result = Handles.Last().GetOutputData();
 
 	// copy property performs compatibility check and will report accordingly
 	(void)CopyProperty(
diff --git a/Engine/Plugins/Runtime/RigVM/Source/RigVM/Private/RigVMFunctions/RigVMDispatch_Switch.cpp b/Engine/Plugins/Runtime/RigVM/Source/RigVM/Private/RigVMFunctions/RigVMDispatch_Switch.cpp
--- a/Engine/Plugins/Runtime/RigVM/Source/RigVM/Private/RigVMFunctions/RigVMDispatch_Switch.cpp
+++ b/Engine/Plugins/Runtime/RigVM/Source/RigVM/Private/RigVMFunctions/RigVMDispatch_Switch.cpp
@@ -131,8 +131,8 @@ void FRigVMDispatch_SwitchInt32::Execute(FRigVMExtendedExecuteContext& InContext
 	check(Handles[1].IsName());
 #endif
 
-	const int32 Index = *(const int32*)Handles[0].GetData();
-	FName& BlockToRun = *(FName*)Handles[1].GetData();
+	const int32 Index = *(const int32*)Handles[0].GetInputData();
+	FName& BlockToRun = *(FName*)Handles[1].GetOutputData();
 
 	if(BlockToRun.IsNone())
 	{
diff --git a/Engine/Plugins/Runtime/RigVM/Source/RigVM/Private/RigVMFunctions/RigVMFunction_String.cpp b/Engine/Plugins/Runtime/RigVM/Source/RigVM/Private/RigVMFunctions/RigVMFunction_String.cpp
--- a/Engine/Plugins/Runtime/RigVM/Source/RigVM/Private/RigVMFunctions/RigVMFunction_String.cpp
+++ b/Engine/Plugins/Runtime/RigVM/Source/RigVM/Private/RigVMFunctions/RigVMFunction_String.cpp
@@ -288,8 +288,8 @@ FRigVMFunctionPtr FRigDispatch_ToString::GetDispatchFunctionImpl(const FRigVMTem
 			check(Handles[0].IsName());
 			check(Handles[1].IsString());
 
-			const FName& Value = *(const FName*)Handles[0].GetData();
-			FString& Result = *(FString*)Handles[1].GetData();
+			const FName& Value = *(const FName*)Handles[0].GetInputData();
+			FString& Result = *(FString*)Handles[1].GetOutputData();
 			Result = Value.ToString();
 		};
 	}
@@ -300,8 +300,8 @@ FRigVMFunctionPtr FRigDispatch_ToString::GetDispatchFunctionImpl(const FRigVMTem
 			check(Handles[0].IsString());
 			check(Handles[1].IsString());
 
-			const FString& Value = *(const FString*)Handles[0].GetData();
-			FString& Result = *(FString*)Handles[1].GetData();
+			const FString& Value = *(const FString*)Handles[0].GetInputData();
+			FString& Result = *(FString*)Handles[1].GetOutputData();
 			Result = Value;
 		};
 	}
@@ -315,8 +315,8 @@ void FRigDispatch_ToString::Execute(FRigVMExtendedExecuteContext& InContext, FRi
 	check(ValueProperty);
 	check(Handles[1].IsString());
 
-	const uint8* Value = Handles[0].GetData();
-	FString& Result = *(FString*)Handles[1].GetData();
+	const uint8* Value = Handles[0].GetInputData();
+	FString& Result = *(FString*)Handles[1].GetOutputData();
 	Result.Reset();
 	ValueProperty->ExportText_Direct(Result, Value, Value, nullptr, PPF_None, nullptr);
 }
@@ -356,8 +356,8 @@ FRigVMFunctionPtr FRigDispatch_FromString::GetDispatchFunctionImpl(const FRigVMT
 			check(Handles[0].IsString());
 			check(Handles[1].IsName());
 
-			const FString& String = *(FString*)Handles[0].GetData();
-			FName& Result = *(FName*)Handles[1].GetData();
+			const FString& String = *(FString*)Handles[0].GetInputData();
+			FName& Result = *(FName*)Handles[1].GetOutputData();
 			Result = *String;
 		};
 	}
@@ -368,8 +368,8 @@ FRigVMFunctionPtr FRigDispatch_FromString::GetDispatchFunctionImpl(const FRigVMT
 			check(Handles[0].IsString());
 			check(Handles[1].IsString());
 
-			const FString& String = *(FString*)Handles[0].GetData();
-			FString& Result = *(FString*)Handles[1].GetData();
+			const FString& String = *(FString*)Handles[0].GetInputData();
+			FString& Result = *(FString*)Handles[1].GetOutputData();
 			Result = String;
 		};
 	}
@@ -383,8 +383,8 @@ void FRigDispatch_FromString::Execute(FRigVMExtendedExecuteContext& InContext, F
 	check(ValueProperty);
 	check(Handles[0].IsString());
 
-	const FString& String = *(const FString*)Handles[0].GetData();
-	uint8* Value = Handles[1].GetData();
+	const FString& String = *(const FString*)Handles[0].GetInputData();
+	uint8* Value = Handles[1].GetOutputData();
 
 	class FRigDispatch_FromString_ErrorPipe : public FOutputDevice
 	{
diff --git a/Engine/Plugins/Runtime/RigVM/Source/RigVM/Public/RigVMCore/RigVMMemoryStorage.h b/Engine/Plugins/Runtime/RigVM/Source/RigVM/Public/RigVMCore/RigVMMemoryStorage.h
--- a/Engine/Plugins/Runtime/RigVM/Source/RigVM/Public/RigVMCore/RigVMMemoryStorage.h
+++ b/Engine/Plugins/Runtime/RigVM/Source/RigVM/Public/RigVMCore/RigVMMemoryStorage.h
@@ -263,7 +263,8 @@ public:
 	 * @param InSliceIndex If this is != INDEX_NONE the memory handle will return the slice of the memory requested
 	 * @return The traversed memory cached by this handle.
 	 */
-	uint8* GetData(bool bFollowPropertyPath = false, int32 InSliceIndex = INDEX_NONE)
+	UE_DEPRECATED(5.7, "Please use GetInputData, GetOutputData or GetPrivateData instead.")
+	uint8* GetData(bool bFollowPropertyPath, int32 InSliceIndex)
 	{
 		return GetData_Internal(bFollowPropertyPath, InSliceIndex);
 	}
@@ -274,21 +275,51 @@ public:
 	 * @param InSliceIndex If this is != INDEX_NONE the memory handle will return the slice of the memory requested
 	 * @return The traversed memory cached by this handle (const)
 	 */
-	const uint8* GetData(bool bFollowPropertyPath = false, int32 InSliceIndex = INDEX_NONE) const
+	UE_DEPRECATED(5.7, "Please use GetInputData, GetOutputData or GetPrivateData instead.")
+	const uint8* GetData(bool bFollowPropertyPath, int32 InSliceIndex) const
 	{
 		return GetData_Internal(bFollowPropertyPath, InSliceIndex);
 	}
 
+	/**
+	 * Returns the cached pointer stored within the handle.
+	 * @param InSliceIndex If this is != INDEX_NONE the memory handle will return the slice of the memory requested
+	 * @return The traversed memory cached by this handle.
+	 */
+	const uint8* GetInputData(int32 InSliceIndex = INDEX_NONE) const
+	{
+		return GetData_Internal(true, InSliceIndex);
+	}
+
+	/**
+	 * Returns the cached pointer stored within the handle.
+	 * @param InSliceIndex If this is != INDEX_NONE the memory handle will return the slice of the memory requested
+	 * @return The traversed memory cached by this handle.
+	 */
+	uint8* GetOutputData(int32 InSliceIndex = INDEX_NONE)
+	{
+		return GetData_Internal(false, InSliceIndex);
+	}
+
+	/**
+	 * Returns the cached pointer stored within the handle.
+	 * @param InSliceIndex If this is != INDEX_NONE the memory handle will return the slice of the memory requested
+	 * @return The traversed memory cached by this handle.
+	 */
+	uint8* GetPrivateData(int32 InSliceIndex)
+	{
+		return GetData_Internal(false, InSliceIndex);
+	}
+
 	/**
 	 * Computes the data if necessary and returns the cached pointer stored within the handle.
-	 * @param bFollowPropertyPath If set to true the memory handle will traverse the memory using the property path
 	 * @param InSliceIndex If this is != INDEX_NONE the memory handle will return the slice of the memory requested
 	 * @return The traversed memory cached by this handle (const)
 	 */
 	template<typename ComputedType>
-	TRigVMLazyValue<ComputedType> GetDataLazily(bool bFollowPropertyPath = false, int32 InSliceIndex = INDEX_NONE) const
+	TRigVMLazyValue<ComputedType> GetDataLazily(int32 InSliceIndex = INDEX_NONE) const
 	{
-		return GetDataLazily_Internal(bFollowPropertyPath, InSliceIndex);
+		return GetDataLazily_Internal(false, InSliceIndex);
 	}
 
 	/**
@@ -966,6 +997,8 @@ private:
 
 	friend class URigVM;
 	friend struct TRigVMLazyValueBase;
+	friend class URigVMMemoryStorage;
+	friend struct FRigVMMemoryStorageStruct;
 };
 
 //////////////////////////////////////////////////////////////////////////////
diff --git a/Engine/Plugins/Runtime/RigVM/Source/RigVM/Public/RigVMCore/RigVMNativized.h b/Engine/Plugins/Runtime/RigVM/Source/RigVM/Public/RigVMCore/RigVMNativized.h
--- a/Engine/Plugins/Runtime/RigVM/Source/RigVM/Public/RigVMCore/RigVMNativized.h
+++ b/Engine/Plugins/Runtime/RigVM/Source/RigVM/Public/RigVMCore/RigVMNativized.h
@@ -540,7 +540,7 @@ protected:
 	TRigVMLazyValue<T> GetLazyValue(int32 InIndex, T& InMemory, const FProperty* InProperty, TFunction<ERigVMExecuteResult()> InLambda)
 	{
 		const FRigVMMemoryHandle& Handle = GetLazyMemoryHandle(InIndex, (uint8*)&InMemory, InProperty, InLambda);
-		return Handle.GetDataLazily<T>(false, INDEX_NONE);
+		return Handle.GetDataLazily<T>(INDEX_NONE);
 	}
 
 	UE_API void AllocateLazyMemoryHandles(int32 InCount);
diff --git a/Engine/Plugins/Runtime/RigVM/Source/RigVM/Public/RigVMFunctions/RigVMDispatch_Core.h b/Engine/Plugins/Runtime/RigVM/Source/RigVM/Public/RigVMFunctions/RigVMDispatch_Core.h
--- a/Engine/Plugins/Runtime/RigVM/Source/RigVM/Public/RigVMFunctions/RigVMDispatch_Core.h
+++ b/Engine/Plugins/Runtime/RigVM/Source/RigVM/Public/RigVMFunctions/RigVMDispatch_Core.h
@@ -42,9 +42,9 @@ protected:
 		check(Handles[0].IsType<T>());
 		check(Handles[1].IsType<T>());
 		check(Handles[2].IsBool());
-		const T& A = *(const T*)Handles[0].GetData();
-		const T& B = *(const T*)Handles[1].GetData();
-		bool& Result = *(bool*)Handles[2].GetData();
+		const T& A = *(const T*)Handles[0].GetInputData();
+		const T& B = *(const T*)Handles[1].GetInputData();
+		bool& Result = *(bool*)Handles[2].GetOutputData();
 		Result = AdaptResult(A == B, InContext);
 	}
 
@@ -53,9 +53,9 @@ protected:
 		check(Handles[0].IsType<FName>());
 		check(Handles[1].IsType<FName>());
 		check(Handles[2].IsBool());
-		const FName& A = *(const FName*)Handles[0].GetData();
-		const FName& B = *(const FName*)Handles[1].GetData();
-		bool& Result = *(bool*)Handles[2].GetData();
+		const FName& A = *(const FName*)Handles[0].GetInputData();
+		const FName& B = *(const FName*)Handles[1].GetInputData();
+		bool& Result = *(bool*)Handles[2].GetOutputData();
 		Result = AdaptResult(A.IsEqual(B, ENameCase::CaseSensitive), InContext);
 	}
 
@@ -64,9 +64,9 @@ protected:
 		check(Handles[0].IsType<FString>());
 		check(Handles[1].IsType<FString>());
 		check(Handles[2].IsBool());
-		const FString& A = *(const FString*)Handles[0].GetData();
-		const FString& B = *(const FString*)Handles[1].GetData();
-		bool& Result = *(bool*)Handles[2].GetData();
+		const FString& A = *(const FString*)Handles[0].GetInputData();
+		const FString& B = *(const FString*)Handles[1].GetInputData();
+		bool& Result = *(bool*)Handles[2].GetOutputData();
 		Result = AdaptResult(A.Equals(B, ESearchCase::CaseSensitive), InContext);
 	}
 
@@ -76,9 +76,9 @@ protected:
 		check(Handles[0].IsType<T>());
 		check(Handles[1].IsType<T>());
 		check(Handles[2].IsBool());
-		const T& A = *(const T*)Handles[0].GetData();
-		const T& B = *(const T*)Handles[1].GetData();
-		bool& Result = *(bool*)Handles[2].GetData();
+		const T& A = *(const T*)Handles[0].GetInputData();
+		const T& B = *(const T*)Handles[1].GetInputData();
+		bool& Result = *(bool*)Handles[2].GetOutputData();
 		Result = AdaptResult(A.Equals(B), InContext);
 	}
 
diff --git a/Engine/Plugins/Runtime/RigVM/Source/RigVMDeveloper/Private/RigVMCompiler/RigVMAST.cpp b/Engine/Plugins/Runtime/RigVM/Source/RigVMDeveloper/Private/RigVMCompiler/RigVMAST.cpp
--- a/Engine/Plugins/Runtime/RigVM/Source/RigVMDeveloper/Private/RigVMCompiler/RigVMAST.cpp
+++ b/Engine/Plugins/Runtime/RigVM/Source/RigVMDeveloper/Private/RigVMCompiler/RigVMAST.cpp
@@ -904,6 +904,35 @@ FName FRigVMEntryExprAST::GetEventName() const
 	return NAME_None;
 }
 
+void FRigVMCallExternExprAST::SetSegmentPathForChild(int32 InChildIndex, const FString& InSegmentPath)
+{
+	check(Children.IsValidIndex(InChildIndex));
+	SegmentPathForChild.FindOrAdd(InChildIndex) = InSegmentPath;
+}
+
+void FRigVMCallExternExprAST::SetSegmentPathForChild(const FRigVMExprAST* InChildExpr, const FString& InSegmentPath)
+{
+	const int32 ChildIndex = Children.Find(const_cast<FRigVMExprAST*>(InChildExpr));
+	SetSegmentPathForChild(ChildIndex, InSegmentPath);
+}
+
+const FString& FRigVMCallExternExprAST::GetSegmentPathForChild(int32 InChildIndex) const
+{
+	check(Children.IsValidIndex(InChildIndex));
+	if(const FString* ExistingSegmentPath = SegmentPathForChild.Find(InChildIndex))
+	{
+		return *ExistingSegmentPath;
+	}
+	static const FString EmptySegmentPath;
+	return EmptySegmentPath;
+}
+
+const FString& FRigVMCallExternExprAST::GetSegmentPathForChild(const FRigVMExprAST* InChildExpr) const
+{
+	const int32 ChildIndex = Children.Find(const_cast<FRigVMExprAST*>(InChildExpr));
+	return GetSegmentPathForChild(ChildIndex);
+}
+
 bool FRigVMVarExprAST::IsConstant() const
 {
 	if (GetPin()->IsExecuteContext())
@@ -1131,6 +1160,11 @@ FRigVMParserAST::FRigVMParserAST(TArray<URigVMGraph*> InGraphs, URigVMController
 		FoldAssignments();
 	}
 
+	if (InSettings.bFoldSubPinCopies)
+	{
+		FoldSubPinCopies();
+	}
+
 	if (InSettings.bFoldLiterals)
 	{
 		FoldLiterals();
@@ -2121,44 +2155,192 @@ void FRigVMParserAST::FoldAssignments()
 
 		// To prevent bad assignments in LWC for VMs compiled in non LWC, we do not allow folding of assignments
 		// to/from external variables of type float
+		if(IsFloatOrFloatArrayPinOnVariable(SourcePin) || IsFloatOrFloatArrayPinOnVariable(TargetPin))
 		{
-			if (SourcePin->GetCPPType() == TEXT("float") || SourcePin->GetCPPType() == TEXT("TArray<float>"))
+			continue;
+		}
+
+		FRigVMExprAST* Child = AssignExpr->Children[0];
+		AssignExpr->RemoveParent(Parent);
+		Child->RemoveParent(AssignExpr);
+
+		TArray<FRigVMExprAST*> GrandParents = Parent->Parents;
+		for (FRigVMExprAST* GrandParent : GrandParents)
+		{
+			GrandParent->ReplaceChild(Parent, Child);
+			if (GrandParent->IsA(FRigVMExprAST::EType::Assign))
 			{
-				if (URigVMVariableNode* SourceVariableNode = Cast<URigVMVariableNode>(SourcePin->GetNode()))
+				FRigVMAssignExprAST* GrandParentAssign = GrandParent->To<FRigVMAssignExprAST>();
+
+				const URigVMPin* OldSourcePin = GrandParentAssign->GetSourcePin();
+				const FString SourceSegmentPath = OldSourcePin->GetSegmentPath();
+				URigVMPin* NewSourcePin = AssignExpr->GetSourcePin();
+
+				if(IsMakeStructDispatchStructPin(NewSourcePin))
 				{
-					if (!SourceVariableNode->IsInputArgument() && !SourceVariableNode->IsLocalVariable())
-					{
-						continue;
-					}
+					NewSourcePin = NewSourcePin->GetNode()->FindRootPinByName(FRigVMDispatch_MakeStruct::ElementsName.Resolve());
+					check(NewSourcePin);
 				}
-			}
-			if (TargetPin->GetCPPType() == TEXT("float") || TargetPin->GetCPPType() == TEXT("TArray<float>"))
-			{
-				if (URigVMVariableNode* TargetVariableNode = Cast<URigVMVariableNode>(TargetPin->GetNode()))
+				
+				if(!SourceSegmentPath.IsEmpty())
 				{
-					if (!TargetVariableNode->IsInputArgument() && !TargetVariableNode->IsLocalVariable())
+					NewSourcePin = NewSourcePin->FindSubPin(SourceSegmentPath);
+					if(NewSourcePin == nullptr)
 					{
+						Settings.Reportf(
+							EMessageSeverity::Error,
+							AssignExpr->GetSourcePin()->GetNode(),
+							TEXT("Cannot resolve segment path '%s' on pin '%s'."),
+							*SourceSegmentPath,
+							*AssignExpr->GetSourcePin()->GetPinPath());
 						continue;
 					}
 				}
+				
+				GrandParentAssign->Link = FRigVMASTLinkDescription(
+					GrandParentAssign->GetSourceProxy().GetSibling(NewSourcePin),
+					GrandParentAssign->GetTargetProxy(), FString());
+				GrandParentAssign->Name = *GetLinkAsString(GrandParentAssign->GetLink());
 			}
 		}
 
-		FRigVMExprAST* Child = AssignExpr->Children[0];
-		AssignExpr->RemoveParent(Parent);
-		Child->RemoveParent(AssignExpr);
+		ExpressionsToRemove.AddUnique(AssignExpr);
+		if (Parent->Parents.Num() == 0)
+		{
+			ExpressionsToRemove.AddUnique(Parent);
+		}
+	}
+
+	RemoveExpressions(ExpressionsToRemove);
+}
+
+void FRigVMParserAST::FoldSubPinCopies()
+{
+	DECLARE_SCOPE_HIERARCHICAL_COUNTER_FUNC()
+
+	TArray<FRigVMExprAST*> ExpressionsToRemove;
+
+	// first - Fold all assignment chains
+	for (FRigVMExprAST* Expression : Expressions)
+	{
+		if (Expression->Parents.Num() == 0)
+		{
+			continue;
+		}
+
+		if (Expression->GetType() != FRigVMExprAST::EType::Copy)
+		{
+			continue;
+		}
+
+		FRigVMCopyExprAST* CopyExpr = Expression->To<FRigVMCopyExprAST>();
+		if (CopyExpr->Parents.Num() != 1 || CopyExpr->Children.Num() != 1)
+		{
+			continue;
+		}
+
+		URigVMPin* SourcePin = CopyExpr->GetSourcePin();
+		URigVMPin* TargetPin = CopyExpr->GetTargetPin();
+
+		// we only consider copies from a sub pin to a root pin in this phase
+		if(SourcePin->IsRootPin() || !TargetPin->IsRootPin())
+		{
+			continue;
+		}
+
+		// if the target pin is lazy - don't fold the copy
+		if(TargetPin->IsLazy())
+		{
+			continue;
+		}
+
+		// in case the assign has different types for left and right - we need to avoid folding
+		// since this assign represents a cast operation
+		if(SourcePin->GetCPPTypeObject() != TargetPin->GetCPPTypeObject())
+		{
+			continue;
+		}
+		if(SourcePin->GetCPPTypeObject() == nullptr)
+		{
+			if(SourcePin->GetCPPType() != TargetPin->GetCPPType())
+			{
+				continue;
+			}
+		}
+
+		// non-input pins on anything but a reroute (passthrough or literal) or array node should be skipped
+		const bool bIsReroute = TargetPin->GetNode()->IsA<URigVMRerouteNode>();
+
+		// a copy / assign to a constant node root pin can be folded 
+		const bool bIsConstantValueTarget = IsConstantDispatchValuePin(TargetPin);
+		
+		if (TargetPin->GetDirection() != ERigVMPinDirection::Input && !bIsReroute && !bIsConstantValueTarget)
+		{
+			continue;
+		}
+
+		// if this node is a loop node - let's skip the folding
+		if (const URigVMNode* ModelNode = TargetPin->GetNode())
+		{
+			if (ModelNode->IsControlFlowNode())
+			{
+				continue;
+			}
+		}
+
+		// if this node is a variable node and the pin requires a watch... skip this
+		if (Cast<URigVMVariableNode>(SourcePin->GetNode()))
+		{
+			if(SourcePin->RequiresWatch(true))
+			{
+				continue;
+			}
+		}
+
+		// Skip folding for programmatic pins as we expect them to always exist in work memory
+		if(TargetPin->IsProgrammaticPin())
+		{
+			continue;
+		}
+		
+		FRigVMExprAST* Parent = CopyExpr->Parents[0];
+		if (!Parent->IsA(FRigVMExprAST::EType::Var))
+		{
+			continue;
+		}
+
+		// To prevent bad assignments in LWC for VMs compiled in non LWC, we do not allow folding of assignments
+		// to/from external variables of type float
+		if(IsFloatOrFloatArrayPinOnVariable(SourcePin) || IsFloatOrFloatArrayPinOnVariable(TargetPin))
+		{
+			continue;
+		}
+
+		FRigVMExprAST* Child = CopyExpr->Children[0];
+		CopyExpr->RemoveParent(Parent);
+		Child->RemoveParent(CopyExpr);
 
 		TArray<FRigVMExprAST*> GrandParents = Parent->Parents;
 		for (FRigVMExprAST* GrandParent : GrandParents)
 		{
 			GrandParent->ReplaceChild(Parent, Child);
-			if (GrandParent->IsA(FRigVMExprAST::EType::Assign))
+
+			if (GrandParent->IsA(FRigVMExprAST::EType::CallExtern))
+			{
+				const URigVMPin* CopySourcePin = CopyExpr->GetSourcePin();
+				if(!CopySourcePin->IsRootPin())
+				{
+					FRigVMCallExternExprAST* GrandParentCallExtern = GrandParent->To<FRigVMCallExternExprAST>();
+					GrandParentCallExtern->SetSegmentPathForChild(Child, CopySourcePin->GetSegmentPath());
+				}
+			}
+			else if (GrandParent->IsA(FRigVMExprAST::EType::Assign))
 			{
 				FRigVMAssignExprAST* GrandParentAssign = GrandParent->To<FRigVMAssignExprAST>();
 
 				const URigVMPin* OldSourcePin = GrandParentAssign->GetSourcePin();
 				const FString SourceSegmentPath = OldSourcePin->GetSegmentPath();
-				URigVMPin* NewSourcePin = AssignExpr->GetSourcePin();
+				URigVMPin* NewSourcePin = CopyExpr->GetSourcePin();
 
 				if(!SourceSegmentPath.IsEmpty())
 				{
@@ -2167,10 +2349,10 @@ void FRigVMParserAST::FoldAssignments()
 					{
 						Settings.Reportf(
 							EMessageSeverity::Error,
-							AssignExpr->GetSourcePin()->GetNode(),
+							CopyExpr->GetSourcePin()->GetNode(),
 							TEXT("Cannot resolve segment path '%s' on pin '%s'."),
 							*SourceSegmentPath,
-							*AssignExpr->GetSourcePin()->GetPinPath());
+							*CopyExpr->GetSourcePin()->GetPinPath());
 						continue;
 					}
 				}
@@ -2182,7 +2364,7 @@ void FRigVMParserAST::FoldAssignments()
 			}
 		}
 
-		ExpressionsToRemove.AddUnique(AssignExpr);
+		ExpressionsToRemove.AddUnique(CopyExpr);
 		if (Parent->Parents.Num() == 0)
 		{
 			ExpressionsToRemove.AddUnique(Parent);
@@ -3738,3 +3920,20 @@ bool FRigVMParserAST::IsBreakStructDispatchElementsPin(const URigVMPin* InPin)
 	return false;
 
 }
+
+bool FRigVMParserAST::IsFloatOrFloatArrayPinOnVariable(const URigVMPin* InPin)
+{
+	// To prevent bad assignments in LWC for VMs compiled in non LWC, we do not allow folding of assignments
+	// to/from external variables of type float
+	if (InPin->GetCPPType() == TEXT("float") || InPin->GetCPPType() == TEXT("TArray<float>"))
+	{
+		if (const URigVMVariableNode* VariableNode = Cast<URigVMVariableNode>(InPin->GetNode()))
+		{
+			if (!VariableNode->IsInputArgument() && !VariableNode->IsLocalVariable())
+			{
+				return true;
+			}
+		}
+	}
+	return false;
+}
diff --git a/Engine/Plugins/Runtime/RigVM/Source/RigVMDeveloper/Private/RigVMCompiler/RigVMCompiler.cpp b/Engine/Plugins/Runtime/RigVM/Source/RigVMDeveloper/Private/RigVMCompiler/RigVMCompiler.cpp
--- a/Engine/Plugins/Runtime/RigVM/Source/RigVMDeveloper/Private/RigVMCompiler/RigVMCompiler.cpp
+++ b/Engine/Plugins/Runtime/RigVM/Source/RigVMDeveloper/Private/RigVMCompiler/RigVMCompiler.cpp
@@ -2124,6 +2124,7 @@ bool URigVMCompiler::TraverseCallExtern(const FRigVMCallExternExprAST* InExpr, F
 		}
 	}
 
+	const FRigVMCallExternExprAST* CallExternExpr = InExpr->To<FRigVMCallExternExprAST>();
 	int32 CallExternInstructionIndex = INDEX_NONE;
 	const FRigVMCallstack Callstack = InExpr->GetProxy().GetCallstack();
 
@@ -2134,6 +2135,29 @@ bool URigVMCompiler::TraverseCallExtern(const FRigVMCallExternExprAST* InExpr, F
 			return false;
 		}
 
+		// setup additional property paths needed for call externs accessing sub pin data directly
+		for(const TPair<int32, FString>& SegmentPathPair : CallExternExpr->SegmentPathForChild)
+		{
+			const FRigVMExprAST* ChildExpr = CallExternExpr->ChildAt(SegmentPathPair.Key);
+			check(ChildExpr);
+
+			if(ChildExpr->IsA(FRigVMExprAST::EType::CachedValue))
+			{
+				ChildExpr = ChildExpr->To<FRigVMCachedValueExprAST>()->GetVarExpr();
+			}
+
+			if(!ChildExpr->IsA(FRigVMExprAST::EType::Var))
+			{
+				continue;
+			}
+
+			const FRigVMVarExprAST* SourceVarExpr = GetSourceVarExpr(ChildExpr);
+			const FRigVMOperand RootOperand = WorkData.ExprToOperand.FindChecked(SourceVarExpr);
+			const FString HeadCPPType = SourceVarExpr->GetPin()->GetRootPin()->GetCPPType();
+
+			WorkData.FindOrAddPropertyPath(RootOperand, HeadCPPType, SegmentPathPair.Value);
+		}
+
 		if(WorkData.Settings.ASTSettings.bSetupTraits)
 		{
 			const TArray<URigVMPin*> TraitPins = Node->GetTraitPins(); 
@@ -2190,6 +2214,8 @@ bool URigVMCompiler::TraverseCallExtern(const FRigVMCallExternExprAST* InExpr, F
 	else
 	{
 		TArray<FRigVMOperand> Operands;
+		TArray<const FRigVMExprAST*> ChildExpressionForOperand;
+		TArray<const FRigVMVarExprAST*> VarExpressionForOperand;
 
 		// iterate over the child expressions in the order of the arguments on the function
 		const FRigVMFunction* Function = nullptr;
@@ -2214,24 +2240,31 @@ bool URigVMCompiler::TraverseCallExtern(const FRigVMCallExternExprAST* InExpr, F
 				&WorkData,
 				Argument,
 				&Operands,
+				&ChildExpressionForOperand,
+				&VarExpressionForOperand,
 				&BlockToRunOperand,
 				&CountOperand,
 				&IndexOperand
 			](const FRigVMExprAST* InExpr)
 			{
+				const FRigVMVarExprAST* SourceVarExpr = nullptr;
 				if (InExpr->GetType() == FRigVMExprAST::EType::CachedValue)
 				{
-					Operands.Add(WorkData.ExprToOperand.FindChecked(GetSourceVarExpr(InExpr->To<FRigVMCachedValueExprAST>()->GetVarExpr())));
+					SourceVarExpr = GetSourceVarExpr(InExpr->To<FRigVMCachedValueExprAST>()->GetVarExpr());
 				}
 				else if (InExpr->IsA(FRigVMExprAST::EType::Var))
 				{
-					Operands.Add(WorkData.ExprToOperand.FindChecked(GetSourceVarExpr(InExpr->To<FRigVMVarExprAST>())));
+					SourceVarExpr = GetSourceVarExpr(InExpr->To<FRigVMVarExprAST>());
 				}
 				else
 				{
 					return;
 				}
 
+				Operands.Add(WorkData.ExprToOperand.FindChecked(SourceVarExpr));
+				ChildExpressionForOperand.Add(InExpr);
+				VarExpressionForOperand.Add(SourceVarExpr);
+
 				if(Argument.Name == FRigVMStruct::ControlFlowBlockToRunName)
 				{
 					BlockToRunOperand = Operands.Last();
@@ -2267,6 +2300,26 @@ bool URigVMCompiler::TraverseCallExtern(const FRigVMCallExternExprAST* InExpr, F
 			ProcessArgument(ChildExpr);
 		}
 
+		if(!CallExternExpr->SegmentPathForChild.IsEmpty())
+		{
+			for(int32 OperandIndex = 0; OperandIndex < Operands.Num(); OperandIndex++)
+			{
+				const FString& SegmentPath = CallExternExpr->GetSegmentPathForChild(ChildExpressionForOperand[OperandIndex]);
+				if(SegmentPath.IsEmpty())
+				{
+					continue;
+				}
+
+				const FRigVMVarExprAST* SourceVarExpr = VarExpressionForOperand[OperandIndex];
+				check(SourceVarExpr);
+
+				const FString HeadCPPType = SourceVarExpr->GetPin()->GetCPPType();
+				const int32 RegisterOffsetIndex = WorkData.FindOrAddPropertyPath(Operands[OperandIndex], HeadCPPType, SegmentPath);
+				check(RegisterOffsetIndex != INDEX_NONE);
+				Operands[OperandIndex].RegisterOffset = static_cast<uint16>(RegisterOffsetIndex);
+			}
+		}
+
 		// make sure to skip the output blocks while we are traversing this call extern
 		TArray<const FRigVMExprAST*> ExpressionsToSkip;
 		TArray<int32> BranchIndices;
diff --git a/Engine/Plugins/Runtime/RigVM/Source/RigVMDeveloper/Public/RigVMCompiler/RigVMAST.h b/Engine/Plugins/Runtime/RigVM/Source/RigVMDeveloper/Public/RigVMCompiler/RigVMAST.h
--- a/Engine/Plugins/Runtime/RigVM/Source/RigVMDeveloper/Public/RigVMCompiler/RigVMAST.h
+++ b/Engine/Plugins/Runtime/RigVM/Source/RigVMDeveloper/Public/RigVMCompiler/RigVMAST.h
@@ -719,6 +719,12 @@ public:
 		return InType == EType::CallExtern;
 	};
 
+	void SetSegmentPathForChild(int32 InChildIndex, const FString& InSegmentPath);
+	void SetSegmentPathForChild(const FRigVMExprAST* InChildExpr, const FString& InSegmentPath);
+
+	const FString& GetSegmentPathForChild(int32 InChildIndex) const;
+	const FString& GetSegmentPathForChild(const FRigVMExprAST* InChildExpr) const;
+
 protected:
 
 	// default constructor (protected so that only parser can access it)
@@ -728,7 +734,10 @@ protected:
 
 private:
 
+	TMap<int32,FString> SegmentPathForChild;
+
 	friend class FRigVMParserAST;
+	friend class URigVMCompiler;
 };
 
 /*
@@ -1152,6 +1161,10 @@ struct FRigVMParserASTSettings
 	UPROPERTY(EditAnywhere, Category = "AST")
 	bool bFoldAssignments = false;
 	
+   // fold sub pin assignments / copies
+   UPROPERTY(EditAnywhere, Category = "AST")
+   bool bFoldSubPinCopies = false;
+   
 	// fold literals and share memory
 	UPROPERTY(EditAnywhere, Category = "AST")
 	bool bFoldLiterals = false;
@@ -1173,6 +1186,7 @@ struct FRigVMParserASTSettings
 	{
 		FRigVMParserASTSettings Settings;
 		Settings.bFoldAssignments = false;
+		Settings.bFoldSubPinCopies = false;
 		Settings.bFoldLiterals = false;
 		Settings.ExecuteContextStruct = InExecuteContextStruct ? InExecuteContextStruct : FRigVMExecuteContext::StaticStruct();
 		return Settings;
@@ -1184,6 +1198,7 @@ struct FRigVMParserASTSettings
 	{
 		FRigVMParserASTSettings Settings;
 		Settings.bFoldAssignments = true;
+		Settings.bFoldSubPinCopies = false;
 		Settings.bFoldLiterals = true;
 		Settings.ExecuteContextStruct = InExecuteContextStruct ? InExecuteContextStruct : FRigVMExecuteContext::StaticStruct();
 		return Settings;
@@ -1390,6 +1405,9 @@ private:
 	// helper function to fold / remove obsolete assignments and reduce assignment chains
 	UE_API void FoldAssignments();
 
+	// helper function to fold / remove obsolete copies from a sub pin to a root pin
+   	UE_API void FoldSubPinCopies();
+
 	// helper function to inline all contributing nodes of the graph
 	UE_API void Inline(const TArray<URigVMGraph*>& InGraphs);
 	UE_API void Inline(const TArray<URigVMGraph*>& InGraphs, const TArray<FRigVMASTProxy>& InNodeProxies);
@@ -1435,6 +1453,8 @@ private:
 	static bool IsBreakStructDispatchStructPin(const URigVMPin* InPin);
 	static bool IsBreakStructDispatchElementsPin(const URigVMPin* InPin);
 
+	static bool IsFloatOrFloatArrayPinOnVariable(const URigVMPin* InPin);
+
 	TMap<FRigVMASTProxy, FRigVMExprAST*> SubjectToExpression;
 	TMap<FRigVMASTProxy, int32> NodeExpressionIndex;
 	TArray<FRigVMExprAST*> Expressions;
diff --git a/Engine/Source/Programs/Shared/EpicGames.UHT/Exporters/CodeGen/UhtHeaderCodeGeneratorHFile.cs b/Engine/Source/Programs/Shared/EpicGames.UHT/Exporters/CodeGen/UhtHeaderCodeGeneratorHFile.cs
--- a/Engine/Source/Programs/Shared/EpicGames.UHT/Exporters/CodeGen/UhtHeaderCodeGeneratorHFile.cs
+++ b/Engine/Source/Programs/Shared/EpicGames.UHT/Exporters/CodeGen/UhtHeaderCodeGeneratorHFile.cs
@@ -315,12 +315,12 @@ namespace EpicGames.UHT.Exporters.CodeGen
 								baseType = baseType[..^1].Trim();
 							}
 							builder.Append("\t\t\t\t*(").Append(baseType).Append("*) Branch.MemoryHandles[")
-								.Append(parameterIndex++).Append("].GetData(false) = ")
+								.Append(parameterIndex++).Append("].GetOutputData() = ")
 								.Append(parameter.Name).Append(";  \\\r\n");
 							
 						}
 						builder.Append("\t\t\t\tBranch.Execute(*Context);  \\\r\n");
-						builder.Append("\t\t\t\treturn *(").Append(methodInfo.ReturnType).Append("*)Branch.MemoryHandles[").Append(methodInfo.Parameters.Count).Append("].GetData(false);  \\\r\n");
+						builder.Append("\t\t\t\treturn *(").Append(methodInfo.ReturnType).Append("*)Branch.MemoryHandles[").Append(methodInfo.Parameters.Count).Append("].GetInputData();  \\\r\n");
 						builder.Append("\t\t\t}  \\\r\n");
 						builder.Append("\t\t\treturn ").Append(methodInfo.Name).Append('(');
 						builder.AppendParameterNames(methodInfo.Parameters, false, ", ", false, false);
@@ -379,15 +379,23 @@ namespace EpicGames.UHT.Exporters.CodeGen
 							string additionalParameters = String.Empty;
 							if (parameter.IsLazy || (!parameter.Input && !parameter.Output && !parameter.Singleton))
 							{
-								additionalParameters = ", RigVMExecuteContext.GetSlice().GetIndex()";
+								additionalParameters = "RigVMExecuteContext.GetSlice().GetIndex()";
 							}
 
-							string getDataMethod = "GetData";
+							string getDataMethod = "GetInputData";
 							if (parameter.IsLazy)
 							{
 								getDataMethod = $"GetDataLazily<{parameter.TypeOriginal(false, false)}>";
 							}
-
+							else if (parameter.Output)
+							{
+								getDataMethod = "GetOutputData";
+							}
+							else if (!parameter.Input && !parameter.Output && !parameter.Singleton)
+							{
+								getDataMethod = "GetPrivateData";
+							}
+							
 							if (parameter.IsArray)
 							{
 								string extendedType = parameter.ExtendedType();
@@ -419,7 +427,7 @@ namespace EpicGames.UHT.Exporters.CodeGen
 									.Append(operandIndex)
 									.Append("].")
 									.Append(getDataMethod)
-									.Append("(false")
+									.Append('(')
 									.Append(additionalParameters)
 									.Append("); \\\r\n");
 								operandIndex++;
@@ -454,7 +462,7 @@ namespace EpicGames.UHT.Exporters.CodeGen
 									.Append(operandIndex)
 									.Append("].")
 									.Append(getDataMethod)
-									.Append("(false")
+									.Append('(')
 									.Append(additionalParameters)
 									.Append("); \\\r\n");
 								operandIndex++;
