Changed RigVM nodes rename to use the model title, so it works also with current UAF nodes functionality

#rnx

diff --git a/Engine/Plugins/Runtime/RigVM/Source/RigVMDeveloper/Private/EdGraph/RigVMEdGraphNode.cpp b/Engine/Plugins/Runtime/RigVM/Source/RigVMDeveloper/Private/EdGraph/RigVMEdGraphNode.cpp
--- a/Engine/Plugins/Runtime/RigVM/Source/RigVMDeveloper/Private/EdGraph/RigVMEdGraphNode.cpp
+++ b/Engine/Plugins/Runtime/RigVM/Source/RigVMDeveloper/Private/EdGraph/RigVMEdGraphNode.cpp
@@ -71,13 +71,17 @@ FText URigVMEdGraphNode::GetNodeTitle(ENodeTitleType::Type TitleType) const
 		FString SubTitle;
 		if(URigVMNode* ModelNode = GetModelNode())
 		{
-			if(bCanRenameNode)
+			if (bCanRenameNode)
 			{
-				FString NodeName = ModelNode->GetName();
-				NodeName.RemoveFromStart(TEXT("DISPATCH_"));
-				NodeName.RemoveFromStart(TEXT("RigVMFunction_"));
-				NodeName.RemoveFromStart(TEXT("RigVMDispatch_"));
-				NodeTitle = FText::FromString(NodeName);
+				FString ModelNodeTitle = ModelNode->GetNodeTitleRaw();
+				if (ModelNodeTitle.IsEmpty())
+				{
+					ModelNodeTitle = ModelNode->GetName();
+					ModelNodeTitle.RemoveFromStart(TEXT("DISPATCH_"));
+					ModelNodeTitle.RemoveFromStart(TEXT("RigVMFunction_"));
+					ModelNodeTitle.RemoveFromStart(TEXT("RigVMDispatch_"));
+				}
+				NodeTitle = FText::FromString(ModelNodeTitle);
 			}
 
 			else if (const URigVMUnitNode* UnitNode = Cast<URigVMUnitNode>(ModelNode))
diff --git a/Engine/Plugins/Runtime/RigVM/Source/RigVMEditor/Private/Editor/RigVMEditor.cpp b/Engine/Plugins/Runtime/RigVM/Source/RigVMEditor/Private/Editor/RigVMEditor.cpp
--- a/Engine/Plugins/Runtime/RigVM/Source/RigVMEditor/Private/Editor/RigVMEditor.cpp
+++ b/Engine/Plugins/Runtime/RigVM/Source/RigVMEditor/Private/Editor/RigVMEditor.cpp
@@ -1380,10 +1380,7 @@ void FRigVMEditorBase::OnNodeTitleCommitted(const FText& NewText, ETextCommit::T
 				const FString NewNodeName = NewText.ToString();
 				if(!ModelNode->GetName().Equals(NewNodeName, ESearchCase::CaseSensitive))
 				{
-					if(Controller->RenameNode(ModelNode, *NewText.ToString(), true, true))
-					{
-						RigVMNodeBeingChanged->InvalidateNodeTitle();
-					}
+					Controller->SetNodeTitle(ModelNode, *NewText.ToString(), true, true, true);
 				}
 			}
 		}
