RigVM: Fix regressions with RemapLinkedPins

#rnx

diff --git a/Engine/Plugins/Runtime/RigVM/Source/RigVMDeveloper/Private/RigVMModel/RigVMController.cpp b/Engine/Plugins/Runtime/RigVM/Source/RigVMDeveloper/Private/RigVMModel/RigVMController.cpp
--- a/Engine/Plugins/Runtime/RigVM/Source/RigVMDeveloper/Private/RigVMModel/RigVMController.cpp
+++ b/Engine/Plugins/Runtime/RigVM/Source/RigVMDeveloper/Private/RigVMModel/RigVMController.cpp
@@ -22099,6 +22099,7 @@ TArray<URigVMController::FLinkedPath> URigVMController::RemapLinkedPaths(
 			return InPinPath;
 		}
 
+		const FString OldNodeName = NodeName;
 		FString RemappedPinPath = InPinPath;
 
 		if (InSettings.bFollowCoreRedirectors)
@@ -22110,17 +22111,32 @@ TArray<URigVMController::FLinkedPath> URigVMController::RemapLinkedPaths(
 			}
 		}
 
+		if(const FRigVMController_PinPathRemapDelegate* RemapDelegateForOldName = InSettings.RemapDelegates.Find(OldNodeName))
+		{
+			FString NewRemappedPinPath = RemapDelegateForOldName->Execute(RemappedPinPath, bAsInput);
+			if (!NewRemappedPinPath.IsEmpty() && !NewRemappedPinPath.Equals(RemappedPinPath, ESearchCase::CaseSensitive))
+			{
+				RemappedPinPath = NewRemappedPinPath;
+				(void)URigVMPin::SplitPinPathAtStart(RemappedPinPath, NodeName, SegmentPath);
+			}
+		}
+
 		if(const FString* RemappedNodeName = InSettings.NodeNameMap.Find(NodeName))
 		{
 			NodeName = *RemappedNodeName;
 			RemappedPinPath = URigVMPin::JoinPinPath(*RemappedNodeName, SegmentPath);
 		}
 
-		if(const FRigVMController_PinPathRemapDelegate* RemapDelegate = InSettings.RemapDelegates.Find(NodeName))
+		if(const FRigVMController_PinPathRemapDelegate* RemapDelegateForNewName = InSettings.RemapDelegates.Find(NodeName))
 		{
-			RemappedPinPath = RemapDelegate->Execute(RemappedPinPath, bAsInput);
+			FString NewRemappedPinPath = RemapDelegateForNewName->Execute(RemappedPinPath, bAsInput);
+			if (!NewRemappedPinPath.IsEmpty() && !NewRemappedPinPath.Equals(RemappedPinPath, ESearchCase::CaseSensitive))
+			{
+				RemappedPinPath = NewRemappedPinPath;
+				(void)URigVMPin::SplitPinPathAtStart(RemappedPinPath, NodeName, SegmentPath);
+			}
 		}
-
+		
 		// if the pin cannot be found based on this path the pin may be a sub pin on something that
 		// hasn't resolved its type yet.
 		if(InSettings.bIsImportingFromText && !RemappedPinPath.IsEmpty() && GetGraph()->FindPin(RemappedPinPath) == nullptr)
