RigVM: FName IsNone and IsValid functions

#rnx

diff --git a/Engine/Plugins/Runtime/RigVM/Source/RigVM/Private/RigVMFunctions/RigVMFunction_Name.cpp b/Engine/Plugins/Runtime/RigVM/Source/RigVM/Private/RigVMFunctions/RigVMFunction_Name.cpp
--- a/Engine/Plugins/Runtime/RigVM/Source/RigVM/Private/RigVMFunctions/RigVMFunction_Name.cpp
+++ b/Engine/Plugins/Runtime/RigVM/Source/RigVM/Private/RigVMFunctions/RigVMFunction_Name.cpp
@@ -252,3 +252,13 @@ IMPLEMENT_RIGVMSTRUCT_AUTOMATION_TEST(FRigVMFunction_Contains)
 	return true;
 }
 #endif
+
+FRigVMFunction_IsNone_Execute()
+{
+	Result = Name.IsNone();
+}
+
+FRigVMFunction_IsNameValid_Execute()
+{
+	Result = !Value.IsNone();
+}
\ No newline at end of file
diff --git a/Engine/Plugins/Runtime/RigVM/Source/RigVM/Public/RigVMFunctions/RigVMFunction_Name.h b/Engine/Plugins/Runtime/RigVM/Source/RigVM/Public/RigVMFunctions/RigVMFunction_Name.h
--- a/Engine/Plugins/Runtime/RigVM/Source/RigVM/Public/RigVMFunctions/RigVMFunction_Name.h
+++ b/Engine/Plugins/Runtime/RigVM/Source/RigVM/Public/RigVMFunctions/RigVMFunction_Name.h
@@ -183,6 +183,54 @@ struct FRigVMFunction_Contains : public FRigVMFunction_NameBase
 	UPROPERTY(meta = (Input))
 	FName Search;
 
+	UPROPERTY(meta = (Output))
+	bool Result;
+};
+
+/**
+ * Returns true if this name is none
+ */
+USTRUCT(meta = (DisplayName = "Is None", TemplateName = "IsNone", Keywords = "Empty,Valid"))
+struct FRigVMFunction_IsNone : public FRigVMFunction_NameBase
+{
+	GENERATED_BODY()
+
+	FRigVMFunction_IsNone()
+	{
+		Name = NAME_None;
+		Result = false;
+	}
+
+	RIGVM_METHOD()
+		RIGVM_API virtual void Execute() override;
+
+	UPROPERTY(meta = (Input))
+	FName Name;
+
+	UPROPERTY(meta = (Output))
+	bool Result;
+};
+
+/**
+ * Returns true if this name is not valid / is none
+ */
+USTRUCT(meta = (DisplayName = "Is Valid", TemplateName = "IsValid", Keywords = "Empty,Valid"))
+struct FRigVMFunction_IsNameValid : public FRigVMFunction_NameBase
+{
+	GENERATED_BODY()
+
+	FRigVMFunction_IsNameValid()
+	{
+		Value = NAME_None;
+		Result = false;
+	}
+
+	RIGVM_METHOD()
+	RIGVM_API virtual void Execute() override;
+
+	UPROPERTY(meta = (Input))
+	FName Value;
+
 	UPROPERTY(meta = (Output))
 	bool Result;
 };
\ No newline at end of file
