Control Rig - Fix GC object leaks when unloading plugins and packages.
Asset registry and plugin manager callbacks happen after UPackageTools::UnloadPackages which results in leaking objects if a system relies on asset registry or plugin manager callbacks to clear references.

diff --git a/Engine/Plugins/Runtime/RigVM/Source/RigVM/Private/RigVMCore/RigVMRegistry.cpp b/Engine/Plugins/Runtime/RigVM/Source/RigVM/Private/RigVMCore/RigVMRegistry.cpp
--- a/Engine/Plugins/Runtime/RigVM/Source/RigVM/Private/RigVMCore/RigVMRegistry.cpp
+++ b/Engine/Plugins/Runtime/RigVM/Source/RigVM/Private/RigVMCore/RigVMRegistry.cpp
@@ -413,6 +413,8 @@ void FRigVMRegistry_NoLock::InitializeBaseTypes_NoLock()
 			IPluginManager::Get().OnPluginUnmounted().RemoveAll(this);
 
 			UE::Anim::AttributeTypes::GetOnAttributeTypesChanged().RemoveAll(this);
+
+			FCoreDelegates::CleanupUnloadingObjects.RemoveAll(this);
 		});
 	}
 	else
@@ -749,6 +751,24 @@ bool FRigVMRegistry_NoLock::RefreshFunctionsAndDispatches_NoLock()
 	return bRegistryChanged; 
 }
 
+bool FRigVMRegistry_NoLock::OnCleanupUnloadingObjects_NoLock(const TArrayView<UObject*> InObjects)
+{
+	bool bRemoved = false;
+	for (UObject* Obj : InObjects)
+	{
+		const UObject* ObjectOuter = Obj ? Obj->GetOuter() : nullptr;
+		if (ObjectOuter && (ObjectOuter->GetClass() == UPackage::StaticClass()))
+		{
+			if (RemoveType_NoLock(FSoftObjectPath(Obj), Obj->GetClass()))
+			{
+				bRemoved = true;
+			}
+		}
+	}
+
+	return bRemoved;
+}
+
 void FRigVMRegistry_NoLock::OnAssetRenamed_NoLock(const FAssetData& InAssetData, const FString& InOldObjectPath)
 {
 	const FSoftObjectPath OldPath(InOldObjectPath);
@@ -1299,6 +1319,20 @@ FRigVMRegistry_RWLock& FRigVMRegistry_RWLock::Get()
 	return s_RigVMRegistry;
 }
 
+void FRigVMRegistry_RWLock::OnCleanupUnloadingObjects(const TArrayView<UObject*> InObjects)
+{
+	bool bAssetRemoved = false;
+	{
+		FConditionalWriteScopeLock _(*this);
+		bAssetRemoved = Super::OnCleanupUnloadingObjects_NoLock(InObjects);
+	}
+
+	if (bAssetRemoved)
+	{
+		OnRigVMRegistryChangedDelegate.Broadcast();
+	}
+}
+
 void FRigVMRegistry_RWLock::OnAssetRemoved(const FAssetData& InAssetData)
 {
 	bool bAssetRemoved = false;
@@ -2519,6 +2553,8 @@ void FRigVMRegistry_RWLock::Initialize_NoLock()
 	IPluginManager::Get().OnPluginUnmounted().AddRaw(this, &FRigVMRegistry_RWLock::OnPluginUnloaded);
 	
 	UE::Anim::AttributeTypes::GetOnAttributeTypesChanged().AddRaw(this, &FRigVMRegistry_RWLock::OnAnimationAttributeTypesChanged);
+
+	FCoreDelegates::CleanupUnloadingObjects.AddRaw(this, &FRigVMRegistry_RWLock::OnCleanupUnloadingObjects);
 }
 
 void FRigVMRegistry_RWLock::EnsureLocked(ELockType InLockType)
diff --git a/Engine/Plugins/Runtime/RigVM/Source/RigVM/Public/RigVMCore/RigVMRegistry.h b/Engine/Plugins/Runtime/RigVM/Source/RigVM/Public/RigVMCore/RigVMRegistry.h
--- a/Engine/Plugins/Runtime/RigVM/Source/RigVM/Public/RigVMCore/RigVMRegistry.h
+++ b/Engine/Plugins/Runtime/RigVM/Source/RigVM/Public/RigVMCore/RigVMRegistry.h
@@ -109,6 +109,9 @@ public:
 	// based on the names.
 	UE_API virtual void RefreshEngineTypesIfRequired_NoLock();
 	
+	// Update the registry when old types are being unloaded
+	UE_API virtual bool OnCleanupUnloadingObjects_NoLock(const TArrayView<UObject*> InObjects);
+
 	// Update the registry when types are renamed
 	UE_API virtual void OnAssetRenamed_NoLock(const FAssetData& InAssetData, const FString& InOldObjectPath);
 	
@@ -639,6 +642,9 @@ public:
 		FConditionalWriteScopeLock _(*this);
 		Super::OnAssetRenamed_NoLock(InAssetData, InOldObjectPath);
 	}
+
+	// Update the registry when old types are being unloaded
+	RIGVM_API void OnCleanupUnloadingObjects(const TArrayView<UObject*> InObjects);
 	
 	// Update the registry when old types are removed
     RIGVM_API void OnAssetRemoved(const FAssetData& InAssetData);
@@ -1002,6 +1008,13 @@ protected:
 		Super::RefreshEngineTypesIfRequired_NoLock();
 	}
 
+	// Update the registry when old types are being unloaded
+	// Note: Only call NoLock methods on the FRigVMRegistry_NoLock
+	virtual bool OnCleanupUnloadingObjects_NoLock(const TArrayView<UObject*> InObjects) override
+	{
+		return Super::OnCleanupUnloadingObjects_NoLock(InObjects);
+	}
+
 	// Update the registry when types are renamed
 	// Note: Only call NoLock methods on the FRigVMRegistry_NoLock
 	virtual void OnAssetRenamed_NoLock(const FAssetData& InAssetData, const FString& InOldObjectPath) override
diff --git a/Engine/Source/Editor/UnrealEd/Private/PackageTools.cpp b/Engine/Source/Editor/UnrealEd/Private/PackageTools.cpp
--- a/Engine/Source/Editor/UnrealEd/Private/PackageTools.cpp
+++ b/Engine/Source/Editor/UnrealEd/Private/PackageTools.cpp
@@ -495,6 +495,7 @@ UPackageTools::UPackageTools(const FObjectInitializer& ObjectInitializer)
 
 				GetObjectsWithPackage(PackageBeingUnloaded, ObjectsInPackage, true, RF_Transient, EInternalObjectFlags::Garbage);
 				// Notify any Blueprints and other systems that are about to be unloaded, also destroy any leftover worlds.
+				FCoreDelegates::CleanupUnloadingObjects.Broadcast(ObjectsInPackage);
 				for (UObject* Obj : ObjectsInPackage)
 				{
 					// Asset manager can hold hard references to this object and prevent GC
diff --git a/Engine/Source/Runtime/Core/Private/Misc/CoreDelegates.cpp b/Engine/Source/Runtime/Core/Private/Misc/CoreDelegates.cpp
--- a/Engine/Source/Runtime/Core/Private/Misc/CoreDelegates.cpp
+++ b/Engine/Source/Runtime/Core/Private/Misc/CoreDelegates.cpp
@@ -11,6 +11,7 @@
 
 TArray<TDelegate<void(void*, int32)>> FCoreDelegates::HotFixDelegates;
 TArray<TDelegate<bool(const FString&, FString&)>> FCoreDelegates::PackageNameResolvers;
+TMulticastDelegate<void(const TArrayView<UObject*>)> FCoreDelegates::CleanupUnloadingObjects;
 
 TDelegate<void(void*, int32)>& FCoreDelegates::GetHotfixDelegate(EHotfixDelegates::Type HotFix)
 {
diff --git a/Engine/Source/Runtime/Core/Public/Misc/CoreDelegates.h b/Engine/Source/Runtime/Core/Public/Misc/CoreDelegates.h
--- a/Engine/Source/Runtime/Core/Public/Misc/CoreDelegates.h
+++ b/Engine/Source/Runtime/Core/Public/Misc/CoreDelegates.h
@@ -496,6 +496,9 @@ public:
 	// Should return True of resolution occured.
 	static CORE_API TArray<TDelegate<bool(const FString&, FString&)>> PackageNameResolvers;
 
+	/** Callback to allow for rare specialized situations where manual clearing of references is necessary */
+	static CORE_API TMulticastDelegate<void(const TArrayView<UObject*>)> CleanupUnloadingObjects;
+
 	// Called to request that systems free whatever memory they are able to. Called early in LoadMap.
 	// Caller is responsible for flushing rendering etc. See UEngine::TrimMemory
 	static CORE_API FSimpleMulticastDelegate& GetMemoryTrimDelegate();
