RigVM: Fix invalid deserialized templating types

#rnx

diff --git a/Engine/Plugins/Runtime/RigVM/Source/RigVM/Public/RigVMTypeUtils.h b/Engine/Plugins/Runtime/RigVM/Source/RigVM/Public/RigVMTypeUtils.h
--- a/Engine/Plugins/Runtime/RigVM/Source/RigVM/Public/RigVMTypeUtils.h
+++ b/Engine/Plugins/Runtime/RigVM/Source/RigVM/Public/RigVMTypeUtils.h
@@ -463,6 +463,15 @@ namespace RigVMTypeUtils
 			// Uses redirectors and updates the CPPType if necessary
 			ObjectFromCPPType(CPPType, true, InResolvalInfo);
 		}
+
+		// make sure to clean up spacing errors with templated types
+		if (CPPType.Contains(TEXT("<")) || CPPType.Contains(TEXT(">")))
+		{
+			CPPType = CPPType.Replace(TEXT("< "), TEXT("<"));
+			CPPType = CPPType.Replace(TEXT(" >"), TEXT(">"));
+			CPPType = CPPType.Replace(TEXT(" <"), TEXT("<"));
+			CPPType = CPPType.Replace(TEXT("> "), TEXT(">"));
+		}
 	
 		return CPPType;
 	}
