RigVM: Fix invalid deserialized templating types

#rnx

diff --git a/Engine/Plugins/Animation/ControlRig/Source/ControlRig/Private/AnimNode_ControlRig.cpp b/Engine/Plugins/Animation/ControlRig/Source/ControlRig/Private/AnimNode_ControlRig.cpp
--- a/Engine/Plugins/Animation/ControlRig/Source/ControlRig/Private/AnimNode_ControlRig.cpp
+++ b/Engine/Plugins/Animation/ControlRig/Source/ControlRig/Private/AnimNode_ControlRig.cpp
@@ -325,11 +325,10 @@ bool FAnimNode_ControlRig::UpdateControlRigIfNeeded(const UAnimInstance* InAnimI
 
 							if(!NewProperty->SameType(OldProperty))
 							{
-								FString OldExtendedCPPType, NewExtendedCPPType;
-								const FString OldCPPType = OldProperty->GetCPPType(&OldExtendedCPPType);
-								const FString NewCPPType = NewProperty->GetCPPType(&NewExtendedCPPType);
+								const FString OldCPPType = RigVMTypeUtils::GetCPPTypeFromProperty(OldProperty);
+								const FString NewCPPType = RigVMTypeUtils::GetCPPTypeFromProperty(NewProperty);
 								static constexpr TCHAR Format[] = TEXT("Property / Variable '%s' has incorrect type (is '%s', expected '%s').");
-								return ReportErrorAndSwitchToDefaultRig(FString::Printf(Format, *NewProperty->GetName(), *(NewCPPType + NewExtendedCPPType), *(OldCPPType + OldExtendedCPPType)));
+								return ReportErrorAndSwitchToDefaultRig(FString::Printf(Format, *NewProperty->GetName(), *NewCPPType, *OldCPPType));
 							}
 						}
 						
diff --git a/Engine/Plugins/Animation/ControlRig/Source/ControlRig/Private/Units/Core/RigUnit_UserData.cpp b/Engine/Plugins/Animation/ControlRig/Source/ControlRig/Private/Units/Core/RigUnit_UserData.cpp
--- a/Engine/Plugins/Animation/ControlRig/Source/ControlRig/Private/Units/Core/RigUnit_UserData.cpp
+++ b/Engine/Plugins/Animation/ControlRig/Source/ControlRig/Private/Units/Core/RigUnit_UserData.cpp
@@ -126,10 +126,9 @@ void FRigDispatch_GetUserData::Execute(FRigVMExtendedExecuteContext& InContext,
 			if(UserData->GetProperty())
 			{
 #if WITH_EDITOR
-				FString ExtendedCPPType;
-				const FString CPPType = PropertyResult->GetCPPType(&ExtendedCPPType);
+				const FString CPPType = RigVMTypeUtils::GetCPPTypeFromProperty(PropertyResult);
 				const FRigVMRegistry& Registry = FRigVMRegistry::Get();
-				const TRigVMTypeIndex ResultTypeIndex = Registry.GetTypeIndexFromCPPType(CPPType + ExtendedCPPType);
+				const TRigVMTypeIndex ResultTypeIndex = Registry.GetTypeIndexFromCPPType(CPPType);
 				const TRigVMTypeIndex UserDataTypeIndex = Registry.GetTypeIndexFromCPPType(UserData->GetCPPType().ToString());
 
 				if(Registry.CanMatchTypes(ResultTypeIndex, UserDataTypeIndex, true))
@@ -142,7 +141,7 @@ void FRigDispatch_GetUserData::Execute(FRigVMExtendedExecuteContext& InContext,
 
 #if WITH_EDITOR
 				static constexpr TCHAR Format[] = TEXT("User data of type '%s' not compatible with pin of type '%s'.");
-				ControlRigContext.Logf(EMessageSeverity::Info, Format, *UserData->GetCPPType().ToString(), *(CPPType + ExtendedCPPType));
+				ControlRigContext.Logf(EMessageSeverity::Info, Format, *UserData->GetCPPType().ToString(), *CPPType);
 #endif
 			}
 		}
diff --git a/Engine/Plugins/Runtime/RigVM/Source/RigVM/Private/RigVMCore/RigVMAssetUserData.cpp b/Engine/Plugins/Runtime/RigVM/Source/RigVM/Private/RigVMCore/RigVMAssetUserData.cpp
--- a/Engine/Plugins/Runtime/RigVM/Source/RigVM/Private/RigVMCore/RigVMAssetUserData.cpp
+++ b/Engine/Plugins/Runtime/RigVM/Source/RigVM/Private/RigVMCore/RigVMAssetUserData.cpp
@@ -1,6 +1,7 @@
 // Copyright Epic Games, Inc. All Rights Reserved.
 
 #include "RigVMCore/RigVMAssetUserData.h"
+#include "RigVMTypeUtils.h"
 #include "StructUtils/UserDefinedStruct.h"
 #include "Misc/PackageName.h"
 
@@ -360,10 +361,7 @@ bool UNameSpacedUserData::IsPropertySupported(const FProperty* InProperty, const
 	{
 		if(OutErrorMessage && OutErrorMessage->IsEmpty())
 		{
-			FString ExtendedCppType;
-			FString CPPTypeString = InProperty->GetCPPType(&ExtendedCppType);
-			CPPTypeString += ExtendedCppType;
-
+			const FString CPPTypeString = RigVMTypeUtils::GetCPPTypeFromProperty(InProperty);
 			static constexpr TCHAR UnsupportedCPPTypeFormat[] = TEXT("Unsupported user data type '%s' at path '%s'.");
 			(*OutErrorMessage) = FString::Printf(UnsupportedCPPTypeFormat, *CPPTypeString, *InPath);
 		}
diff --git a/Engine/Plugins/Runtime/RigVM/Source/RigVM/Private/RigVMCore/RigVMMemoryStorage.cpp b/Engine/Plugins/Runtime/RigVM/Source/RigVM/Private/RigVMCore/RigVMMemoryStorage.cpp
--- a/Engine/Plugins/Runtime/RigVM/Source/RigVM/Private/RigVMCore/RigVMMemoryStorage.cpp
+++ b/Engine/Plugins/Runtime/RigVM/Source/RigVM/Private/RigVMCore/RigVMMemoryStorage.cpp
@@ -154,9 +154,7 @@ FRigVMPropertyDescription::FRigVMPropertyDescription(const FProperty* InProperty
 
 	if(CPPType.IsEmpty() && Property != nullptr)
 	{
-		FString ExtendedType;
-		CPPType = Property->GetCPPType(&ExtendedType);
-		CPPType += ExtendedType;
+		CPPType = RigVMTypeUtils::GetCPPTypeFromProperty(Property);
 
 		if(CPPType == RigVMTypeUtils::UInt8Type)
 		{
@@ -1112,11 +1110,8 @@ bool URigVMMemoryStorage::CopyProperty(
 			// Only log the issue once, rather than spam.
 			static TSet<FString> ReportedErrors;
 			
-			FString TargetType, SourceType, TargetExtendedType, SourceExtendedType;
-			TargetType = InTargetProperty->GetCPPType(&TargetExtendedType);
-			SourceType = InSourceProperty->GetCPPType(&SourceExtendedType);
-			TargetType += TargetExtendedType;
-			SourceType += SourceExtendedType;
+			const FString TargetType = RigVMTypeUtils::GetCPPTypeFromProperty(InTargetProperty);
+			const FString SourceType = RigVMTypeUtils::GetCPPTypeFromProperty(InSourceProperty);
 			
 			UPackage* Package = InTargetProperty->GetOutermost();
 			
diff --git a/Engine/Plugins/Runtime/RigVM/Source/RigVM/Private/RigVMCore/RigVMMemoryStorageStruct.cpp b/Engine/Plugins/Runtime/RigVM/Source/RigVM/Private/RigVMCore/RigVMMemoryStorageStruct.cpp
--- a/Engine/Plugins/Runtime/RigVM/Source/RigVM/Private/RigVMCore/RigVMMemoryStorageStruct.cpp
+++ b/Engine/Plugins/Runtime/RigVM/Source/RigVM/Private/RigVMCore/RigVMMemoryStorageStruct.cpp
@@ -387,11 +387,8 @@ bool FRigVMMemoryStorageStruct::CopyProperty(
 			// Only log the issue once, rather than spam.
 			static TSet<FString> ReportedErrors;
 
-			FString TargetType, SourceType, TargetExtendedType, SourceExtendedType;
-			TargetType = InTargetProperty->GetCPPType(&TargetExtendedType);
-			SourceType = InSourceProperty->GetCPPType(&SourceExtendedType);
-			TargetType += TargetExtendedType;
-			SourceType += SourceExtendedType;
+			const FString TargetType = RigVMTypeUtils::GetCPPTypeFromProperty(InTargetProperty);
+			const FString SourceType = RigVMTypeUtils::GetCPPTypeFromProperty(InSourceProperty);
 
 			UPackage* Package = InTargetProperty->GetOutermost();
 
diff --git a/Engine/Plugins/Runtime/RigVM/Source/RigVM/Public/RigVMCore/RigVMAssetUserData.h b/Engine/Plugins/Runtime/RigVM/Source/RigVM/Public/RigVMCore/RigVMAssetUserData.h
--- a/Engine/Plugins/Runtime/RigVM/Source/RigVM/Public/RigVMCore/RigVMAssetUserData.h
+++ b/Engine/Plugins/Runtime/RigVM/Source/RigVM/Public/RigVMCore/RigVMAssetUserData.h
@@ -3,6 +3,7 @@
 #pragma once
 
 #include "CoreMinimal.h"
+#include "RigVMTypeUtils.h"
 #include "Engine/AssetUserData.h"
 #include "Engine/DataAsset.h"
 #include "Misc/TransactionallySafeCriticalSection.h"
@@ -75,10 +76,7 @@ public:
 		{
 			if(CPPType.IsNone() && Property != nullptr)
 			{
-				FString ExtendedCppType;
-  				FString CPPTypeString = Property->GetCPPType(&ExtendedCppType);
-  				CPPTypeString += ExtendedCppType;
-				CPPType = *CPPTypeString;
+				CPPType = *RigVMTypeUtils::GetCPPTypeFromProperty(Property);
 			}
 			return CPPType;
 		}
diff --git a/Engine/Plugins/Runtime/RigVM/Source/RigVM/Public/RigVMTypeUtils.h b/Engine/Plugins/Runtime/RigVM/Source/RigVM/Public/RigVMTypeUtils.h
--- a/Engine/Plugins/Runtime/RigVM/Source/RigVM/Public/RigVMTypeUtils.h
+++ b/Engine/Plugins/Runtime/RigVM/Source/RigVM/Public/RigVMTypeUtils.h
@@ -441,9 +441,34 @@ namespace RigVMTypeUtils
 		return CPPTypeObject;
 	}
 
+	static void CleanupCPPType(FString& CPPType)
+	{
+		// make sure to clean up spacing errors with templated types
+		if (CPPType.Contains(TEXT("<")) || CPPType.Contains(TEXT(">")))
+		{
+			CPPType.ReplaceInline(TEXT("< "), TEXT("<"));
+			CPPType.ReplaceInline(TEXT(" >"), TEXT(">"));
+			CPPType.ReplaceInline(TEXT(" <"), TEXT("<"));
+			CPPType.ReplaceInline(TEXT("> "), TEXT(">"));
+		}
+	}
+
+	static FString GetCPPTypeFromProperty(const FProperty* InProperty)
+	{
+		check(InProperty);
+		FString ExtendedType;
+		FString CPPType = InProperty->GetCPPType(&ExtendedType);
+		CPPType += ExtendedType;
+
+		CleanupCPPType(CPPType);
+		return CPPType;
+	}
+
 	static FString PostProcessCPPType(const FString& InCPPType, UObject* InCPPTypeObject = nullptr, const FRigVMUserDefinedTypeResolver* InResolvalInfo = nullptr)
 	{
 		FString CPPType = InCPPType;
+		CleanupCPPType(CPPType);
+		
 		if (InCPPTypeObject)
 		{
 			const bool bIsClass = CPPType.StartsWith(TSubclassOfPrefix);
@@ -463,15 +488,6 @@ namespace RigVMTypeUtils
 			// Uses redirectors and updates the CPPType if necessary
 			ObjectFromCPPType(CPPType, true, InResolvalInfo);
 		}
-
-		// make sure to clean up spacing errors with templated types
-		if (CPPType.Contains(TEXT("<")) || CPPType.Contains(TEXT(">")))
-		{
-			CPPType = CPPType.Replace(TEXT("< "), TEXT("<"));
-			CPPType = CPPType.Replace(TEXT(" >"), TEXT(">"));
-			CPPType = CPPType.Replace(TEXT(" <"), TEXT("<"));
-			CPPType = CPPType.Replace(TEXT("> "), TEXT(">"));
-		}
 	
 		return CPPType;
 	}
diff --git a/Engine/Plugins/Runtime/RigVM/Source/RigVMDeveloper/Private/RigVMCompiler/RigVMCodeGenerator.cpp b/Engine/Plugins/Runtime/RigVM/Source/RigVMDeveloper/Private/RigVMCompiler/RigVMCodeGenerator.cpp
--- a/Engine/Plugins/Runtime/RigVM/Source/RigVMDeveloper/Private/RigVMCompiler/RigVMCodeGenerator.cpp
+++ b/Engine/Plugins/Runtime/RigVM/Source/RigVMDeveloper/Private/RigVMCompiler/RigVMCodeGenerator.cpp
@@ -1751,9 +1751,7 @@ FString FRigVMCodeGenerator::GetOperandName(const FRigVMExtendedExecuteContext&
 				Left = Left.TrimChar(TEXT('['));
 				Left = Left.TrimChar(TEXT(']'));
 
-				FString ExtendedType;
-				FString CPPType = CurrentProperty->GetCPPType(&ExtendedType);
-				CPPType += ExtendedType;
+				FString CPPType = RigVMTypeUtils::GetCPPTypeFromProperty(CurrentProperty);
 				if(CPPType == RigVMTypeUtils::UInt8Type && CastField<FBoolProperty>(CurrentProperty))
 				{
 					CPPType = RigVMTypeUtils::BoolType;
@@ -1791,9 +1789,7 @@ FString FRigVMCodeGenerator::GetOperandCPPType(const FRigVMExtendedExecuteContex
 		const FProperty* TailProperty = Path.GetTailProperty();
 		check(TailProperty);
 
-		FString ExtendedType;
-		const FString CPPType = TailProperty->GetCPPType(&ExtendedType);
-		return CPPType + ExtendedType;
+		return RigVMTypeUtils::GetCPPTypeFromProperty(TailProperty);
 	}
 	
 	if (Property.IsValid())
diff --git a/Engine/Plugins/Runtime/RigVM/Source/RigVMDeveloper/Private/RigVMModel/RigVMController.cpp b/Engine/Plugins/Runtime/RigVM/Source/RigVMDeveloper/Private/RigVMModel/RigVMController.cpp
--- a/Engine/Plugins/Runtime/RigVM/Source/RigVMDeveloper/Private/RigVMModel/RigVMController.cpp
+++ b/Engine/Plugins/Runtime/RigVM/Source/RigVMDeveloper/Private/RigVMModel/RigVMController.cpp
@@ -592,9 +592,7 @@ FRigVMPinInfo::FRigVMPinInfo(FProperty* InProperty, ERigVMPinDirection InDirecti
 		}
 	}
 
-	FString ExtendedCppType;
-	FString CPPType = InProperty->GetCPPType(&ExtendedCppType);
-	CPPType += ExtendedCppType;
+	FString CPPType = RigVMTypeUtils::GetCPPTypeFromProperty(InProperty);
 	CPPType = RigVMTypeUtils::PostProcessCPPType(CPPType, CPPTypeObject);
 	TypeIndex = FRigVMRegistry::Get().GetTypeIndexFromCPPType(CPPType);
 	CorrectExecuteTypeIndex();
@@ -17897,9 +17895,7 @@ void URigVMController::ConfigurePinFromProperty(FProperty* InProperty, URigVMPin
 
 #endif
 
-	FString ExtendedCppType;
-	InOutPin->CPPType = InProperty->GetCPPType(&ExtendedCppType);
-	InOutPin->CPPType += ExtendedCppType;
+	InOutPin->CPPType = RigVMTypeUtils::GetCPPTypeFromProperty(InProperty);
 
 	InOutPin->bIsDynamicArray = false;
 	FProperty* PropertyForType = InProperty;
