[CoreGlobals]
- Add IsInActualLoadingThread to know when we're in the loader, but not in the game thread
- Refactor the code that needed that logic

diff --git a/Engine/Plugins/Experimental/DaySequence/Source/DaySequenceEditor/Private/DaySequenceActorPreview.cpp b/Engine/Plugins/Experimental/DaySequence/Source/DaySequenceEditor/Private/DaySequenceActorPreview.cpp
--- a/Engine/Plugins/Experimental/DaySequence/Source/DaySequenceEditor/Private/DaySequenceActorPreview.cpp
+++ b/Engine/Plugins/Experimental/DaySequence/Source/DaySequenceEditor/Private/DaySequenceActorPreview.cpp
@@ -259,7 +259,7 @@ void FDaySequenceActorPreview::ClosePreviewToolkit()
 
 void FDaySequenceActorPreview::OnBlueprintPreCompile(UBlueprint* Blueprint)
 {
-	if (IsInAsyncLoadingThread() && !IsInGameThread())
+	if (IsInActualLoadingThread())
 	{
 		return;
 	}
diff --git a/Engine/Plugins/Runtime/RigVM/Source/RigVMEditor/Private/Editor/Kismet/RigVMBlueprintCompilationManager.cpp b/Engine/Plugins/Runtime/RigVM/Source/RigVMEditor/Private/Editor/Kismet/RigVMBlueprintCompilationManager.cpp
--- a/Engine/Plugins/Runtime/RigVM/Source/RigVMEditor/Private/Editor/Kismet/RigVMBlueprintCompilationManager.cpp
+++ b/Engine/Plugins/Runtime/RigVM/Source/RigVMEditor/Private/Editor/Kismet/RigVMBlueprintCompilationManager.cpp
@@ -1984,8 +1984,7 @@ void FRigVMBlueprintCompilationManagerImpl::FlushReinstancingQueueImpl(bool bFin
 		ClassHashes.Reset();
 
 		// Special case when we run on ALT, we want to cleanup all classes flagged for reinstantiation right away.
-		const bool bIsInActualAsyncLoadingThread = IsInAsyncLoadingThread() && !IsInGameThread();
-		if (IsAsyncLoading() && (!IsAsyncLoadingMultithreaded() || !bIsInActualAsyncLoadingThread))
+		if (IsAsyncLoading() && (!IsAsyncLoadingMultithreaded() || !IsInActualLoadingThread()))
 		{
 			// While async loading we only remove classes that have no instances being
 			// async loaded. Those instances will need to be reinstanced once they finish
diff --git a/Engine/Source/Editor/Kismet/Private/BlueprintCompilationManager.cpp b/Engine/Source/Editor/Kismet/Private/BlueprintCompilationManager.cpp
--- a/Engine/Source/Editor/Kismet/Private/BlueprintCompilationManager.cpp
+++ b/Engine/Source/Editor/Kismet/Private/BlueprintCompilationManager.cpp
@@ -2077,8 +2077,7 @@ void FBlueprintCompilationManagerImpl::FlushReinstancingQueueImpl(bool bFindAndR
 		}
 
 		// Special case when we run on ALT, we want to cleanup all classes flagged for reinstantiation right away.
-		const bool bIsInActualAsyncLoadingThread = IsInAsyncLoadingThread() && !IsInGameThread();
-		if (IsAsyncLoading() && (!IsAsyncLoadingMultithreaded() || !bIsInActualAsyncLoadingThread))
+		if (IsAsyncLoading() && (!IsAsyncLoadingMultithreaded() || !IsInActualLoadingThread()))
 		{
 			// While async loading we only remove classes that have no instances being
 			// async loaded. Those instances will need to be reinstanced once they finish
diff --git a/Engine/Source/Runtime/Core/Private/HAL/ThreadingBase.cpp b/Engine/Source/Runtime/Core/Private/HAL/ThreadingBase.cpp
--- a/Engine/Source/Runtime/Core/Private/HAL/ThreadingBase.cpp
+++ b/Engine/Source/Runtime/Core/Private/HAL/ThreadingBase.cpp
@@ -352,6 +352,11 @@ static bool IsInAsyncLoadingThreadCoreInternal()
 }
 bool(*IsInAsyncLoadingThread)() = &IsInAsyncLoadingThreadCoreInternal;
 
+CORE_API bool IsInActualLoadingThread()
+{
+	return IsInAsyncLoadingThread() && !IsInGameThread();
+}
+
 /**
  * Fake thread created when multi-threading is disabled.
  */
diff --git a/Engine/Source/Runtime/Core/Public/CoreGlobals.h b/Engine/Source/Runtime/Core/Public/CoreGlobals.h
--- a/Engine/Source/Runtime/Core/Public/CoreGlobals.h
+++ b/Engine/Source/Runtime/Core/Public/CoreGlobals.h
@@ -754,6 +754,9 @@ extern CORE_API bool IsInActualRenderingThread();
 /** @return True if called from the async loading thread if it's enabled, otherwise if called from game thread while is async loading code. */
 extern CORE_API bool (*IsInAsyncLoadingThread)();
 
+/** @return True only if called from the async loading thread if it's enabled. */
+extern CORE_API bool IsInActualLoadingThread();
+
 /** Thread used for rendering */
 UE_DEPRECATED(4.26, "Please use `GIsThreadedRendering` or `IsInActualRenderingThread()`")
 extern CORE_API FRunnableThread* GRenderingThread;
