Output failure reason from GetSchema()->CanAddLink for URigVMController::AddLink

diff --git a/Engine/Plugins/Runtime/RigVM/Source/RigVMDeveloper/Private/RigVMModel/RigVMController.cpp b/Engine/Plugins/Runtime/RigVM/Source/RigVMDeveloper/Private/RigVMModel/RigVMController.cpp
--- a/Engine/Plugins/Runtime/RigVM/Source/RigVMDeveloper/Private/RigVMModel/RigVMController.cpp
+++ b/Engine/Plugins/Runtime/RigVM/Source/RigVMDeveloper/Private/RigVMModel/RigVMController.cpp
@@ -12396,6 +12396,10 @@ bool URigVMController::AddLink(URigVMPin* OutputPin, URigVMPin* InputPin, bool b
 
 	if(!GetSchema()->CanAddLink(this, OutputPin, InputPin, nullptr, UserLinkDirection, true, true, OutFailureReason))
 	{
+		if(OutFailureReason)
+		{
+			ReportError(*OutFailureReason);
+		}
 		return false;
 	}
 	
