RigVM: Fix sub pin copy folding for dispatches with property access

#rnx

diff --git a/Engine/Plugins/Runtime/RigVM/Source/RigVM/Private/RigVMCore/RigVMTemplate.cpp b/Engine/Plugins/Runtime/RigVM/Source/RigVM/Private/RigVMCore/RigVMTemplate.cpp
--- a/Engine/Plugins/Runtime/RigVM/Source/RigVM/Private/RigVMCore/RigVMTemplate.cpp
+++ b/Engine/Plugins/Runtime/RigVM/Source/RigVM/Private/RigVMCore/RigVMTemplate.cpp
@@ -61,9 +61,7 @@ FRigVMTemplateArgument::FRigVMTemplateArgument(FProperty* InProperty, FRigVMRegi
 	Direction = FRigVMStruct::GetPinDirectionFromProperty(InProperty);
 #endif
 
-	FString ExtendedType;
-	const FString CPPType = InProperty->GetCPPType(&ExtendedType);
-	const FName CPPTypeName = *(CPPType + ExtendedType);
+	const FName CPPTypeName = *RigVMTypeUtils::GetCPPTypeFromProperty(InProperty);
 	UObject* CPPTypeObject = nullptr;
 
 	if (FArrayProperty* ArrayProperty = CastField<FArrayProperty>(InProperty))
@@ -1720,7 +1718,10 @@ const FRigVMFunction* FRigVMTemplate::GetPermutation_NoLock(int32 InIndex, const
 	if(InRegistry->GetFunctions_NoLock().IsValidIndex(FunctionIndex))
 	{
 		const FRigVMFunction* Function = &InRegistry->GetFunctions_NoLock()[Permutations[InIndex]];
-		check(Function->PermutationIndex == InIndex);
+		if (Function)
+		{
+			check(Function->PermutationIndex == InIndex);
+		}
 		return Function;
 	}
 	return nullptr;
diff --git a/Engine/Plugins/Runtime/RigVM/Source/RigVM/Private/RigVMFunctions/RigVMDispatch_Array.cpp b/Engine/Plugins/Runtime/RigVM/Source/RigVM/Private/RigVMFunctions/RigVMDispatch_Array.cpp
--- a/Engine/Plugins/Runtime/RigVM/Source/RigVM/Private/RigVMFunctions/RigVMDispatch_Array.cpp
+++ b/Engine/Plugins/Runtime/RigVM/Source/RigVM/Private/RigVMFunctions/RigVMDispatch_Array.cpp
@@ -286,7 +286,7 @@ void FRigVMDispatch_ArrayMake::Execute(FRigVMExtendedExecuteContext& InContext,
 		{
 			const uint8* SourceMemory = Handles[Index].GetInputData();
 			uint8* TargetMemory = ArrayHelper.GetRawPtr(Index);
-			URigVMMemoryStorage::CopyProperty(ArrayProperty->Inner, TargetMemory, Handles[Index].GetProperty(), SourceMemory);
+			URigVMMemoryStorage::CopyProperty(ArrayProperty->Inner, TargetMemory, Handles[Index].GetResolvedProperty(), SourceMemory);
 		}
 	}
 	else
@@ -378,7 +378,7 @@ FRigVMTemplateTypeMap FRigVMDispatch_ArrayGetNum::OnNewArgumentType(const FName&
 
 void FRigVMDispatch_ArrayGetNum::Execute(FRigVMExtendedExecuteContext& InContext, FRigVMMemoryHandleArray Handles, FRigVMPredicateBranchArray Predicates)
 {
-	const FArrayProperty* ArrayProperty = CastFieldChecked<FArrayProperty>(Handles[0].GetProperty());
+	const FArrayProperty* ArrayProperty = CastFieldChecked<FArrayProperty>(Handles[0].GetResolvedProperty());
 	const FScriptArrayHelper ArrayHelper(ArrayProperty, Handles[0].GetInputData());
 	int32* Num = (int32*)Handles[1].GetOutputData();
 	*Num = ArrayHelper.Num();
@@ -509,7 +509,7 @@ void FRigVMDispatch_ArrayInit::Execute(FRigVMExtendedExecuteContext& InContext,
 	const FArrayProperty* ArrayProperty = CastFieldChecked<FArrayProperty>(Handles[0].GetProperty());
 	FScriptArrayHelper ArrayHelper(ArrayProperty, Handles[0].GetOutputData());
 	const int32* Num = (const int32*)Handles[1].GetInputData();
-	const FProperty* ValueProperty = Handles[2].GetProperty();
+	const FProperty* ValueProperty = Handles[2].GetResolvedProperty();
 	const uint8* ValueData = Handles[2].GetInputData();
 
 	if(InContext.IsValidArraySize(*Num))
@@ -594,7 +594,7 @@ FRigVMTemplateTypeMap FRigVMDispatch_ArrayGetAtIndex::OnNewArgumentType(const FN
 
 void FRigVMDispatch_ArrayGetAtIndex::Execute(FRigVMExtendedExecuteContext& InContext, FRigVMMemoryHandleArray Handles, FRigVMPredicateBranchArray Predicates)
 {
-	const FArrayProperty* ArrayProperty = CastFieldChecked<FArrayProperty>(Handles[0].GetProperty());
+	const FArrayProperty* ArrayProperty = CastFieldChecked<FArrayProperty>(Handles[0].GetResolvedProperty());
 	FScriptArrayHelper ArrayHelper(ArrayProperty, Handles[0].GetInputData());
 	const int32 Index = *(const int32*)Handles[1].GetInputData();
 	int32 OutIndex = Index;
@@ -688,7 +688,7 @@ void FRigVMDispatch_ArraySetAtIndex::Execute(FRigVMExtendedExecuteContext& InCon
 	if(InContext.IsValidArrayIndex(OutIndex, ArrayHelper))
 	{
 		uint8* TargetMemory = ArrayHelper.GetRawPtr(OutIndex);
-		URigVMMemoryStorage::CopyProperty(ArrayProperty->Inner, TargetMemory, Handles[2].GetProperty(), SourceMemory);
+		URigVMMemoryStorage::CopyProperty(ArrayProperty->Inner, TargetMemory, Handles[2].GetResolvedProperty(), SourceMemory);
 	}
 }
 
@@ -753,7 +753,7 @@ void FRigVMDispatch_ArrayAdd::Execute(FRigVMExtendedExecuteContext& InContext, F
 	{
 		Index = ArrayHelper.AddValue();
 		uint8* TargetMemory = ArrayHelper.GetRawPtr(Index);
-		URigVMMemoryStorage::CopyProperty(ArrayProperty->Inner, TargetMemory, Handles[1].GetProperty(), SourceMemory);
+		URigVMMemoryStorage::CopyProperty(ArrayProperty->Inner, TargetMemory, Handles[1].GetResolvedProperty(), SourceMemory);
 	}
 	else
 	{
@@ -808,7 +808,7 @@ void FRigVMDispatch_ArrayInsert::Execute(FRigVMExtendedExecuteContext& InContext
 		ArrayHelper.InsertValues(Index, 1);
 
 		uint8* TargetMemory = ArrayHelper.GetRawPtr(Index);
-		URigVMMemoryStorage::CopyProperty(ArrayProperty->Inner, TargetMemory, Handles[2].GetProperty(), SourceMemory);
+		URigVMMemoryStorage::CopyProperty(ArrayProperty->Inner, TargetMemory, Handles[2].GetResolvedProperty(), SourceMemory);
 	}
 }
 
@@ -985,7 +985,7 @@ FRigVMTemplateTypeMap FRigVMDispatch_ArrayFind::OnNewArgumentType(const FName& I
 
 void FRigVMDispatch_ArrayFind::Execute(FRigVMExtendedExecuteContext& InContext, FRigVMMemoryHandleArray Handles, FRigVMPredicateBranchArray Predicates)
 {
-	const FArrayProperty* ArrayProperty = CastFieldChecked<FArrayProperty>(Handles[0].GetProperty());
+	const FArrayProperty* ArrayProperty = CastFieldChecked<FArrayProperty>(Handles[0].GetResolvedProperty());
 	FScriptArrayHelper ArrayHelper(ArrayProperty, Handles[0].GetInputData());
 	int32& Index = *(int32*)Handles[2].GetOutputData();
 	bool& bSuccess = *(bool*)Handles[3].GetOutputData();
@@ -993,7 +993,7 @@ void FRigVMDispatch_ArrayFind::Execute(FRigVMExtendedExecuteContext& InContext,
 	Index = INDEX_NONE;
 	bSuccess = false;
 
-	const FProperty* PropertyA = Handles[1].GetProperty();
+	const FProperty* PropertyA = Handles[1].GetResolvedProperty();
 	const FProperty* PropertyB = ArrayProperty->Inner;
 
 	if(PropertyA->SameType(PropertyB))
@@ -1083,7 +1083,7 @@ FRigVMTemplateTypeMap FRigVMDispatch_ArrayAppend::OnNewArgumentType(const FName&
 
 void FRigVMDispatch_ArrayAppend::Execute(FRigVMExtendedExecuteContext& InContext, FRigVMMemoryHandleArray Handles, FRigVMPredicateBranchArray Predicates)
 {
-	const FArrayProperty* SourceArrayProperty = CastFieldChecked<FArrayProperty>(Handles[1].GetProperty());
+	const FArrayProperty* SourceArrayProperty = CastFieldChecked<FArrayProperty>(Handles[1].GetResolvedProperty());
 	FScriptArrayHelper SourceArrayHelper(SourceArrayProperty, Handles[1].GetInputData());
 
 	if(SourceArrayHelper.Num() > 0)
@@ -1164,7 +1164,7 @@ FRigVMTemplateTypeMap FRigVMDispatch_ArrayClone::OnNewArgumentType(const FName&
 
 void FRigVMDispatch_ArrayClone::Execute(FRigVMExtendedExecuteContext& InContext, FRigVMMemoryHandleArray Handles, FRigVMPredicateBranchArray Predicates)
 {
-	const FArrayProperty* ArrayProperty = CastFieldChecked<FArrayProperty>(Handles[0].GetProperty());
+	const FArrayProperty* ArrayProperty = CastFieldChecked<FArrayProperty>(Handles[0].GetResolvedProperty());
 	FScriptArrayHelper ArrayHelper(ArrayProperty, Handles[0].GetInputData());
 	const FArrayProperty* CloneProperty = CastFieldChecked<FArrayProperty>(Handles[1].GetProperty());
 	FScriptArrayHelper CloneHelper(CloneProperty, Handles[1].GetOutputData());
@@ -1209,7 +1209,7 @@ FText FRigVMDispatch_ArrayClone::GetArgumentTooltip(const FName& InArgumentName,
 
 void FRigVMDispatch_ArrayUnion::Execute(FRigVMExtendedExecuteContext& InContext, FRigVMMemoryHandleArray Handles, FRigVMPredicateBranchArray Predicates)
 {
-	const FArrayProperty* OtherProperty = CastFieldChecked<FArrayProperty>(Handles[1].GetProperty());
+	const FArrayProperty* OtherProperty = CastFieldChecked<FArrayProperty>(Handles[1].GetResolvedProperty());
 	FScriptArrayHelper OtherHelper(OtherProperty, Handles[1].GetInputData());
 
 	if(OtherHelper.Num() > 0)
@@ -1351,9 +1351,9 @@ FRigVMTemplateTypeMap FRigVMDispatch_ArrayDifference::OnNewArgumentType(const FN
 
 void FRigVMDispatch_ArrayDifference::Execute(FRigVMExtendedExecuteContext& InContext, FRigVMMemoryHandleArray Handles, FRigVMPredicateBranchArray Predicates)
 {
-	const FArrayProperty* ArrayProperty = CastFieldChecked<FArrayProperty>(Handles[0].GetProperty());
+	const FArrayProperty* ArrayProperty = CastFieldChecked<FArrayProperty>(Handles[0].GetResolvedProperty());
 	FScriptArrayHelper ArrayHelper(ArrayProperty, Handles[0].GetInputData());
-	const FArrayProperty* OtherProperty = CastFieldChecked<FArrayProperty>(Handles[1].GetProperty());
+	const FArrayProperty* OtherProperty = CastFieldChecked<FArrayProperty>(Handles[1].GetResolvedProperty());
 	FScriptArrayHelper OtherHelper(OtherProperty, Handles[1].GetInputData());
 	const FArrayProperty* ResultProperty = CastFieldChecked<FArrayProperty>(Handles[2].GetProperty());
 	FScriptArrayHelper ResultHelper(ResultProperty, Handles[2].GetOutputData());
@@ -1425,9 +1425,9 @@ FText FRigVMDispatch_ArrayDifference::GetArgumentTooltip(const FName& InArgument
 
 void FRigVMDispatch_ArrayIntersection::Execute(FRigVMExtendedExecuteContext& InContext, FRigVMMemoryHandleArray Handles, FRigVMPredicateBranchArray Predicates)
 {
-	const FArrayProperty* ArrayProperty = CastFieldChecked<FArrayProperty>(Handles[0].GetProperty());
+	const FArrayProperty* ArrayProperty = CastFieldChecked<FArrayProperty>(Handles[0].GetResolvedProperty());
 	FScriptArrayHelper ArrayHelper(ArrayProperty, Handles[0].GetInputData());
-	const FArrayProperty* OtherProperty = CastFieldChecked<FArrayProperty>(Handles[1].GetProperty());
+	const FArrayProperty* OtherProperty = CastFieldChecked<FArrayProperty>(Handles[1].GetResolvedProperty());
 	FScriptArrayHelper OtherHelper(OtherProperty, Handles[1].GetInputData());
 	const FArrayProperty* ResultProperty = CastFieldChecked<FArrayProperty>(Handles[2].GetProperty());
 	FScriptArrayHelper ResultHelper(ResultProperty, Handles[2].GetOutputData());
@@ -1563,7 +1563,7 @@ FRigVMTemplateTypeMap FRigVMDispatch_ArrayIterator::OnNewArgumentType(const FNam
 
 void FRigVMDispatch_ArrayIterator::Execute(FRigVMExtendedExecuteContext& InContext, FRigVMMemoryHandleArray Handles, FRigVMPredicateBranchArray Predicates)
 {
-	const FArrayProperty* ArrayProperty = CastFieldChecked<FArrayProperty>(Handles[0].GetProperty());
+	const FArrayProperty* ArrayProperty = CastFieldChecked<FArrayProperty>(Handles[0].GetResolvedProperty());
 	FScriptArrayHelper ArrayHelper(ArrayProperty, Handles[0].GetInputData());
 	int32& Index = *(int32*)Handles[2].GetOutputData();
 	int32& Count = *(int32*)Handles[3].GetOutputData();
diff --git a/Engine/Plugins/Runtime/RigVM/Source/RigVM/Private/RigVMFunctions/RigVMDispatch_CastEnum.cpp b/Engine/Plugins/Runtime/RigVM/Source/RigVM/Private/RigVMFunctions/RigVMDispatch_CastEnum.cpp
--- a/Engine/Plugins/Runtime/RigVM/Source/RigVM/Private/RigVMFunctions/RigVMDispatch_CastEnum.cpp
+++ b/Engine/Plugins/Runtime/RigVM/Source/RigVM/Private/RigVMFunctions/RigVMDispatch_CastEnum.cpp
@@ -51,7 +51,7 @@ FText FRigVMDispatch_CastEnumToInt::GetNodeTooltip(const FRigVMTemplateTypeMap&
 
 void FRigVMDispatch_CastEnumToInt::Execute(FRigVMExtendedExecuteContext& InContext, FRigVMMemoryHandleArray Handles, FRigVMPredicateBranchArray RigVMBranches)
 {
-	const FProperty* ValueProperty = CastFieldChecked<FProperty>(Handles[0].GetProperty());
+	const FProperty* ValueProperty = CastFieldChecked<FProperty>(Handles[0].GetResolvedProperty());
 	const FProperty* ResultProperty = CastFieldChecked<FProperty>(Handles[1].GetProperty());
 	if (!ResultProperty || !ValueProperty)
 	{
@@ -145,7 +145,7 @@ FText FRigVMDispatch_CastIntToEnum::GetNodeTooltip(const FRigVMTemplateTypeMap&
 
 void FRigVMDispatch_CastIntToEnum::Execute(FRigVMExtendedExecuteContext& InContext, FRigVMMemoryHandleArray Handles, FRigVMPredicateBranchArray RigVMBranches)
 {
-	const FProperty* ValueProperty = CastFieldChecked<FProperty>(Handles[0].GetProperty());
+	const FProperty* ValueProperty = CastFieldChecked<FProperty>(Handles[0].GetResolvedProperty());
 	const FProperty* ResultProperty = CastFieldChecked<FProperty>(Handles[1].GetProperty());
 	if (!ResultProperty || !ValueProperty)
 	{
diff --git a/Engine/Plugins/Runtime/RigVM/Source/RigVM/Private/RigVMFunctions/RigVMDispatch_CastObject.cpp b/Engine/Plugins/Runtime/RigVM/Source/RigVM/Private/RigVMFunctions/RigVMDispatch_CastObject.cpp
--- a/Engine/Plugins/Runtime/RigVM/Source/RigVM/Private/RigVMFunctions/RigVMDispatch_CastObject.cpp
+++ b/Engine/Plugins/Runtime/RigVM/Source/RigVM/Private/RigVMFunctions/RigVMDispatch_CastObject.cpp
@@ -106,7 +106,7 @@ void FRigVMDispatch_CastObject::Execute(FRigVMExtendedExecuteContext& InContext,
 {
 	UObject** ValueObjectPtr = (UObject**)Handles[0].GetInputData();
 	UObject** ResultObjectPtr = (UObject**)Handles[1].GetOutputData();
-	const FObjectPropertyBase* ValueObjectProperty = CastFieldChecked<FObjectPropertyBase>(Handles[0].GetProperty());
+	const FObjectPropertyBase* ValueObjectProperty = CastFieldChecked<FObjectPropertyBase>(Handles[0].GetResolvedProperty());
 	const FObjectPropertyBase* ResultObjectProperty = CastFieldChecked<FObjectPropertyBase>(Handles[1].GetProperty());
 	if(ValueObjectPtr && ResultObjectPtr && ValueObjectProperty && ResultObjectProperty)
 	{
diff --git a/Engine/Plugins/Runtime/RigVM/Source/RigVM/Private/RigVMFunctions/RigVMDispatch_If.cpp b/Engine/Plugins/Runtime/RigVM/Source/RigVM/Private/RigVMFunctions/RigVMDispatch_If.cpp
--- a/Engine/Plugins/Runtime/RigVM/Source/RigVM/Private/RigVMFunctions/RigVMDispatch_If.cpp
+++ b/Engine/Plugins/Runtime/RigVM/Source/RigVM/Private/RigVMFunctions/RigVMDispatch_If.cpp
@@ -88,6 +88,6 @@ void FRigVMDispatch_If::Execute(FRigVMExtendedExecuteContext& InContext, FRigVMM
 
 	(void)CopyProperty(
 		Handles[3].GetProperty(), Result,
-		InputHandle.GetProperty(), Input);
+		InputHandle.GetResolvedProperty(), Input);
 }
 
diff --git a/Engine/Plugins/Runtime/RigVM/Source/RigVM/Private/RigVMFunctions/RigVMDispatch_MakeStruct.cpp b/Engine/Plugins/Runtime/RigVM/Source/RigVM/Private/RigVMFunctions/RigVMDispatch_MakeStruct.cpp
--- a/Engine/Plugins/Runtime/RigVM/Source/RigVM/Private/RigVMFunctions/RigVMDispatch_MakeStruct.cpp
+++ b/Engine/Plugins/Runtime/RigVM/Source/RigVM/Private/RigVMFunctions/RigVMDispatch_MakeStruct.cpp
@@ -93,7 +93,7 @@ FString FRigVMDispatch_MakeStruct::GetKeywords() const
 
 void FRigVMDispatch_MakeStruct::Execute(FRigVMExtendedExecuteContext& InContext, FRigVMMemoryHandleArray Handles, FRigVMPredicateBranchArray RigVMBranches)
 {
-	const FProperty* SourceProperty = Handles[0].GetProperty();
+	const FProperty* SourceProperty = Handles[0].GetResolvedProperty();
 	const FProperty* TargetProperty = Handles[1].GetProperty();
 	check(SourceProperty->SameType(TargetProperty));
 	const uint8* SourceMemory = Handles[0].GetInputData();
diff --git a/Engine/Plugins/Runtime/RigVM/Source/RigVM/Private/RigVMFunctions/RigVMDispatch_Select.cpp b/Engine/Plugins/Runtime/RigVM/Source/RigVM/Private/RigVMFunctions/RigVMDispatch_Select.cpp
--- a/Engine/Plugins/Runtime/RigVM/Source/RigVM/Private/RigVMFunctions/RigVMDispatch_Select.cpp
+++ b/Engine/Plugins/Runtime/RigVM/Source/RigVM/Private/RigVMFunctions/RigVMDispatch_Select.cpp
@@ -143,6 +143,6 @@ void FRigVMDispatch_SelectInt32::Execute(FRigVMExtendedExecuteContext& InContext
 	// copy property performs compatibility check and will report accordingly
 	(void)CopyProperty(
 		Handles.Last().GetProperty(), Result,
-		InputHandle.GetProperty(), Input);
+		InputHandle.GetResolvedProperty(), Input);
 }
 
diff --git a/Engine/Plugins/Runtime/RigVM/Source/RigVMDeveloper/Private/RigVMCompiler/RigVMAST.cpp b/Engine/Plugins/Runtime/RigVM/Source/RigVMDeveloper/Private/RigVMCompiler/RigVMAST.cpp
--- a/Engine/Plugins/Runtime/RigVM/Source/RigVMDeveloper/Private/RigVMCompiler/RigVMAST.cpp
+++ b/Engine/Plugins/Runtime/RigVM/Source/RigVMDeveloper/Private/RigVMCompiler/RigVMAST.cpp
@@ -589,7 +589,7 @@ void FRigVMExprAST::RemoveChild(FRigVMExprAST* InChild)
 	InChild->RemoveParent(this);
 }
 
-void FRigVMExprAST::ReplaceParent(FRigVMExprAST* InCurrentParent, FRigVMExprAST* InNewParent)
+void FRigVMExprAST::ReplaceParent(FRigVMExprAST* InCurrentParent, FRigVMExprAST* InNewParent, TArray<int32>* OutParentIndices)
 {
 	ensure(IsValid());
 	ensure(InCurrentParent->IsValid());
@@ -603,11 +603,16 @@ void FRigVMExprAST::ReplaceParent(FRigVMExprAST* InCurrentParent, FRigVMExprAST*
 			InCurrentParent->Children.Remove(this);
 			InNewParent->Children.Add(this);
 			InvalidateCaches();
+
+			if (OutParentIndices)
+			{
+				OutParentIndices->Add(ParentIndex);
+			}
 		}
 	}
 }
 
-void FRigVMExprAST::ReplaceChild(FRigVMExprAST* InCurrentChild, FRigVMExprAST* InNewChild)
+void FRigVMExprAST::ReplaceChild(FRigVMExprAST* InCurrentChild, FRigVMExprAST* InNewChild, TArray<int32>* OutChildIndices)
 {
 	ensure(IsValid());
 	ensure(InCurrentChild->IsValid());
@@ -626,6 +631,11 @@ void FRigVMExprAST::ReplaceChild(FRigVMExprAST* InCurrentChild, FRigVMExprAST* I
 			InCurrentChild->InvalidateCachesImpl(InvalidateCachesProcessed);
 			InNewChild->InvalidateCachesImpl(InvalidateCachesProcessed);
 			InvalidateCachesImpl(InvalidateCachesProcessed);
+
+			if (OutChildIndices)
+			{
+				OutChildIndices->Add(ChildIndex);
+			}
 		}
 	}
 }
@@ -889,34 +899,23 @@ const FRigVMVarExprAST* FRigVMNodeExprAST::FindVarWithPinName(const FName& InPin
 	return nullptr;
 }
 
-FRigVMNodeExprAST::FRigVMNodeExprAST(EType InType, const FRigVMASTProxy& InNodeProxy)
-	: FRigVMBlockExprAST(InType)
-	, Proxy(InNodeProxy)
+const FName& FRigVMNodeExprAST::GetPinNameForChildIndex(int32 InChildIndex) const
 {
-}
-
-FName FRigVMEntryExprAST::GetEventName() const
-{
-	if (URigVMNode* EventNode = GetNode())
+	if (const FName* PinName = PinNameToChildIndex.FindKey(InChildIndex))
 	{
-		return EventNode->GetEventName();
+		return *PinName;
 	}
-	return NAME_None;
+	static const FName EmptyString = NAME_None;
+	return EmptyString;
 }
 
-void FRigVMCallExternExprAST::SetSegmentPathForChild(int32 InChildIndex, const FString& InSegmentPath)
+void FRigVMNodeExprAST::SetSegmentPathForChild(int32 InChildIndex, const FString& InSegmentPath)
 {
 	check(Children.IsValidIndex(InChildIndex));
 	SegmentPathForChild.FindOrAdd(InChildIndex) = InSegmentPath;
 }
 
-void FRigVMCallExternExprAST::SetSegmentPathForChild(const FRigVMExprAST* InChildExpr, const FString& InSegmentPath)
-{
-	const int32 ChildIndex = Children.Find(const_cast<FRigVMExprAST*>(InChildExpr));
-	SetSegmentPathForChild(ChildIndex, InSegmentPath);
-}
-
-const FString& FRigVMCallExternExprAST::GetSegmentPathForChild(int32 InChildIndex) const
+const FString& FRigVMNodeExprAST::GetSegmentPathForChild(int32 InChildIndex) const
 {
 	check(Children.IsValidIndex(InChildIndex));
 	if(const FString* ExistingSegmentPath = SegmentPathForChild.Find(InChildIndex))
@@ -927,10 +926,19 @@ const FString& FRigVMCallExternExprAST::GetSegmentPathForChild(int32 InChildInde
 	return EmptySegmentPath;
 }
 
-const FString& FRigVMCallExternExprAST::GetSegmentPathForChild(const FRigVMExprAST* InChildExpr) const
+FRigVMNodeExprAST::FRigVMNodeExprAST(EType InType, const FRigVMASTProxy& InNodeProxy)
+	: FRigVMBlockExprAST(InType)
+	, Proxy(InNodeProxy)
+{
+}
+
+FName FRigVMEntryExprAST::GetEventName() const
 {
-	const int32 ChildIndex = Children.Find(const_cast<FRigVMExprAST*>(InChildExpr));
-	return GetSegmentPathForChild(ChildIndex);
+	if (URigVMNode* EventNode = GetNode())
+	{
+		return EventNode->GetEventName();
+	}
+	return NAME_None;
 }
 
 bool FRigVMVarExprAST::IsConstant() const
@@ -2234,7 +2242,9 @@ void FRigVMParserAST::FoldSubPinCopies()
 		}
 
 		FRigVMCopyExprAST* CopyExpr = Expression->To<FRigVMCopyExprAST>();
-		if (CopyExpr->Parents.Num() != 1 || CopyExpr->Children.Num() != 1)
+		if (CopyExpr->Parents.Num() != 1 ||
+			CopyExpr->Children.Num() != 1 ||
+			!CopyExpr->Link.SegmentPath.IsEmpty())
 		{
 			continue;
 		}
@@ -2268,13 +2278,16 @@ void FRigVMParserAST::FoldSubPinCopies()
 			}
 		}
 
-		// non-input pins on anything but a reroute (passthrough or literal) or array node should be skipped
+		// if the source is a sub-pin of a fixed size array, don't fold the copy
+		if (SourcePin->GetRootPin()->IsFixedSizeArray())
+		{
+			continue;
+		}
+
+		// non-input pins on anything but a reroute (passthrough or literal)
 		const bool bIsReroute = TargetPin->GetNode()->IsA<URigVMRerouteNode>();
 
-		// a copy / assign to a constant node root pin can be folded 
-		const bool bIsConstantValueTarget = IsConstantDispatchValuePin(TargetPin);
-		
-		if (TargetPin->GetDirection() != ERigVMPinDirection::Input && !bIsReroute && !bIsConstantValueTarget)
+		if (TargetPin->GetDirection() != ERigVMPinDirection::Input && !bIsReroute)
 		{
 			continue;
 		}
@@ -2288,13 +2301,16 @@ void FRigVMParserAST::FoldSubPinCopies()
 			}
 		}
 
-		// if this node is a variable node and the pin requires a watch... skip this
-		if (Cast<URigVMVariableNode>(SourcePin->GetNode()))
+		// if the source or target node is a variable node, don't fold this copy
+		if (Cast<URigVMVariableNode>(SourcePin->GetNode()) || Cast<URigVMVariableNode>(TargetPin->GetNode()))
 		{
-			if(SourcePin->RequiresWatch(true))
-			{
-				continue;
-			}
+			continue;
+		}
+
+		// if the source node is a reroute - skip this
+		if (Cast<URigVMRerouteNode>(SourcePin->GetNode()))
+		{
+			continue;
 		}
 
 		// Skip folding for programmatic pins as we expect them to always exist in work memory
@@ -2316,6 +2332,23 @@ void FRigVMParserAST::FoldSubPinCopies()
 			continue;
 		}
 
+		// avoid referencing elements of arrays - since the array may change in the interim.
+		const URigVMPin* SourceArrayPin = SourcePin;
+		bool bIsArrayPin = false;
+		while (SourceArrayPin && !SourceArrayPin->IsRootPin())
+		{
+			if (SourceArrayPin->IsArray() || SourceArrayPin->IsArrayElement())
+			{
+				bIsArrayPin = true;
+				break;
+			}
+			SourceArrayPin = SourceArrayPin->GetParentPin();
+		}
+		if (bIsArrayPin)
+		{
+			continue;
+		}
+
 		FRigVMExprAST* Child = CopyExpr->Children[0];
 		CopyExpr->RemoveParent(Parent);
 		Child->RemoveParent(CopyExpr);
@@ -2323,44 +2356,23 @@ void FRigVMParserAST::FoldSubPinCopies()
 		TArray<FRigVMExprAST*> GrandParents = Parent->Parents;
 		for (FRigVMExprAST* GrandParent : GrandParents)
 		{
-			GrandParent->ReplaceChild(Parent, Child);
+			TArray<int32> ChildIndices;
+			GrandParent->ReplaceChild(Parent, Child, &ChildIndices);
 
-			if (GrandParent->IsA(FRigVMExprAST::EType::CallExtern))
+			if (GrandParent->IsA(FRigVMExprAST::EType::CallExtern) ||
+				GrandParent->IsA(FRigVMExprAST::EType::InlineFunction))
 			{
 				const URigVMPin* CopySourcePin = CopyExpr->GetSourcePin();
-				if(!CopySourcePin->IsRootPin())
-				{
-					FRigVMCallExternExprAST* GrandParentCallExtern = GrandParent->To<FRigVMCallExternExprAST>();
-					GrandParentCallExtern->SetSegmentPathForChild(Child, CopySourcePin->GetSegmentPath());
-				}
-			}
-			else if (GrandParent->IsA(FRigVMExprAST::EType::Assign))
-			{
-				FRigVMAssignExprAST* GrandParentAssign = GrandParent->To<FRigVMAssignExprAST>();
-
-				const URigVMPin* OldSourcePin = GrandParentAssign->GetSourcePin();
-				const FString SourceSegmentPath = OldSourcePin->GetSegmentPath();
-				URigVMPin* NewSourcePin = CopyExpr->GetSourcePin();
+				check(!CopySourcePin->IsRootPin());
 
-				if(!SourceSegmentPath.IsEmpty())
+				FRigVMNodeExprAST* GrandParentNodeExpr = GrandParent->To<FRigVMNodeExprAST>();
+				if(ensure(GrandParentNodeExpr))
 				{
-					NewSourcePin = NewSourcePin->FindSubPin(SourceSegmentPath);
-					if(NewSourcePin == nullptr)
+					for (int32 ChildIndex : ChildIndices)
 					{
-						Settings.Reportf(
-							EMessageSeverity::Error,
-							CopyExpr->GetSourcePin()->GetNode(),
-							TEXT("Cannot resolve segment path '%s' on pin '%s'."),
-							*SourceSegmentPath,
-							*CopyExpr->GetSourcePin()->GetPinPath());
-						continue;
+						GrandParentNodeExpr->SetSegmentPathForChild(ChildIndex, CopySourcePin->GetSegmentPath());
 					}
 				}
-				
-				GrandParentAssign->Link = FRigVMASTLinkDescription(
-					GrandParentAssign->GetSourceProxy().GetSibling(NewSourcePin),
-					GrandParentAssign->GetTargetProxy(), FString());
-				GrandParentAssign->Name = *GetLinkAsString(GrandParentAssign->GetLink());
 			}
 		}
 
diff --git a/Engine/Plugins/Runtime/RigVM/Source/RigVMDeveloper/Private/RigVMCompiler/RigVMCompiler.cpp b/Engine/Plugins/Runtime/RigVM/Source/RigVMDeveloper/Private/RigVMCompiler/RigVMCompiler.cpp
--- a/Engine/Plugins/Runtime/RigVM/Source/RigVMDeveloper/Private/RigVMCompiler/RigVMCompiler.cpp
+++ b/Engine/Plugins/Runtime/RigVM/Source/RigVMDeveloper/Private/RigVMCompiler/RigVMCompiler.cpp
@@ -2135,28 +2135,7 @@ bool URigVMCompiler::TraverseCallExtern(const FRigVMCallExternExprAST* InExpr, F
 			return false;
 		}
 
-		// setup additional property paths needed for call externs accessing sub pin data directly
-		for(const TPair<int32, FString>& SegmentPathPair : CallExternExpr->SegmentPathForChild)
-		{
-			const FRigVMExprAST* ChildExpr = CallExternExpr->ChildAt(SegmentPathPair.Key);
-			check(ChildExpr);
-
-			if(ChildExpr->IsA(FRigVMExprAST::EType::CachedValue))
-			{
-				ChildExpr = ChildExpr->To<FRigVMCachedValueExprAST>()->GetVarExpr();
-			}
-
-			if(!ChildExpr->IsA(FRigVMExprAST::EType::Var))
-			{
-				continue;
-			}
-
-			const FRigVMVarExprAST* SourceVarExpr = GetSourceVarExpr(ChildExpr);
-			const FRigVMOperand RootOperand = WorkData.ExprToOperand.FindChecked(SourceVarExpr);
-			const FString HeadCPPType = SourceVarExpr->GetPin()->GetRootPin()->GetCPPType();
-
-			WorkData.FindOrAddPropertyPath(RootOperand, HeadCPPType, SegmentPathPair.Value);
-		}
+		SetupPropertyPathsForFoldedCopies(InExpr, WorkData);
 
 		if(WorkData.Settings.ASTSettings.bSetupTraits)
 		{
@@ -2214,8 +2193,7 @@ bool URigVMCompiler::TraverseCallExtern(const FRigVMCallExternExprAST* InExpr, F
 	else
 	{
 		TArray<FRigVMOperand> Operands;
-		TArray<const FRigVMExprAST*> ChildExpressionForOperand;
-		TArray<const FRigVMVarExprAST*> VarExpressionForOperand;
+		TArray<const FRigVMVarExprAST*> VarExpressionsPerOperand;
 
 		// iterate over the child expressions in the order of the arguments on the function
 		const FRigVMFunction* Function = nullptr;
@@ -2240,8 +2218,7 @@ bool URigVMCompiler::TraverseCallExtern(const FRigVMCallExternExprAST* InExpr, F
 				&WorkData,
 				Argument,
 				&Operands,
-				&ChildExpressionForOperand,
-				&VarExpressionForOperand,
+				&VarExpressionsPerOperand,
 				&BlockToRunOperand,
 				&CountOperand,
 				&IndexOperand
@@ -2262,8 +2239,7 @@ bool URigVMCompiler::TraverseCallExtern(const FRigVMCallExternExprAST* InExpr, F
 				}
 
 				Operands.Add(WorkData.ExprToOperand.FindChecked(SourceVarExpr));
-				ChildExpressionForOperand.Add(InExpr);
-				VarExpressionForOperand.Add(SourceVarExpr);
+				VarExpressionsPerOperand.Add(SourceVarExpr);
 
 				if(Argument.Name == FRigVMStruct::ControlFlowBlockToRunName)
 				{
@@ -2300,24 +2276,35 @@ bool URigVMCompiler::TraverseCallExtern(const FRigVMCallExternExprAST* InExpr, F
 			ProcessArgument(ChildExpr);
 		}
 
-		if(!CallExternExpr->SegmentPathForChild.IsEmpty())
+		// for call externs we can only apply the segment path after all operands are known
+		for (const TPair<int32, FString>& Pair : CallExternExpr->SegmentPathForChild)
 		{
-			for(int32 OperandIndex = 0; OperandIndex < Operands.Num(); OperandIndex++)
+			FRigVMRegistryReadLock LockedRegistry;
+			
+			const FString& SegmentPath = Pair.Value;
+			check(!SegmentPath.IsEmpty());
+
+			int32 bNumUpdatedOperands = 0;
+			const FName& PinName = CallExternExpr->GetPinNameForChildIndex(Pair.Key);
+			for (int32 OperandIndex = 0; OperandIndex < Operands.Num(); OperandIndex++)
 			{
-				const FString& SegmentPath = CallExternExpr->GetSegmentPathForChild(ChildExpressionForOperand[OperandIndex]);
-				if(SegmentPath.IsEmpty())
+				const FName& ArgumentName = Function->GetArgumentNameForOperandIndex(OperandIndex, Operands.Num(), LockedRegistry);
+				if (ArgumentName != PinName)
 				{
 					continue;
 				}
-
-				const FRigVMVarExprAST* SourceVarExpr = VarExpressionForOperand[OperandIndex];
+				
+				const FRigVMVarExprAST* SourceVarExpr = VarExpressionsPerOperand[OperandIndex];
 				check(SourceVarExpr);
 
 				const FString HeadCPPType = SourceVarExpr->GetPin()->GetCPPType();
 				const int32 RegisterOffsetIndex = WorkData.FindOrAddPropertyPath(Operands[OperandIndex], HeadCPPType, SegmentPath);
 				check(RegisterOffsetIndex != INDEX_NONE);
 				Operands[OperandIndex].RegisterOffset = static_cast<uint16>(RegisterOffsetIndex);
+				bNumUpdatedOperands++;
 			}
+
+			check(bNumUpdatedOperands == 1);
 		}
 
 		// make sure to skip the output blocks while we are traversing this call extern
@@ -2864,6 +2851,8 @@ bool URigVMCompiler::TraverseInlineFunction(const FRigVMInlineFunctionExprAST* I
 					FString::Printf(TEXT("(%s)"), *FString::Join(IndexStrings, TEXT(",")));
 			}
 		}
+
+		SetupPropertyPathsForFoldedCopies(InExpr, WorkData);
 	}
 	else
 	{
@@ -2873,29 +2862,40 @@ bool URigVMCompiler::TraverseInlineFunction(const FRigVMInlineFunctionExprAST* I
 		for(const URigVMPin* Pin : FunctionReferenceNode->GetPins())
 		{
 			const FRigVMExprAST* ChildExpr = InExpr->FindExprWithPinName(Pin->GetFName());
-			checkf(ChildExpr, TEXT("Found unexpected opaque argument for %s while inlining function %s in package %s"), *InExpr->Name.ToString(), *FunctionReferenceNode->GetPathName(), *GetPackage()->GetPathName());			
+			checkf(ChildExpr, TEXT("Found unexpected opaque argument for %s while inlining function %s in package %s"), *InExpr->Name.ToString(), *FunctionReferenceNode->GetPathName(), *GetPackage()->GetPathName());
+
+			const FRigVMVarExprAST* SourceVarExpr = nullptr;
 			if (ChildExpr->GetType() == FRigVMExprAST::EType::CachedValue)
 			{
-				const FRigVMVarExprAST* SourceVarExpr = GetSourceVarExpr(ChildExpr->To<FRigVMCachedValueExprAST>()->GetVarExpr());
-				if(!SourceVarExpr->IsExecuteContext())
-				{
-					InterfaceOperands.Add(WorkData.ExprToOperand.FindChecked(SourceVarExpr));
-					OperandsPinNames.Add(GetPinNameWithDirectionPrefix(Pin));
-				}
+				SourceVarExpr = GetSourceVarExpr(ChildExpr->To<FRigVMCachedValueExprAST>()->GetVarExpr());
 			}
 			else if (ChildExpr->IsA(FRigVMExprAST::EType::Var))
 			{
-				const FRigVMVarExprAST* SourceVarExpr = GetSourceVarExpr(ChildExpr->To<FRigVMVarExprAST>());
-				if(!SourceVarExpr->IsExecuteContext())
-				{
-					InterfaceOperands.Add(WorkData.ExprToOperand.FindChecked(SourceVarExpr));
-					OperandsPinNames.Add(GetPinNameWithDirectionPrefix(Pin));
-				}
+				SourceVarExpr = GetSourceVarExpr(ChildExpr->To<FRigVMVarExprAST>());
 			}
 			else
 			{
 				break;
 			}
+			check (SourceVarExpr);
+			
+			if(SourceVarExpr->IsExecuteContext())
+			{
+				continue;
+			}
+
+			InterfaceOperands.Add(WorkData.ExprToOperand.FindChecked(SourceVarExpr));
+			OperandsPinNames.Add(GetPinNameWithDirectionPrefix(Pin));
+
+			if ( const FString* SegmentPath = InExpr->SegmentPathForChild.Find(Pin->GetPinIndex()))
+			{
+				check(!SegmentPath->IsEmpty());
+				
+				const FString HeadCPPType = SourceVarExpr->GetPin()->GetCPPType();
+				const int32 RegisterOffsetIndex = WorkData.FindOrAddPropertyPath(InterfaceOperands.Last(), HeadCPPType, *SegmentPath);
+				check(RegisterOffsetIndex != INDEX_NONE);
+				InterfaceOperands.Last().RegisterOffset = static_cast<uint16>(RegisterOffsetIndex);
+			}
 		}
 
 		if (!TraverseChildren(InExpr, WorkData))
@@ -2943,6 +2943,7 @@ bool URigVMCompiler::TraverseInlineFunction(const FRigVMInlineFunctionExprAST* I
 				FRigVMOperand* Operand = (FRigVMOperand*)(ByteCode.ByteCode.GetData() + OperandsIndex + j*sizeof(FRigVMOperand));
 				FRigVMOperand FunctionOperand = *Operand;
 				ERigVMMemoryType OriginalMemoryType = Operand->GetMemoryType();
+				FRigVMOperand InterfaceOperand;
 
 				// Remap the variable: find the operand index of the outer variable
 				if (Operand->GetMemoryType() == ERigVMMemoryType::External)
@@ -2963,9 +2964,9 @@ bool URigVMCompiler::TraverseInlineFunction(const FRigVMInlineFunctionExprAST* I
 				// If Operand is an interface pin: replace the index and memory type
 				else if (const int32 FunctionInterfaceParameterIndex = GetOperandFunctionInterfaceParameterIndex(OperandsPinNames, FunctionCompilationData, *Operand); FunctionInterfaceParameterIndex != INDEX_NONE)
 				{
-					const FRigVMOperand& InterfaceOperand = InterfaceOperands[FunctionInterfaceParameterIndex];
-					Operand->MemoryType = InterfaceOperand.GetMemoryType();
-					Operand->RegisterIndex = static_cast<uint16>(InterfaceOperand.GetRegisterIndex());
+					InterfaceOperand = InterfaceOperands[FunctionInterfaceParameterIndex];
+					Operand->MemoryType = InterfaceOperand.MemoryType;
+					Operand->RegisterIndex = InterfaceOperand.RegisterIndex;
 				}
 				else
 				{
@@ -2978,33 +2979,56 @@ bool URigVMCompiler::TraverseInlineFunction(const FRigVMInlineFunctionExprAST* I
 				}
 
 				// For all operands, check to see if we need to add a property path
-				if (Operand->GetRegisterOffset() != INDEX_NONE)
+				if (Operand->GetRegisterOffset() != INDEX_NONE || InterfaceOperand.GetRegisterOffset() != INDEX_NONE)
 				{
-					FRigVMFunctionCompilationPropertyPath Description;
-					switch (OriginalMemoryType)
+					if (InterfaceOperand.GetRegisterOffset() != INDEX_NONE && Operand->GetRegisterOffset() == INDEX_NONE)
 					{
-						case ERigVMMemoryType::Work:
-						{
-							Description = FunctionCompilationData->WorkPropertyPathDescriptions[Operand->GetRegisterOffset()];
-							break;
-						}
-						case ERigVMMemoryType::Literal:
+						Operand->RegisterOffset = IntCastChecked<uint16>(InterfaceOperand.GetRegisterOffset());
+					}
+					else
+					{
+						FRigVMFunctionCompilationPropertyPath DescriptionFromFunction;
+						switch (OriginalMemoryType)
 						{
-							Description = FunctionCompilationData->LiteralPropertyPathDescriptions[Operand->GetRegisterOffset()];
-							break;
+							case ERigVMMemoryType::Work:
+							{
+								DescriptionFromFunction = FunctionCompilationData->WorkPropertyPathDescriptions[Operand->GetRegisterOffset()];
+								break;
+							}
+							case ERigVMMemoryType::Literal:
+							{
+								DescriptionFromFunction = FunctionCompilationData->LiteralPropertyPathDescriptions[Operand->GetRegisterOffset()];
+								break;
+							}
+							case ERigVMMemoryType::External:
+							{
+								DescriptionFromFunction = FunctionCompilationData->ExternalPropertyPathDescriptions[Operand->GetRegisterOffset()];
+								break;
+							}
+							case ERigVMMemoryType::Debug:
+							{
+								DescriptionFromFunction = FunctionCompilationData->DebugPropertyPathDescriptions[Operand->GetRegisterOffset()];
+								break;
+							}
+							default:
+							{
+								checkNoEntry();
+							}
 						}
-						case ERigVMMemoryType::External:
+
+						if (InterfaceOperand.GetRegisterOffset() == INDEX_NONE)
 						{
-							Description = FunctionCompilationData->ExternalPropertyPathDescriptions[Operand->GetRegisterOffset()];
-							break;
+							Operand->RegisterOffset = IntCastChecked<uint16>(WorkData.FindOrAddPropertyPath(*Operand, DescriptionFromFunction.HeadCPPType, DescriptionFromFunction.SegmentPath));
 						}
-						case ERigVMMemoryType::Debug:
+						else
 						{
-							Description = FunctionCompilationData->DebugPropertyPathDescriptions[Operand->GetRegisterOffset()];
-							break;
+							// in this case both the external operand and the operand within the function have a register offset.
+							// they need to be combined.
+							const FRigVMPropertyPathDescription& Description = WorkData.PropertyPathDescriptions.FindOrAdd(Operand->GetMemoryType())[Operand->GetRegisterOffset()];
+							const FString CombinedSegmentPath = URigVMPin::JoinPinPath(Description.SegmentPath, DescriptionFromFunction.SegmentPath);
+							Operand->RegisterOffset = IntCastChecked<uint16>(WorkData.FindOrAddPropertyPath(*Operand, Description.HeadCPPType, CombinedSegmentPath));
 						}
 					}
-					Operand->RegisterOffset = IntCastChecked<uint16>(WorkData.FindOrAddPropertyPath(*Operand, Description.HeadCPPType, Description.SegmentPath));
 				}
 
 				// Store correspondence between Function operand and New operands
@@ -4422,6 +4446,32 @@ bool URigVMCompiler::ValidateNode(const FRigVMCompileSettings& InSettings, URigV
 	return false;
 }
 
+void URigVMCompiler::SetupPropertyPathsForFoldedCopies(const FRigVMNodeExprAST* InExpr, FRigVMCompilerWorkData& WorkData)
+{
+	// setup additional property paths needed for call externs accessing sub pin data directly
+	for(const TPair<int32, FString>& SegmentPathPair : InExpr->SegmentPathForChild)
+	{
+		const FRigVMExprAST* ChildExpr = InExpr->ChildAt(SegmentPathPair.Key);
+		check(ChildExpr);
+
+		if(ChildExpr->IsA(FRigVMExprAST::EType::CachedValue))
+		{
+			ChildExpr = ChildExpr->To<FRigVMCachedValueExprAST>()->GetVarExpr();
+		}
+
+		if(!ChildExpr->IsA(FRigVMExprAST::EType::Var))
+		{
+			continue;
+		}
+
+		const FRigVMVarExprAST* SourceVarExpr = GetSourceVarExpr(ChildExpr);
+		const FRigVMOperand RootOperand = WorkData.ExprToOperand.FindChecked(SourceVarExpr);
+		const FString HeadCPPType = SourceVarExpr->GetPin()->GetRootPin()->GetCPPType();
+
+		WorkData.FindOrAddPropertyPath(RootOperand, HeadCPPType, SegmentPathPair.Value);
+	}
+}
+
 void URigVMCompiler::ReportInfo(const FRigVMCompileSettings& InSettings, const FString& InMessage)
 {
 	if (InSettings.SurpressInfoMessages)
diff --git a/Engine/Plugins/Runtime/RigVM/Source/RigVMDeveloper/Public/RigVMCompiler/RigVMAST.h b/Engine/Plugins/Runtime/RigVM/Source/RigVMDeveloper/Public/RigVMCompiler/RigVMAST.h
--- a/Engine/Plugins/Runtime/RigVM/Source/RigVMDeveloper/Public/RigVMCompiler/RigVMAST.h
+++ b/Engine/Plugins/Runtime/RigVM/Source/RigVMDeveloper/Public/RigVMCompiler/RigVMAST.h
@@ -311,12 +311,14 @@ protected:
 	// replaces a parent of this expression with a new one
 	// @param InCurrentParent the current parent to replace
 	// @param InNewParent the new parent to replace it with
-	UE_API void ReplaceParent(FRigVMExprAST* InCurrentParent, FRigVMExprAST* InNewParent);
+	// @param OutParentIndices The list of parent indices affected
+	UE_API void ReplaceParent(FRigVMExprAST* InCurrentParent, FRigVMExprAST* InNewParent, TArray<int32>* OutParentIndices = nullptr);
 
 	// replaces a child of this expression with a new one
 	// @param InCurrentChild the current child to replace
 	// @param InNewChild the new child to replace it with
-	UE_API void ReplaceChild(FRigVMExprAST* InCurrentChild, FRigVMExprAST* InNewChild);
+	// @param OutChildIndices The list of child indices affected
+	UE_API void ReplaceChild(FRigVMExprAST* InCurrentChild, FRigVMExprAST* InNewChild, TArray<int32>* OutChildIndices = nullptr);
 
 	// replaces this expression with another one
 	// @param InReplacement the expression to replace this one
@@ -597,6 +599,10 @@ public:
 	UE_API virtual bool IsConstant() const override;
 	UE_API const FRigVMExprAST* FindExprWithPinName(const FName& InPinName) const;
 	UE_API const FRigVMVarExprAST* FindVarWithPinName(const FName& InPinName) const;
+	UE_API const FName& GetPinNameForChildIndex(int32 InChildIndex) const;
+
+	void SetSegmentPathForChild(int32 InChildIndex, const FString& InSegmentPath);
+	const FString& GetSegmentPathForChild(int32 InChildIndex) const;
 
 protected:
 
@@ -611,8 +617,11 @@ protected:
 
 private:
 
+	TMap<int32,FString> SegmentPathForChild;
+	
 	FRigVMASTProxy Proxy;
 	friend class FRigVMParserAST;
+	friend class URigVMCompiler;
 };
 
 /*
@@ -719,12 +728,6 @@ public:
 		return InType == EType::CallExtern;
 	};
 
-	void SetSegmentPathForChild(int32 InChildIndex, const FString& InSegmentPath);
-	void SetSegmentPathForChild(const FRigVMExprAST* InChildExpr, const FString& InSegmentPath);
-
-	const FString& GetSegmentPathForChild(int32 InChildIndex) const;
-	const FString& GetSegmentPathForChild(const FRigVMExprAST* InChildExpr) const;
-
 protected:
 
 	// default constructor (protected so that only parser can access it)
@@ -734,10 +737,7 @@ protected:
 
 private:
 
-	TMap<int32,FString> SegmentPathForChild;
-
 	friend class FRigVMParserAST;
-	friend class URigVMCompiler;
 };
 
 /*
diff --git a/Engine/Plugins/Runtime/RigVM/Source/RigVMDeveloper/Public/RigVMCompiler/RigVMCompiler.h b/Engine/Plugins/Runtime/RigVM/Source/RigVMDeveloper/Public/RigVMCompiler/RigVMCompiler.h
--- a/Engine/Plugins/Runtime/RigVM/Source/RigVMDeveloper/Public/RigVMCompiler/RigVMCompiler.h
+++ b/Engine/Plugins/Runtime/RigVM/Source/RigVMDeveloper/Public/RigVMCompiler/RigVMCompiler.h
@@ -389,6 +389,8 @@ private:
 	static UE_API int32 GetOperandFunctionInterfaceParameterIndex(const TArray<FString>& OperandsPinNames, const FRigVMFunctionCompilationData* FunctionCompilationData, const FRigVMOperand& Operand);
 
 	UE_API bool ValidateNode(const FRigVMCompileSettings& InSettings, URigVMNode* InNode, bool bCheck = true);
+
+	void SetupPropertyPathsForFoldedCopies(const FRigVMNodeExprAST* InExpr, FRigVMCompilerWorkData& WorkData);
 	
 	UE_API void ReportInfo(const FRigVMCompileSettings& InSettings, const FString& InMessage);
 	UE_API void ReportWarning(const FRigVMCompileSettings& InSettings, const FString& InMessage);
