RigVM: Further improve sub pin folding, add more checks

#rnx

diff --git a/Engine/Plugins/Runtime/RigVM/Source/RigVM/Private/RigVMCore/RigVM.cpp b/Engine/Plugins/Runtime/RigVM/Source/RigVM/Private/RigVMCore/RigVM.cpp
--- a/Engine/Plugins/Runtime/RigVM/Source/RigVM/Private/RigVMCore/RigVM.cpp
+++ b/Engine/Plugins/Runtime/RigVM/Source/RigVM/Private/RigVMCore/RigVM.cpp
@@ -3042,6 +3042,7 @@ void URigVM::CacheSingleMemoryHandle(FRigVMExtendedExecuteContext& Context, cons
 	{
 		check(Memory->GetPropertyPaths().IsValidIndex(InArg.GetRegisterOffset()));
 		PropertyPath = &Memory->GetPropertyPaths()[InArg.GetRegisterOffset()];
+		check(!PropertyPath->IsEmpty());
 	}
 
 	// if you are hitting this it's likely that the VM was created outside of a valid
diff --git a/Engine/Plugins/Runtime/RigVM/Source/RigVMDeveloper/Private/RigVMCompiler/RigVMAST.cpp b/Engine/Plugins/Runtime/RigVM/Source/RigVMDeveloper/Private/RigVMCompiler/RigVMAST.cpp
--- a/Engine/Plugins/Runtime/RigVM/Source/RigVMDeveloper/Private/RigVMCompiler/RigVMAST.cpp
+++ b/Engine/Plugins/Runtime/RigVM/Source/RigVMDeveloper/Private/RigVMCompiler/RigVMAST.cpp
@@ -2318,6 +2318,12 @@ void FRigVMParserAST::FoldSubPinCopies()
 		{
 			continue;
 		}
+
+		// Skip folding onto a return node
+		if (TargetPin->GetNode()->IsA<URigVMFunctionReturnNode>())
+		{
+			continue;
+		}
 		
 		FRigVMExprAST* Parent = CopyExpr->Parents[0];
 		if (!Parent->IsA(FRigVMExprAST::EType::Var))
diff --git a/Engine/Plugins/Runtime/RigVM/Source/RigVMDeveloper/Private/RigVMCompiler/RigVMCompiler.cpp b/Engine/Plugins/Runtime/RigVM/Source/RigVMDeveloper/Private/RigVMCompiler/RigVMCompiler.cpp
--- a/Engine/Plugins/Runtime/RigVM/Source/RigVMDeveloper/Private/RigVMCompiler/RigVMCompiler.cpp
+++ b/Engine/Plugins/Runtime/RigVM/Source/RigVMDeveloper/Private/RigVMCompiler/RigVMCompiler.cpp
@@ -186,6 +186,7 @@ int32 FRigVMCompilerWorkData::FindOrAddPropertyPath(const FRigVMOperand& InOpera
 {
 	if(InSegmentPath.IsEmpty())
 	{
+		checkNoEntry();
 		return INDEX_NONE;
 	}
 
@@ -198,7 +199,9 @@ int32 FRigVMCompilerWorkData::FindOrAddPropertyPath(const FRigVMOperand& InOpera
 			return Index;
 		}
 	}
-	return Descriptions.Add(FRigVMPropertyPathDescription(InOperand.GetRegisterIndex(), InHeadCPPType, InSegmentPath));
+	const int32 Index = Descriptions.Add(FRigVMPropertyPathDescription(InOperand.GetRegisterIndex(), InHeadCPPType, InSegmentPath));
+	check(Descriptions.Last().IsValid());
+	return Index;
 }
 
 const FProperty* FRigVMCompilerWorkData::GetPropertyForOperand(const FRigVMOperand& InOperand) const
@@ -2426,6 +2429,59 @@ bool URigVMCompiler::TraverseCallExtern(const FRigVMCallExternExprAST* InExpr, F
 			bSetupTraits = true;
 		}
 
+		/*
+		// sanity check the operand types
+		check(Function->Arguments.Num() == Operands.Num() || Function->Factory != nullptr);
+		if (Function->Arguments.Num() == Operands.Num() && Function->Factory == nullptr)
+		{
+			for (int32 ArgumentIndex = 0; ArgumentIndex < Function->Arguments.Num(); ArgumentIndex++)
+			{
+				const FRigVMFunctionArgument& Argument = Function->Arguments[ArgumentIndex];
+				FString ExpectedType = Argument.Type;
+				if (ExpectedType.Contains(TEXT("TSubclassOf<")))
+				{
+					continue;
+				}
+
+				if (ExpectedType.Contains(TEXT("TEnumAsByte<")))
+				{
+					ExpectedType.ReplaceInline(TEXT("TEnumAsByte<"), TEXT(""));
+					ExpectedType.RemoveFromEnd(TEXT(">"));
+				}
+
+				const FRigVMOperand& Operand = Operands[ArgumentIndex];
+				const FProperty* Property = WorkData.GetPropertyForOperand(Operand);
+				check(Property != nullptr);
+
+				FString CPPType = RigVMTypeUtils::GetCPPTypeFromProperty(Property);
+
+				// the argument may have an element type of an array, while the operand is a sliced array.
+				if (ExpectedType != CPPType)
+				{
+					const FString SlicedArrayCPPType = RigVMTypeUtils::ArrayTypeFromBaseType(ExpectedType);
+					if (SlicedArrayCPPType == CPPType)
+					{
+						// make sure to check that the pin is internal
+						const FRigVMVarExprAST* SourceVarExpr = InExpr->FindVarWithPinName(Argument.Name);
+						if (SourceVarExpr->GetType() == FRigVMExprAST::EType::CachedValue)
+						{
+							SourceVarExpr = GetSourceVarExpr(SourceVarExpr->To<FRigVMCachedValueExprAST>()->GetVarExpr());
+						}
+						else if (SourceVarExpr->IsA(FRigVMExprAST::EType::Var))
+						{
+							SourceVarExpr = GetSourceVarExpr(SourceVarExpr->To<FRigVMVarExprAST>());
+						}
+						check(SourceVarExpr);
+						check(SourceVarExpr->GetPinDirection() == ERigVMPinDirection::Hidden);
+						ExpectedType = SlicedArrayCPPType;;
+					}
+				}
+				
+				check (ExpectedType == CPPType);
+			}
+		}
+		*/
+		
 		// setup the instruction
 		const int32 FunctionIndex = WorkData.VM->AddRigVMFunction(Function->GetName());
 		check(FunctionIndex != INDEX_NONE);
@@ -3132,7 +3188,6 @@ bool URigVMCompiler::TraverseInlineFunction(const FRigVMInlineFunctionExprAST* I
 				WorkData.VM->OperandToDebugRegisters.Add(*NewOperand, NewValues);
 			}
 		}
-		
 
 #if WITH_EDITORONLY_DATA
 		TArray<FRigVMOperand> InputsOperands, OutputOperands;
