RigVM: Further solidify sub pin folding - make sure link inline respect the break struct dispatch

#rnx

diff --git a/Engine/Plugins/Runtime/RigVM/Source/RigVMDeveloper/Private/RigVMCompiler/RigVMAST.cpp b/Engine/Plugins/Runtime/RigVM/Source/RigVMDeveloper/Private/RigVMCompiler/RigVMAST.cpp
--- a/Engine/Plugins/Runtime/RigVM/Source/RigVMDeveloper/Private/RigVMCompiler/RigVMAST.cpp
+++ b/Engine/Plugins/Runtime/RigVM/Source/RigVMDeveloper/Private/RigVMCompiler/RigVMAST.cpp
@@ -3302,11 +3302,19 @@ void FRigVMParserAST::Inline(const TArray<URigVMGraph*>& InGraphs, const TArray<
 				return false;
 			}
 			
-			if (URigVMRerouteNode* RerouteNode = Cast<URigVMRerouteNode>(Node))
+			if (Cast<URigVMRerouteNode>(Node))
 			{
 				return true;
 			}
-			
+
+			if (OutTraversalInfo.Settings->bFoldSubPinCopies)
+			{
+				if (IsBreakStructDispatchNode(Node))
+				{
+					return true;
+				}
+			}
+				
 			return Node->IsA<URigVMLibraryNode>() ||
 				Node->IsA<URigVMFunctionInterfaceNode>();
 		}
@@ -3486,6 +3494,24 @@ void FRigVMParserAST::Inline(const TArray<URigVMGraph*>& InGraphs, const TArray<
 							}
 						}
 					}
+					else if (IsBreakStructDispatchNode(ChildPin->GetNode()))
+					{
+						if (OutTraversalInfo.Settings->bFoldSubPinCopies)
+						{
+							if (URigVMPin* StructPin = ChildPin->GetNode()->FindRootPinByName(FRigVMDispatch_BreakStruct::StructName.Resolve()))
+							{
+								SourcePinProxy = InPinProxy.GetSibling(StructPin); 
+								if (ShouldRecursePin(SourcePinProxy, OutTraversalInfo))
+								{
+									FRigVMASTProxy SourceSourcePinProxy = FindSourcePin(SourcePinProxy, OutTraversalInfo);
+									if(SourceSourcePinProxy.IsValid())
+									{
+										SourcePinProxy = SourceSourcePinProxy;
+									}
+								}
+							}
+						}
+					}
 				}
 
 				TArray<URigVMLink*> SourceLinks = ChildPin->GetSourceLinks(false /* recursive */);
