Fixed RigVMTypeUtils type names using wrong FLazyName constructor.
Also converted a few FLazyNames to FNames where appropriate.

The FLazyName(FName) constructor was being used here because of const TCHAR* being used as the parameter. That constructor will likely be deprecated soon.

diff --git a/Engine/Plugins/Runtime/RigVM/Source/RigVM/Private/RigVMFunctions/RigVMFunction_ControlFlow.cpp b/Engine/Plugins/Runtime/RigVM/Source/RigVM/Private/RigVMFunctions/RigVMFunction_ControlFlow.cpp
--- a/Engine/Plugins/Runtime/RigVM/Source/RigVM/Private/RigVMFunctions/RigVMFunction_ControlFlow.cpp
+++ b/Engine/Plugins/Runtime/RigVM/Source/RigVM/Private/RigVMFunctions/RigVMFunction_ControlFlow.cpp
@@ -10,12 +10,12 @@ FRigVMFunction_ControlFlowBranch_Execute()
 	{
 		if(Condition)
 		{
-			static const FLazyName TrueName(GET_MEMBER_NAME_CHECKED(FRigVMFunction_ControlFlowBranch, True)); 
+			static const FName TrueName(GET_MEMBER_NAME_CHECKED(FRigVMFunction_ControlFlowBranch, True)); 
 			BlockToRun = TrueName;
 		}
 		else
 		{
-			static const FLazyName FalseName(GET_MEMBER_NAME_CHECKED(FRigVMFunction_ControlFlowBranch, False)); 
+			static const FName FalseName(GET_MEMBER_NAME_CHECKED(FRigVMFunction_ControlFlowBranch, False)); 
 			BlockToRun = FalseName;
 		}
 	}
diff --git a/Engine/Plugins/Runtime/RigVM/Source/RigVM/Public/RigVMTypeUtils.h b/Engine/Plugins/Runtime/RigVM/Source/RigVM/Public/RigVMTypeUtils.h
--- a/Engine/Plugins/Runtime/RigVM/Source/RigVM/Public/RigVMTypeUtils.h
+++ b/Engine/Plugins/Runtime/RigVM/Source/RigVM/Public/RigVMTypeUtils.h
@@ -49,57 +49,57 @@ private:
 
 namespace RigVMTypeUtils
 {
-	constexpr TCHAR TArrayPrefix[] = TEXT("TArray<");
-	constexpr TCHAR TObjectPtrPrefix[] = TEXT("TObjectPtr<");
-	constexpr TCHAR TSubclassOfPrefix[] = TEXT("TSubclassOf<");
-	constexpr TCHAR TScriptInterfacePrefix[] = TEXT("TScriptInterface<");
-	constexpr TCHAR TArrayTemplate[] = TEXT("TArray<%s>");
-	constexpr TCHAR TObjectPtrTemplate[] = TEXT("TObjectPtr<%s%s>");
-	constexpr TCHAR TSubclassOfTemplate[] = TEXT("TSubclassOf<%s%s>");
-	constexpr TCHAR TScriptInterfaceTemplate[] = TEXT("TScriptInterface<%s%s>");
-
-	const inline TCHAR* BoolType = TEXT("bool");
-	const inline TCHAR* FloatType = TEXT("float");
-	const inline TCHAR* DoubleType = TEXT("double");
-	const inline TCHAR* IntType = TEXT("int");
-	const inline TCHAR* Int32Type = TEXT("int32");
-	const inline TCHAR* Int64Type = TEXT("int64");
-	const inline TCHAR* UInt8Type = TEXT("uint8");
-	const inline TCHAR* UInt32Type = TEXT("uint32");
-	const inline TCHAR* UInt64Type = TEXT("uint64");
-	const inline TCHAR* FNameType = TEXT("FName");
-	const inline TCHAR* FStringType = TEXT("FString");
-	const inline TCHAR* FTextType = TEXT("FText");
-	const inline TCHAR* BoolArrayType = TEXT("TArray<bool>");
-	const inline TCHAR* FloatArrayType = TEXT("TArray<float>");
-	const inline TCHAR* DoubleArrayType = TEXT("TArray<double>");
-	const inline TCHAR* Int32ArrayType = TEXT("TArray<int32>");
-	const inline TCHAR* UInt32ArrayType = TEXT("TArray<uint32>");
-	const inline TCHAR* UInt8ArrayType = TEXT("TArray<uint8>");
-	const inline TCHAR* FNameArrayType = TEXT("TArray<FName>");
-	const inline TCHAR* FStringArrayType = TEXT("TArray<FString>");
-	const inline TCHAR* FTextArrayType = TEXT("TArray<FText>");
-
-	const FLazyName BoolTypeName(BoolType);
-	const FLazyName FloatTypeName(FloatType);
-	const FLazyName DoubleTypeName(DoubleType);
-	const FLazyName IntTypeName(IntType);
-	const FLazyName Int32TypeName(Int32Type);
-	const FLazyName Int64TypeName(Int64Type);
-	const FLazyName UInt8TypeName(UInt8Type);
-	const FLazyName UInt32TypeName(UInt32Type);
-	const FLazyName UInt64TypeName(UInt64Type);
-	const FLazyName FNameTypeName(FNameType);
-	const FLazyName FStringTypeName(FStringType);
-	const FLazyName FTextTypeName(FTextType);
-	const FLazyName BoolArrayTypeName(BoolArrayType);
-	const FLazyName FloatArrayTypeName(FloatArrayType);
-	const FLazyName DoubleArrayTypeName(DoubleArrayType);
-	const FLazyName Int32ArrayTypeName(Int32ArrayType);
-	const FLazyName UInt32ArrayTypeName(UInt32ArrayType);
-	const FLazyName UInt8ArrayTypeName(UInt8ArrayType);
-	const FLazyName FNameArrayTypeName(FNameArrayType);
-	const FLazyName FStringArrayTypeName(FStringArrayType);
+	inline constexpr TCHAR TArrayPrefix[] = TEXT("TArray<");
+	inline constexpr TCHAR TObjectPtrPrefix[] = TEXT("TObjectPtr<");
+	inline constexpr TCHAR TSubclassOfPrefix[] = TEXT("TSubclassOf<");
+	inline constexpr TCHAR TScriptInterfacePrefix[] = TEXT("TScriptInterface<");
+	inline constexpr TCHAR TArrayTemplate[] = TEXT("TArray<%s>");
+	inline constexpr TCHAR TObjectPtrTemplate[] = TEXT("TObjectPtr<%s%s>");
+	inline constexpr TCHAR TSubclassOfTemplate[] = TEXT("TSubclassOf<%s%s>");
+	inline constexpr TCHAR TScriptInterfaceTemplate[] = TEXT("TScriptInterface<%s%s>");
+
+	inline constexpr TCHAR BoolType[] = TEXT("bool");
+	inline constexpr TCHAR FloatType[] = TEXT("float");
+	inline constexpr TCHAR DoubleType[] = TEXT("double");
+	inline constexpr TCHAR IntType[] = TEXT("int");
+	inline constexpr TCHAR Int32Type[] = TEXT("int32");
+	inline constexpr TCHAR Int64Type[] = TEXT("int64");
+	inline constexpr TCHAR UInt8Type[] = TEXT("uint8");
+	inline constexpr TCHAR UInt32Type[] = TEXT("uint32");
+	inline constexpr TCHAR UInt64Type[] = TEXT("uint64");
+	inline constexpr TCHAR FNameType[] = TEXT("FName");
+	inline constexpr TCHAR FStringType[] = TEXT("FString");
+	inline constexpr TCHAR FTextType[] = TEXT("FText");
+	inline constexpr TCHAR BoolArrayType[] = TEXT("TArray<bool>");
+	inline constexpr TCHAR FloatArrayType[] = TEXT("TArray<float>");
+	inline constexpr TCHAR DoubleArrayType[] = TEXT("TArray<double>");
+	inline constexpr TCHAR Int32ArrayType[] = TEXT("TArray<int32>");
+	inline constexpr TCHAR UInt32ArrayType[] = TEXT("TArray<uint32>");
+	inline constexpr TCHAR UInt8ArrayType[] = TEXT("TArray<uint8>");
+	inline constexpr TCHAR FNameArrayType[] = TEXT("TArray<FName>");
+	inline constexpr TCHAR FStringArrayType[] = TEXT("TArray<FString>");
+	inline constexpr TCHAR FTextArrayType[] = TEXT("TArray<FText>");
+
+	inline const FLazyName BoolTypeName(BoolType);
+	inline const FLazyName FloatTypeName(FloatType);
+	inline const FLazyName DoubleTypeName(DoubleType);
+	inline const FLazyName IntTypeName(IntType);
+	inline const FLazyName Int32TypeName(Int32Type);
+	inline const FLazyName Int64TypeName(Int64Type);
+	inline const FLazyName UInt8TypeName(UInt8Type);
+	inline const FLazyName UInt32TypeName(UInt32Type);
+	inline const FLazyName UInt64TypeName(UInt64Type);
+	inline const FLazyName FNameTypeName(FNameType);
+	inline const FLazyName FStringTypeName(FStringType);
+	inline const FLazyName FTextTypeName(FTextType);
+	inline const FLazyName BoolArrayTypeName(BoolArrayType);
+	inline const FLazyName FloatArrayTypeName(FloatArrayType);
+	inline const FLazyName DoubleArrayTypeName(DoubleArrayType);
+	inline const FLazyName Int32ArrayTypeName(Int32ArrayType);
+	inline const FLazyName UInt32ArrayTypeName(UInt32ArrayType);
+	inline const FLazyName UInt8ArrayTypeName(UInt8ArrayType);
+	inline const FLazyName FNameArrayTypeName(FNameArrayType);
+	inline const FLazyName FStringArrayTypeName(FStringArrayType);
 
 	class TypeIndex
 	{
@@ -195,9 +195,9 @@ namespace RigVMTypeUtils
 		return WildCardCPPType;
 	}
 
-	static const FLazyName& GetWildCardCPPTypeName()
+	static const FName& GetWildCardCPPTypeName()
 	{
-		static const FLazyName WildCardCPPTypeName(*GetWildCardCPPType()); 
+		static const FName WildCardCPPTypeName(*GetWildCardCPPType()); 
 		return WildCardCPPTypeName;
 	}
 
@@ -207,9 +207,9 @@ namespace RigVMTypeUtils
 		return WildCardArrayCPPType;
 	}
 
-	static const FLazyName& GetWildCardArrayCPPTypeName()
+	static const FName& GetWildCardArrayCPPTypeName()
 	{
-		static const FLazyName WildCardArrayCPPTypeName(*GetWildCardArrayCPPType()); 
+		static const FName WildCardArrayCPPTypeName(*GetWildCardArrayCPPType()); 
 		return WildCardArrayCPPTypeName;
 	}
 
