RigVM: Fix sub pin folding edge case with functions within functions

#rnx

diff --git a/Engine/Plugins/Runtime/RigVM/Source/RigVM/Private/RigVMCore/RigVM.cpp b/Engine/Plugins/Runtime/RigVM/Source/RigVM/Private/RigVMCore/RigVM.cpp
--- a/Engine/Plugins/Runtime/RigVM/Source/RigVM/Private/RigVMCore/RigVM.cpp
+++ b/Engine/Plugins/Runtime/RigVM/Source/RigVM/Private/RigVMCore/RigVM.cpp
@@ -3042,7 +3042,11 @@ void URigVM::CacheSingleMemoryHandle(FRigVMExtendedExecuteContext& Context, cons
 	{
 		check(Memory->GetPropertyPaths().IsValidIndex(InArg.GetRegisterOffset()));
 		PropertyPath = &Memory->GetPropertyPaths()[InArg.GetRegisterOffset()];
-		check(!PropertyPath->IsEmpty());
+		if (PropertyPath->IsEmpty())
+		{
+			UE_LOG(LogRigVM, Error, TEXT("VM '%s' uses an invalid property path."), *GetPathName());
+			PropertyPath = nullptr;
+		}
 	}
 
 	// if you are hitting this it's likely that the VM was created outside of a valid
diff --git a/Engine/Plugins/Runtime/RigVM/Source/RigVMDeveloper/Private/RigVMCompiler/RigVMCompiler.cpp b/Engine/Plugins/Runtime/RigVM/Source/RigVMDeveloper/Private/RigVMCompiler/RigVMCompiler.cpp
--- a/Engine/Plugins/Runtime/RigVM/Source/RigVMDeveloper/Private/RigVMCompiler/RigVMCompiler.cpp
+++ b/Engine/Plugins/Runtime/RigVM/Source/RigVMDeveloper/Private/RigVMCompiler/RigVMCompiler.cpp
@@ -3080,7 +3080,7 @@ bool URigVMCompiler::TraverseInlineFunction(const FRigVMInlineFunctionExprAST* I
 						{
 							// in this case both the external operand and the operand within the function have a register offset.
 							// they need to be combined.
-							const FRigVMPropertyPathDescription& Description = WorkData.PropertyPathDescriptions.FindOrAdd(Operand->GetMemoryType())[Operand->GetRegisterOffset()];
+							const FRigVMPropertyPathDescription& Description = WorkData.PropertyPathDescriptions.FindOrAdd(InterfaceOperand.GetMemoryType())[InterfaceOperand.GetRegisterOffset()];
 							const FString CombinedSegmentPath = URigVMPin::JoinPinPath(Description.SegmentPath, DescriptionFromFunction.SegmentPath);
 							Operand->RegisterOffset = IntCastChecked<uint16>(WorkData.FindOrAddPropertyPath(*Operand, Description.HeadCPPType, CombinedSegmentPath));
 						}
