RigVM: Fix localized registry not using correct permutations

#rnx

diff --git a/Engine/Plugins/Runtime/RigVM/Source/RigVM/Private/RigVMCore/RigVM.cpp b/Engine/Plugins/Runtime/RigVM/Source/RigVM/Private/RigVMCore/RigVM.cpp
--- a/Engine/Plugins/Runtime/RigVM/Source/RigVM/Private/RigVMCore/RigVM.cpp
+++ b/Engine/Plugins/Runtime/RigVM/Source/RigVM/Private/RigVMCore/RigVM.cpp
@@ -2639,8 +2639,7 @@ bool URigVM::CreateLocalizedRegistryIfRequired(bool bForce)
 		// flatten the argument infos to the permutations used only
 		TArray<FRigVMTemplateArgumentInfo> TemporaryArgumentInfos;
 		{
-			FRigVMRegistryReadLock ReadLock;
-			const TArray<FRigVMTemplateArgumentInfo> FlattenedArgumentInfos = Factory->GetTemplate_NoLock(ReadLock)->GetFlattenedArgumentInfosForFunctions_NoLock(Pair.Value, ReadLock);
+			FRigVMRegistryReadLock ReadLock;const TArray<FRigVMTemplateArgumentInfo> FlattenedArgumentInfos = Factory->GetTemplate_NoLock(ReadLock)->GetFlattenedArgumentInfosForFunctions_NoLock(Pair.Value, ReadLock);
 			TemporaryArgumentInfos.Reserve(FlattenedArgumentInfos.Num());
 
 			for(const FRigVMTemplateArgumentInfo& ArgumentInfo : FlattenedArgumentInfos)
@@ -2651,8 +2650,11 @@ bool URigVM::CreateLocalizedRegistryIfRequired(bool bForce)
 					const FRigVMTemplateArgumentType Type = ReadLock.GetRegistry().GetType_NoLock(TypeIndices[TypeIndex]);
 					TypeIndices[TypeIndex] = TemporaryRegistry.FindOrAddType_NoLock(Type);
 					check(TypeIndices[TypeIndex] != INDEX_NONE);
+
+					const FRigVMTemplateArgumentType TemporaryType = TemporaryRegistry.GetType_NoLock(TypeIndices[TypeIndex]);
+					check(Type == TemporaryType);
 				}
-				TemporaryArgumentInfos.Emplace(ArgumentInfo.Name, ArgumentInfo.Direction, TypeIndices);
+				TemporaryArgumentInfos.Add(FRigVMTemplateArgumentInfo::MakeFlattenedFromTypeIndices(ArgumentInfo.Name, ArgumentInfo.Direction, TypeIndices));
 			}
 		}
 
diff --git a/Engine/Plugins/Runtime/RigVM/Source/RigVM/Private/RigVMCore/RigVMDispatchFactory.cpp b/Engine/Plugins/Runtime/RigVM/Source/RigVM/Private/RigVMCore/RigVMDispatchFactory.cpp
--- a/Engine/Plugins/Runtime/RigVM/Source/RigVM/Private/RigVMCore/RigVMDispatchFactory.cpp
+++ b/Engine/Plugins/Runtime/RigVM/Source/RigVM/Private/RigVMCore/RigVMDispatchFactory.cpp
@@ -253,6 +253,8 @@ const FRigVMTemplate* FRigVMDispatchFactory::CreateTemplateForArgumentInfos_NoLo
 	
 	CachedTemplate = InRegistry->AddTemplateFromArguments_NoLock(GetFactoryName(), InArguments, Delegates);
 	(void)UpdateArgumentNameCache_NoLock(InArguments.Num(), OwnerRegistry->Context);
+
+	CachedArgumentInfos = InArguments;
 	
 	return CachedTemplate;
 }
diff --git a/Engine/Plugins/Runtime/RigVM/Source/RigVM/Private/RigVMCore/RigVMFunction.cpp b/Engine/Plugins/Runtime/RigVM/Source/RigVM/Private/RigVMCore/RigVMFunction.cpp
--- a/Engine/Plugins/Runtime/RigVM/Source/RigVM/Private/RigVMCore/RigVMFunction.cpp
+++ b/Engine/Plugins/Runtime/RigVM/Source/RigVM/Private/RigVMCore/RigVMFunction.cpp
@@ -148,8 +148,10 @@ const TArray<TRigVMTypeIndex>& FRigVMFunction::GetArgumentTypeIndices_NoLock(FRi
 				}
 			}
 		}
-		else if(const FRigVMTemplate* Template = GetTemplate_NoLock(InRegistry))
+		else if(InRegistry->GetTemplates_NoLock().IsValidIndex(TemplateIndex))
 		{
+			const FRigVMTemplate* Template = &InRegistry->GetTemplates_NoLock()[TemplateIndex];
+			
 			check(PermutationIndex != INDEX_NONE);
 			check(PermutationIndex == Template->FindPermutation_NoLock(this, InRegistry));
 
diff --git a/Engine/Plugins/Runtime/RigVM/Source/RigVM/Private/RigVMCore/RigVMTemplate.cpp b/Engine/Plugins/Runtime/RigVM/Source/RigVM/Private/RigVMCore/RigVMTemplate.cpp
--- a/Engine/Plugins/Runtime/RigVM/Source/RigVM/Private/RigVMCore/RigVMTemplate.cpp
+++ b/Engine/Plugins/Runtime/RigVM/Source/RigVM/Private/RigVMCore/RigVMTemplate.cpp
@@ -813,7 +813,6 @@ FRigVMTemplateArgumentInfo::FRigVMTemplateArgumentInfo(
 FRigVMTemplateArgumentInfo::FRigVMTemplateArgumentInfo(const FName InName, ERigVMPinDirection InDirection)
 	: Name(InName)
 	, Direction(InDirection)
-	, FactoryCallback([](const FName InName, ERigVMPinDirection InDirection, const FRigVMRegistryHandle&){ return FRigVMTemplateArgument(InName, InDirection); })
 {}
 
 FRigVMTemplateArgumentInfo::FRigVMTemplateArgumentInfo(const FName InName, ERigVMPinDirection InDirection, ArgumentCallback&& InCallback)
@@ -822,6 +821,13 @@ FRigVMTemplateArgumentInfo::FRigVMTemplateArgumentInfo(const FName InName, ERigV
 	, FactoryCallback(InCallback)
 {}
 
+FRigVMTemplateArgumentInfo FRigVMTemplateArgumentInfo::MakeFlattenedFromTypeIndices(const FName InName, ERigVMPinDirection InDirection, const TArray<TRigVMTypeIndex>& InTypeIndices)
+{
+	FRigVMTemplateArgumentInfo Info(InName, InDirection);
+	Info.TypeIndices = InTypeIndices;
+	return Info;
+}
+
 FRigVMTemplateArgument FRigVMTemplateArgumentInfo::GetArgument() const
 {
 	return GetArgument_NoLock(FRigVMRegistryReadLock());
@@ -829,6 +835,12 @@ FRigVMTemplateArgument FRigVMTemplateArgumentInfo::GetArgument() const
 
 FRigVMTemplateArgument FRigVMTemplateArgumentInfo::GetArgument_NoLock(const FRigVMRegistryHandle& InRegistry) const
 {
+	if (!TypeIndices.IsEmpty())
+	{
+		return FRigVMTemplateArgument(Name, Direction, TypeIndices, InRegistry);
+	}
+	
+	check (FactoryCallback)
 	FRigVMTemplateArgument Argument = FactoryCallback(Name, Direction, InRegistry);
 	return MoveTemp(Argument);
 }
@@ -2645,59 +2657,71 @@ TArray<FRigVMTemplateArgumentInfo> FRigVMTemplate::GetFlattenedArgumentInfos_NoL
 {
 	check(InRegistry.GetRegistry() == OwnerRegistry);
 
-	TArray<int32> PermutationIndices;
-	PermutationIndices.Reserve(Permutations.Num());
-	for(int32 PermutationIndex = 0; PermutationIndex < Permutations.Num(); PermutationIndex++)
+	TArray<FRigVMTemplateArgumentInfo> Infos;
+	for (const FRigVMTemplateArgument& Argument : Arguments)
 	{
-		PermutationIndices.Add(PermutationIndex);
+		TArray<TRigVMTypeIndex> TypeIndices;
+		Argument.GetAllTypes_NoLock(TypeIndices, InRegistry);
+		check(!TypeIndices.IsEmpty());
+		
+		if (Argument.IsSingleton_NoLock({}, InRegistry))
+		{
+			Infos.Add(FRigVMTemplateArgumentInfo::MakeFlattenedFromTypeIndices(Argument.Name, Argument.Direction, {TypeIndices[0]}));
+		}
+		else
+		{
+			Infos.Add(FRigVMTemplateArgumentInfo::MakeFlattenedFromTypeIndices(Argument.Name, Argument.Direction, TypeIndices));
+		}
 	}
-	return GetFlattenedArgumentInfosForPermutations_NoLock(PermutationIndices, InRegistry);
+	return Infos;
 }
 
 TArray<FRigVMTemplateArgumentInfo> FRigVMTemplate::GetFlattenedArgumentInfosForFunctions_NoLock(const TArray<const FRigVMFunction*>& InFunctions, FRigVMRegistryHandle& InRegistry) const
 {
 	check(InRegistry.GetRegistry() == OwnerRegistry);
 
-	TArray<int32> PermutationIndices;
-	PermutationIndices.Reserve(InFunctions.Num());
+	TMap<FName, TArray<TRigVMTypeIndex>> TypeIndices;
+	
 	for(const FRigVMFunction* Function : InFunctions)
 	{
 		check(Function->OwnerRegistry == this->OwnerRegistry);
 		check(Function->TemplateIndex == this->Index);
-		PermutationIndices.Add(Function->PermutationIndex);
-	}
-	return GetFlattenedArgumentInfosForPermutations_NoLock(PermutationIndices, InRegistry);
-}
 
-TArray<FRigVMTemplateArgumentInfo> FRigVMTemplate::GetFlattenedArgumentInfosForPermutations_NoLock(const TArray<int32>& InPermutations, FRigVMRegistryHandle& InRegistry) const
-{
-	if(InPermutations.IsEmpty())
-	{
-		return {};
-	}
-	
-	FRigVMTemplateTypeMap Types;
-	TMap<FName, TArray<TRigVMTypeIndex>> TypeIndices;
-	
-	for(const int32& PermutationIndex : InPermutations)
-	{
-		Types = GetTypesForPermutation_NoLock(PermutationIndex, InRegistry);
-		for(const TPair<FName, TRigVMTypeIndex>& Pair : Types)
+		const TArray<TRigVMTypeIndex>& TypeIndicesForFunction = Function->GetArgumentTypeIndices_NoLock(InRegistry);
+		for (int32 ArgumentIndex = 0; ArgumentIndex < Function->Arguments.Num(); ArgumentIndex++)
 		{
-			TypeIndices.FindOrAdd(Pair.Key).Add(Pair.Value);
+			TypeIndices.FindOrAdd(Function->Arguments[ArgumentIndex].Name).Add(TypeIndicesForFunction[ArgumentIndex]);
 		}
 	}
 
 	TArray<FRigVMTemplateArgumentInfo> Infos;
 	for(const FRigVMTemplateArgument& Argument : Arguments)
 	{
-		if(Argument.IsSingleton_NoLock(InPermutations, InRegistry))
+		TArray<TRigVMTypeIndex> ArgumentTypeIndices = TypeIndices.FindChecked(Argument.Name);
+		check(!ArgumentTypeIndices.IsEmpty());
+
+		bool bIsSingleton = false;
+		if(ArgumentTypeIndices.Num() > 1)
+		{
+			bIsSingleton = true;
+			const TRigVMTypeIndex FirstTypeIndex = ArgumentTypeIndices[0];
+			for (const TRigVMTypeIndex& TypeIndex : ArgumentTypeIndices)
+			{
+				if (FirstTypeIndex != TypeIndex)
+				{
+					bIsSingleton = false;
+					break;
+				}
+			}
+		}
+	
+		if(bIsSingleton)
 		{
-			Infos.Emplace(Argument.Name, Argument.Direction, TypeIndices.FindChecked(Argument.Name)[0]);
+			Infos.Add(FRigVMTemplateArgumentInfo::MakeFlattenedFromTypeIndices(Argument.Name, Argument.Direction, {ArgumentTypeIndices[0]}));
 		}
 		else
 		{
-			Infos.Emplace(Argument.Name, Argument.Direction, TypeIndices.FindChecked(Argument.Name));
+			Infos.Add(FRigVMTemplateArgumentInfo::MakeFlattenedFromTypeIndices(Argument.Name, Argument.Direction, ArgumentTypeIndices));
 		}
 	}
 
diff --git a/Engine/Plugins/Runtime/RigVM/Source/RigVM/Public/RigVMCore/RigVMTemplate.h b/Engine/Plugins/Runtime/RigVM/Source/RigVM/Public/RigVMCore/RigVMTemplate.h
--- a/Engine/Plugins/Runtime/RigVM/Source/RigVM/Public/RigVMCore/RigVMTemplate.h
+++ b/Engine/Plugins/Runtime/RigVM/Source/RigVM/Public/RigVMCore/RigVMTemplate.h
@@ -370,12 +370,15 @@ private:
 
 struct FRigVMTemplateArgumentInfo
 {
+public:
+	
 	using ArgumentCallback = TFunction<FRigVMTemplateArgument(const FName /*InName*/, ERigVMPinDirection /*InDirection*/, const FRigVMRegistryHandle& /*InContext*/)>;
 	using TypeFilterCallback = TFunction<bool(const TRigVMTypeIndex& /*InNewType*/, const FRigVMRegistryHandle& /*InContext*/)>;
 	
 	FName Name = NAME_None;
 	ERigVMPinDirection Direction = ERigVMPinDirection::Invalid;
 	ArgumentCallback FactoryCallback = [](const FName, ERigVMPinDirection, const FRigVMRegistryHandle&) { return FRigVMTemplateArgument(); };
+	TArray<TRigVMTypeIndex> TypeIndices;
 	
 	RIGVM_API FRigVMTemplateArgumentInfo(const FName InName, ERigVMPinDirection InDirection, TRigVMTypeIndex InTypeIndex);
 	RIGVM_API FRigVMTemplateArgumentInfo(const FName InName, ERigVMPinDirection InDirection, const TArray<TRigVMTypeIndex>& InTypeIndices);
@@ -383,6 +386,8 @@ struct FRigVMTemplateArgumentInfo
 	RIGVM_API FRigVMTemplateArgumentInfo(const FName InName, ERigVMPinDirection InDirection);
 	RIGVM_API FRigVMTemplateArgumentInfo(const FName InName, ERigVMPinDirection InDirection, ArgumentCallback&& InCallback);
 
+	RIGVM_API static FRigVMTemplateArgumentInfo MakeFlattenedFromTypeIndices(const FName InName, ERigVMPinDirection InDirection, const TArray<TRigVMTypeIndex>& InTypeIndices);
+
 	RIGVM_API FRigVMTemplateArgument GetArgument() const;
 	RIGVM_API FRigVMTemplateArgument GetArgument_NoLock(const FRigVMRegistryHandle& InRegistry) const;
 
@@ -398,7 +403,6 @@ struct FRigVMTemplateArgumentInfo
     	const TArray<FRigVMTemplateArgument::ETypeCategory>& InTypeCategories,
     	const FRigVMTemplateArgument::FTypeFilter& InTypeFilter,
     	const FRigVMRegistryHandle& InRegistry);
-
 };
 
 /**
@@ -591,7 +595,6 @@ public:
 
 	RIGVM_API TArray<FRigVMTemplateArgumentInfo> GetFlattenedArgumentInfos_NoLock(FRigVMRegistryHandle& InRegistry) const;
 	RIGVM_API TArray<FRigVMTemplateArgumentInfo> GetFlattenedArgumentInfosForFunctions_NoLock(const TArray<const FRigVMFunction*>& InFunctions, FRigVMRegistryHandle& InRegistry) const;
-	RIGVM_API TArray<FRigVMTemplateArgumentInfo> GetFlattenedArgumentInfosForPermutations_NoLock(const TArray<int32>& InPermutations, FRigVMRegistryHandle& InRegistry) const;
 
 private:
 
