Fix for duplcated Control Rig assets not reseting the asset variant ID

diff --git a/Engine/Plugins/Runtime/RigVM/Source/RigVMDeveloper/Private/RigVMBlueprint.cpp b/Engine/Plugins/Runtime/RigVM/Source/RigVMDeveloper/Private/RigVMBlueprint.cpp
--- a/Engine/Plugins/Runtime/RigVM/Source/RigVMDeveloper/Private/RigVMBlueprint.cpp
+++ b/Engine/Plugins/Runtime/RigVM/Source/RigVMDeveloper/Private/RigVMBlueprint.cpp
@@ -2632,6 +2632,8 @@ void URigVMBlueprint::PostDuplicate(bool bDuplicateForPIE)
 	PreDuplicateAssetPath.Reset();
 	PreDuplicateHostPath.Reset();
 
+	SplitAssetVariant();
+
 	FBlueprintEditorUtils::MarkBlueprintAsStructurallyModified(this);
 
 	// now that the data in the new BP is correct, make sure there is a call to PostLoad()
