RigVM: Visit object properties for external dependenceis

#rnx

diff --git a/Engine/Plugins/Runtime/RigVM/Source/RigVM/Private/RigVMHost.cpp b/Engine/Plugins/Runtime/RigVM/Source/RigVM/Private/RigVMHost.cpp
--- a/Engine/Plugins/Runtime/RigVM/Source/RigVM/Private/RigVMHost.cpp
+++ b/Engine/Plugins/Runtime/RigVM/Source/RigVM/Private/RigVMHost.cpp
@@ -1186,6 +1186,33 @@ TArray<const UObject*> URigVMHost::GetUserDefinedDependencies(const TArray<const
 					Dependencies.AddUnique(UserDefinedEnum);
 				}
 			}
+			else if (const FObjectProperty* ObjectProperty = CastField<FObjectProperty>(PropertyToVisit))
+			{
+				// if we are a root property - just get the value
+				if (PropertyToVisit == Property)
+				{
+					const UObject* Object = ObjectProperty->GetObjectPropertyValue_InContainer(Memory);
+					if (Object && IsValid(Object))
+					{
+						Dependencies.AddUnique(Object);
+					}
+				}
+				else if (const FArrayProperty* ArrayProperty = CastField<FArrayProperty>(Property)) // if we are within an array property
+				{
+					if (const FObjectProperty* InnerObjectProperty = CastField<FObjectProperty>(ArrayProperty->Inner))
+					{
+						FScriptArrayHelper_InContainer ArrayHelper(ArrayProperty, Memory);
+						for (int32 ElementIndex = 0; ElementIndex < ArrayHelper.Num(); ElementIndex++)
+						{
+							UObject* Object = InnerObjectProperty->GetObjectPropertyValue(ArrayHelper.GetRawPtr(ElementIndex));
+							if (Object && IsValid(Object))
+							{
+								Dependencies.AddUnique(Object);
+							}
+						}
+					}
+				}
+			}
 		}
 	};
 
