RigVM: Visit object properties for external dependenceis

#rnx

diff --git a/Engine/Plugins/Runtime/RigVM/Source/RigVM/Private/RigVMCore/RigVMMemoryStorageStruct.cpp b/Engine/Plugins/Runtime/RigVM/Source/RigVM/Private/RigVMCore/RigVMMemoryStorageStruct.cpp
--- a/Engine/Plugins/Runtime/RigVM/Source/RigVM/Private/RigVMCore/RigVMMemoryStorageStruct.cpp
+++ b/Engine/Plugins/Runtime/RigVM/Source/RigVM/Private/RigVMCore/RigVMMemoryStorageStruct.cpp
@@ -5,6 +5,7 @@
 #include "RigVMModule.h"
 #include "Misc/Guid.h"
 #include "Hash/Blake3.h"
+#include "UObject/Object.h"
 
 #include UE_INLINE_GENERATED_CPP_BY_NAME(RigVMMemoryStorageStruct)
 
@@ -66,6 +67,67 @@ void FRigVMMemoryStorageStruct::AddStructReferencedObjects(FReferenceCollector&
 	Super::AddStructReferencedObjects(Collector);
 }
 
+void FRigVMMemoryStorageStruct::GetUserDefinedDependencies(TArray<const UObject*>& OutDependencies) const
+{
+	const TArray<const FProperty*>& Properties = GetProperties();
+	for (const FProperty* Property : Properties)
+	{
+		const FProperty* PropertyToVisit = Property;
+		while (const FArrayProperty* ArrayProperty = CastField<FArrayProperty>(PropertyToVisit))
+		{
+			PropertyToVisit = ArrayProperty->Inner;
+		}
+		if (const FStructProperty* StructProperty = CastField<FStructProperty>(PropertyToVisit))
+		{
+			if (const UUserDefinedStruct* UserDefinedStruct = Cast<UUserDefinedStruct>(StructProperty->Struct))
+			{
+				OutDependencies.AddUnique(UserDefinedStruct);
+			}
+		}
+		else if (const FEnumProperty* EnumProperty = CastField<FEnumProperty>(PropertyToVisit))
+		{
+			if (const UUserDefinedEnum* UserDefinedEnum = Cast<UUserDefinedEnum>(EnumProperty->GetEnum()))
+			{
+				OutDependencies.AddUnique(UserDefinedEnum);
+			}
+		}
+		else if (const FByteProperty* ByteProperty = CastField<FByteProperty>(PropertyToVisit))
+		{
+			if (const UUserDefinedEnum* UserDefinedEnum = Cast<UUserDefinedEnum>(ByteProperty->Enum))
+			{
+				OutDependencies.AddUnique(UserDefinedEnum);
+			}
+		}
+		else if (const FObjectProperty* ObjectProperty = CastField<FObjectProperty>(PropertyToVisit))
+		{
+			// if we are a root property - just get the value
+			if (PropertyToVisit == Property)
+			{
+				const UObject* Dependency = ObjectProperty->GetObjectPropertyValue_InContainer(Value.GetMemory());
+				if (Dependency && FInternalUObjectBaseUtilityIsValidFlagsChecker::CheckObjectValidBasedOnItsFlags(Dependency))
+				{
+					OutDependencies.AddUnique(Dependency);
+				}
+			}
+			else if (const FArrayProperty* ArrayProperty = CastField<FArrayProperty>(Property)) // if we are within an array property
+			{
+				if (const FObjectProperty* InnerObjectProperty = CastField<FObjectProperty>(ArrayProperty->Inner))
+				{
+					FScriptArrayHelper_InContainer ArrayHelper(ArrayProperty, Value.GetMemory());
+					for (int32 ElementIndex = 0; ElementIndex < ArrayHelper.Num(); ElementIndex++)
+					{
+						UObject* Dependency = InnerObjectProperty->GetObjectPropertyValue(ArrayHelper.GetRawPtr(ElementIndex));
+						if (Dependency && FInternalUObjectBaseUtilityIsValidFlagsChecker::CheckObjectValidBasedOnItsFlags(Dependency))
+						{
+							OutDependencies.AddUnique(Dependency);
+						}
+					}
+				}
+			}
+		}
+	}
+}
+
 void FRigVMMemoryStorageStruct::AddProperties(const TArray<FRigVMPropertyDescription>& InPropertyDescriptions, const TArray<FRigVMPropertyPathDescription>& InPropertyPathDescriptions)
 {
 	TArray<FPropertyBagPropertyDesc> BagDescriptors;
diff --git a/Engine/Plugins/Runtime/RigVM/Source/RigVM/Private/RigVMHost.cpp b/Engine/Plugins/Runtime/RigVM/Source/RigVM/Private/RigVMHost.cpp
--- a/Engine/Plugins/Runtime/RigVM/Source/RigVM/Private/RigVMHost.cpp
+++ b/Engine/Plugins/Runtime/RigVM/Source/RigVM/Private/RigVMHost.cpp
@@ -1149,49 +1149,12 @@ void URigVMHost::GenerateUserDefinedDependenciesData(FRigVMExtendedExecuteContex
 TArray<const UObject*> URigVMHost::GetUserDefinedDependencies(const TArray<const FRigVMMemoryStorageStruct*> InMemory)
 {
 	TArray<const UObject*> Dependencies;
-	auto ProcessMemory = [&Dependencies](const FRigVMMemoryStorageStruct* Memory)
+	for (const FRigVMMemoryStorageStruct* MemoryStorage : InMemory)
 	{
-		if (Memory == nullptr)
+		if (MemoryStorage)
 		{
-			return;
+			MemoryStorage->GetUserDefinedDependencies(Dependencies);
 		}
-
-		const TArray<const FProperty*>& Properties = Memory->GetProperties();
-
-		for (const FProperty* Property : Properties)
-		{
-			const FProperty* PropertyToVisit = Property;
-			while (const FArrayProperty* ArrayProperty = CastField<FArrayProperty>(PropertyToVisit))
-			{
-				PropertyToVisit = ArrayProperty->Inner;
-			}
-			if (const FStructProperty* StructProperty = CastField<FStructProperty>(PropertyToVisit))
-			{
-				if (const UUserDefinedStruct* UserDefinedStruct = Cast<UUserDefinedStruct>(StructProperty->Struct))
-				{
-					Dependencies.AddUnique(UserDefinedStruct);
-				}
-			}
-			else if (const FEnumProperty* EnumProperty = CastField<FEnumProperty>(PropertyToVisit))
-			{
-				if (const UUserDefinedEnum* UserDefinedEnum = Cast<UUserDefinedEnum>(EnumProperty->GetEnum()))
-				{
-					Dependencies.AddUnique(UserDefinedEnum);
-				}
-			}
-			else if (const FByteProperty* ByteProperty = CastField<FByteProperty>(PropertyToVisit))
-			{
-				if (const UUserDefinedEnum* UserDefinedEnum = Cast<UUserDefinedEnum>(ByteProperty->Enum))
-				{
-					Dependencies.AddUnique(UserDefinedEnum);
-				}
-			}
-		}
-	};
-
-	for (const FRigVMMemoryStorageStruct* MemoryStorage : InMemory)
-	{
-		ProcessMemory(MemoryStorage);
 	}
 
 	const TArray<const FRigVMFunction*>& Functions = VM->GetFunctions();
diff --git a/Engine/Plugins/Runtime/RigVM/Source/RigVM/Public/RigVMCore/RigVMMemoryStorageStruct.h b/Engine/Plugins/Runtime/RigVM/Source/RigVM/Public/RigVMCore/RigVMMemoryStorageStruct.h
--- a/Engine/Plugins/Runtime/RigVM/Source/RigVM/Public/RigVMCore/RigVMMemoryStorageStruct.h
+++ b/Engine/Plugins/Runtime/RigVM/Source/RigVM/Public/RigVMCore/RigVMMemoryStorageStruct.h
@@ -22,6 +22,7 @@ struct FRigVMMemoryStorageStruct : public FInstancedPropertyBag
 
 	UE_API bool Serialize(FArchive& Ar);
 	UE_API void AddStructReferencedObjects(FReferenceCollector& Collector);
+	UE_API void GetUserDefinedDependencies(TArray<const UObject*>& OutDependencies) const;
 
 	friend FArchive& operator<<(FArchive& Ar, FRigVMMemoryStorageStruct& P)
 	{
