RigVM: Enable localized registry as per default

#rnx

diff --git a/Engine/Plugins/Runtime/RigVM/Source/RigVM/Private/RigVMCore/RigVMRegistry.cpp b/Engine/Plugins/Runtime/RigVM/Source/RigVM/Private/RigVMCore/RigVMRegistry.cpp
--- a/Engine/Plugins/Runtime/RigVM/Source/RigVM/Private/RigVMCore/RigVMRegistry.cpp
+++ b/Engine/Plugins/Runtime/RigVM/Source/RigVM/Private/RigVMCore/RigVMRegistry.cpp
@@ -21,7 +21,7 @@
 static FDelayedAutoRegisterHelper GRigVMRegistrySingletonHelper(EDelayedRegisterRunPhase::EndOfEngineInit, &FRigVMRegistry_NoLock::OnEngineInit);
 
 static TAutoConsoleVariable<bool> CVarRigVMUpdateDispatchFactoriesGreedily(TEXT("RigVM.UpdateDispatchFactoriesGreedily"), true, TEXT("Set this to false to avoid loading dispatch factories during engine init / plugin mount."));
-TAutoConsoleVariable<bool> CVarRigVMEnableLocalizedRegistry(TEXT("RigVM.EnableLocalizedRegistry"), false, TEXT("Set this to true to allow VMs to serialize their own localized registries."));
+TAutoConsoleVariable<bool> CVarRigVMEnableLocalizedRegistry(TEXT("RigVM.EnableLocalizedRegistry"), true, TEXT("Set this to true to allow VMs to serialize their own localized registries."));
 
 FRigVMRegistry_NoLock::FRigVMRegistry_NoLock(bool InIsGlobalRegistry)
 	: Context(this)
@@ -216,9 +216,8 @@ uint32 FRigVMRegistry_NoLock::GetHashForProperty_NoLock(const FProperty* InPrope
 {
 	uint32 Hash = GetTypeHash(InProperty->GetName());
 
-	FString ExtendedCPPType;
-	const FString CPPType = InProperty->GetCPPType(&ExtendedCPPType);
-	Hash = HashCombine(Hash, GetTypeHash(CPPType + ExtendedCPPType));
+	const FString CPPType = RigVMTypeUtils::GetCPPTypeFromProperty(InProperty);
+	Hash = HashCombine(Hash, GetTypeHash(CPPType));
 	
 	if(const FArrayProperty* ArrayProperty = CastField<FArrayProperty>(InProperty))
 	{
