UAF: Fix copy-pasting UAF template-driven nodes losing custom color

diff --git a/Engine/Plugins/Runtime/RigVM/Source/RigVMDeveloper/Private/RigVMModel/RigVMController.cpp b/Engine/Plugins/Runtime/RigVM/Source/RigVMDeveloper/Private/RigVMModel/RigVMController.cpp
--- a/Engine/Plugins/Runtime/RigVM/Source/RigVMDeveloper/Private/RigVMModel/RigVMController.cpp
+++ b/Engine/Plugins/Runtime/RigVM/Source/RigVMDeveloper/Private/RigVMModel/RigVMController.cpp
@@ -1739,6 +1739,7 @@ URigVMUnitNode* URigVMController::AddUnitNodeWithDefaults(UScriptStruct* InScrip
 	if (!NodeColorMetadata.IsEmpty())
 	{
 		Node->NodeColor = GetColorFromMetadata(NodeColorMetadata);
+		Node->NodeColorType = ERigVMNodeColorType::FromMetadata;
 	}
 
 	FString ExportedDefaultValue;
@@ -3088,6 +3089,7 @@ URigVMCommentNode* URigVMController::AddCommentNode(const FString& InCommentText
 	Node->Position = InPosition;
 	Node->Size = InSize;
 	Node->NodeColor = InColor;
+	Node->NodeColorType = ERigVMNodeColorType::UserDefined;
 	Node->CommentText = InCommentText;
 
 	if(!AddGraphNode(Node, false))
@@ -7711,6 +7713,7 @@ bool URigVMController::SetNodeColor(URigVMNode* InNode, const FLinearColor& InCo
 	}
 
 	InNode->NodeColor = InColor;
+	InNode->NodeColorType = ERigVMNodeColorType::UserDefined;
 	Notify(ERigVMGraphNotifType::NodeColorChanged, InNode);
 
 	if (URigVMLibraryNode* LibraryNode = Cast<URigVMLibraryNode>(InNode))
@@ -18479,11 +18482,15 @@ bool URigVMController::GenerateNewPinInfos(const FRigVMRegistry& Registry, URigV
 			return false;
 		}
 
-		FString NodeColorMetadata;
-		ScriptStruct->GetStringMetaDataHierarchical(URigVMNode::NodeColorName, &NodeColorMetadata);
-		if (!NodeColorMetadata.IsEmpty())
+		// Only update the node color from struct metadata if it is indicated that the color still comes from there
+		if(UnitNode->NodeColorType == ERigVMNodeColorType::FromMetadata)
 		{
-			UnitNode->NodeColor = GetColorFromMetadata(NodeColorMetadata);
+			FString NodeColorMetadata;
+			ScriptStruct->GetStringMetaDataHierarchical(URigVMNode::NodeColorName, &NodeColorMetadata);
+			if (!NodeColorMetadata.IsEmpty())
+			{
+				UnitNode->NodeColor = GetColorFromMetadata(NodeColorMetadata);
+			}
 		}
 
 		const TSharedPtr<FStructOnScope> DefaultValueContent = UnitNode->ConstructStructInstance(false);
diff --git a/Engine/Plugins/Runtime/RigVM/Source/RigVMDeveloper/Private/RigVMModel/RigVMNode.cpp b/Engine/Plugins/Runtime/RigVM/Source/RigVMDeveloper/Private/RigVMModel/RigVMNode.cpp
--- a/Engine/Plugins/Runtime/RigVM/Source/RigVMDeveloper/Private/RigVMModel/RigVMNode.cpp
+++ b/Engine/Plugins/Runtime/RigVM/Source/RigVMDeveloper/Private/RigVMModel/RigVMNode.cpp
@@ -26,6 +26,7 @@ URigVMNode::URigVMNode()
 , Position(FVector2D::ZeroVector)
 , Size(FVector2D::ZeroVector)
 , NodeColor(FLinearColor::White)
+, NodeColorType(ERigVMNodeColorType::FromMetadata)
 , bHasBreakpoint(false)
 , bHaltedAtThisNode(false)
 , NodeVersion(0)
diff --git a/Engine/Plugins/Runtime/RigVM/Source/RigVMDeveloper/Public/RigVMModel/RigVMNode.h b/Engine/Plugins/Runtime/RigVM/Source/RigVMDeveloper/Public/RigVMModel/RigVMNode.h
--- a/Engine/Plugins/Runtime/RigVM/Source/RigVMDeveloper/Public/RigVMModel/RigVMNode.h
+++ b/Engine/Plugins/Runtime/RigVM/Source/RigVMDeveloper/Public/RigVMModel/RigVMNode.h
@@ -25,6 +25,13 @@ namespace ERigVMNodeDefaultValueOverrideState
 	};
 }
 
+UENUM()
+enum class ERigVMNodeColorType : uint8
+{
+	FromMetadata,	// The color was determined from node metadata
+	UserDefined,	// The color was supplied/edited by a user operation
+};
+
 /**
  * The Node represents a single statement within a Graph. 
  * Nodes can represent values such as Variables / Parameters,
@@ -468,6 +475,9 @@ protected:
 	UPROPERTY()
 	FLinearColor NodeColor;
 
+	UPROPERTY()
+	ERigVMNodeColorType NodeColorType;
+
 	UPROPERTY(transient)
 	FName PreviousName;
 
