[AutoRTFM] Update RigVM registry locking code for transactional safety.

FConditionalScopeLock uses atomics to manage its lock count and lock type, so it needs manually-written rollback support.

#rnx

diff --git a/Engine/Plugins/Runtime/RigVM/Source/RigVM/Private/RigVMCore/RigVMRegistry.cpp b/Engine/Plugins/Runtime/RigVM/Source/RigVM/Private/RigVMCore/RigVMRegistry.cpp
--- a/Engine/Plugins/Runtime/RigVM/Source/RigVM/Private/RigVMCore/RigVMRegistry.cpp
+++ b/Engine/Plugins/Runtime/RigVM/Source/RigVM/Private/RigVMCore/RigVMRegistry.cpp
@@ -2526,17 +2526,9 @@ const TArray<FRigVMFunction>* FRigVMRegistry_NoLock::GetPredicatesForStruct_NoLo
 	return StructNameToPredicates.Find(InStructName);
 }
 
-FRigVMRegistry_RWLock::FRigVMRegistry_RWLock(bool InIsGlobalRegistry)
-	: FRigVMRegistry_NoLock(InIsGlobalRegistry)
+FRigVMRegistry_RWLock::FRigVMRegistry_RWLock(bool bInIsGlobalRegistry)
+	: FRigVMRegistry_NoLock(bInIsGlobalRegistry)
 {
-	Initialize();
-}
-
-void FRigVMRegistry_RWLock::Initialize()
-{
-	LockType = LockType_Invalid;
-	LockCount = 0;
-	
 	FRigVMRegistryWriteLock WriteLock(*this);
 	Initialize_NoLock();
 }
@@ -2562,9 +2554,14 @@ void FRigVMRegistry_RWLock::EnsureLocked(ELockType InLockType)
 	check(InLockType != LockType_Invalid);
 
 	const FRigVMRegistry_RWLock& Registry = Get();
-	const ELockType CurrentLockType = Registry.LockType.load();
+	ELockType CurrentLockType = LockType_Invalid;
+
+	UE_AUTORTFM_OPEN
+	{
+		CurrentLockType = Registry.LockType.load();
+	};
 
-	switch(InLockType)
+	switch (InLockType)
 	{
 		case LockType_Read:
 		{
diff --git a/Engine/Plugins/Runtime/RigVM/Source/RigVM/Public/RigVMCore/RigVMRegistry.h b/Engine/Plugins/Runtime/RigVM/Source/RigVM/Public/RigVMCore/RigVMRegistry.h
--- a/Engine/Plugins/Runtime/RigVM/Source/RigVM/Public/RigVMCore/RigVMRegistry.h
+++ b/Engine/Plugins/Runtime/RigVM/Source/RigVM/Public/RigVMCore/RigVMRegistry.h
@@ -7,6 +7,7 @@
 #include "Containers/Map.h"
 #include "Containers/UnrealString.h"
 #include "CoreMinimal.h"
+#include "Misc/TransactionallySafeRWLock.h"
 #include "RigVMRegistryHandle.h"
 #include "RigVMCore/RigVMTraits.h"
 #include "RigVMDispatchFactory.h"
@@ -470,55 +471,97 @@ protected:
 	{
 	public:
 		UE_NODISCARD_CTOR explicit FConditionalScopeLock(const FRigVMRegistry_RWLock& InRegistry, ELockType InLockType, bool bInLockEnabled = true)
-		: FRigVMRegistryHandleLock(&InRegistry, bInLockEnabled)
-		, Registry(const_cast<FRigVMRegistry_RWLock*>(&InRegistry))
-		, DesiredLockType(InLockType)
+			: FRigVMRegistryHandleLock(&InRegistry, bInLockEnabled)
+			, Registry(const_cast<FRigVMRegistry_RWLock*>(&InRegistry))
+			, DesiredLockType(InLockType)
 		{
-			if(bLockEnabled)
+			if (!bLockEnabled)
 			{
-				int32 CurrentLockCount = 0;
-				if(DesiredLockType == LockType_Read)
-				{
-					Registry->Lock.ReadLock();
+				return;
+			}
 
-					CurrentLockCount = ++Registry->LockCount;
-				}
-				else if(DesiredLockType == LockType_Write)
-				{
-					Registry->Lock.WriteLock();
+			checkSlow(DesiredLockType == LockType_Read || DesiredLockType == LockType_Write);
+
+			// First, take the actual underlying lock.
+			if (DesiredLockType == LockType_Read)
+			{
+				Registry->Lock.ReadLock();
+			}
+			else if (DesiredLockType == LockType_Write)
+			{
+				Registry->Lock.WriteLock();
+			}
 
-					CurrentLockCount = ++Registry->LockCount;
-					ensure(CurrentLockCount == 1);
+			// Increment the matching lock count within the RigVM registry.
+			int32 CurrentLockCount = 0;
+			UE_AUTORTFM_OPEN
+			{
+				CurrentLockCount = ++Registry->LockCount;
+			};
+
+			// Handle rollback for the lock count.
+			AutoRTFM::PushOnAbortHandler(Registry, [Registry = this->Registry]
+			{
+				int32 RollbackLockCount = --Registry->LockCount;
+
+				// If we are rolling back the initial lock, invalidate the lock type.
+				if (RollbackLockCount == 0)
+				{
+					Registry->LockType.store(LockType_Invalid);
 				}
+			});
 
-				if(CurrentLockCount == 1)
+			if (CurrentLockCount == 1)
+			{
+				// If we have taken the initial lock, update the lock type.
+				UE_AUTORTFM_OPEN
 				{
 					Registry->LockType.store(DesiredLockType);
-				}
+				};
+			}
+			else
+			{
+				// We should never have more than one write lock!
+				ensure(DesiredLockType != LockType_Write);
 			}
 		}
 
 		~FConditionalScopeLock()
 		{
-			if(bLockEnabled)
+			if (!bLockEnabled)
 			{
-				if(DesiredLockType == LockType_Read)
-				{
-					const int32 CurrentLockCount = --Registry->LockCount;
-					ensure(CurrentLockCount >= 0);
-					if(CurrentLockCount == 0)
-					{
-						Registry->LockType.store(LockType_Invalid);
-					}
-      				Registry->Lock.ReadUnlock();
-				}
-				else if(DesiredLockType == LockType_Write)
+				return;
+			}
+
+			// Decrement the matching lock count within the RigVM registry.
+			int32 CurrentLockCount = 0;
+			UE_AUTORTFM_OPEN
+			{
+				CurrentLockCount = --Registry->LockCount;
+			};
+
+			// Since we have balanced out the increment above, remove our OnAbort handler.
+			AutoRTFM::PopOnAbortHandler(Registry);
+
+			// If we have released the final lock, invalidate the lock type.
+			if (CurrentLockCount == 0)
+			{
+				UE_AUTORTFM_OPEN
 				{
-					const int32 CurrentLockCount = --Registry->LockCount;
-					ensure(CurrentLockCount == 0);
 					Registry->LockType.store(LockType_Invalid);
-					Registry->Lock.WriteUnlock();
-				}
+				};
+			}
+
+			// Finally, release the actual underlying lock.
+			if (DesiredLockType == LockType_Read)
+			{
+				ensure(CurrentLockCount >= 0);
+				Registry->Lock.ReadUnlock();
+			}
+			else if (DesiredLockType == LockType_Write)
+			{
+				ensure(CurrentLockCount == 0);
+				Registry->Lock.WriteUnlock();
 			}
 		}
 
@@ -1315,16 +1358,15 @@ protected:
 
 private:
 
-	FRigVMRegistry_RWLock(bool InIsGlobalRegistry);
+	FRigVMRegistry_RWLock(bool bInIsGlobalRegistry);
 
 	static void EnsureLocked(ELockType InLockType);
 
-	UE_API void Initialize();
 	UE_API void Initialize_NoLock();
 
-	mutable FRWLock Lock;
-	mutable std::atomic<ELockType> LockType;
-	mutable std::atomic<int32> LockCount;
+	mutable FTransactionallySafeRWLock Lock;
+	mutable std::atomic<ELockType> LockType = LockType_Invalid;
+	mutable std::atomic<int32> LockCount = 0;
 
 	// Notifies other system that types have been added/removed, and template permutations have been updated
 	FOnRigVMRegistryChanged OnRigVMRegistryChangedDelegate;
