RigVM: Remove checkNoEntry() from URigVMEdGraphNode::PinListForPin as this can fire always for hidden pins

diff --git a/Engine/Plugins/Runtime/RigVM/Source/RigVMDeveloper/Private/EdGraph/RigVMEdGraphNode.cpp b/Engine/Plugins/Runtime/RigVM/Source/RigVMDeveloper/Private/EdGraph/RigVMEdGraphNode.cpp
--- a/Engine/Plugins/Runtime/RigVM/Source/RigVMDeveloper/Private/EdGraph/RigVMEdGraphNode.cpp
+++ b/Engine/Plugins/Runtime/RigVM/Source/RigVMDeveloper/Private/EdGraph/RigVMEdGraphNode.cpp
@@ -2313,8 +2313,6 @@ TArray<URigVMPin*>& URigVMEdGraphNode::PinListForPin(const URigVMPin* InModelPin
 		{
 			return OutputPins;
 		}
-
-		checkNoEntry();
 	}
 	static TArray<URigVMPin*> EmptyList;
 	return EmptyList;
