RigVM: Registry to support unloading plugins

#rnx

diff --git a/Engine/Plugins/Runtime/RigVM/Source/RigVM/Private/RigVMCore/RigVMDispatchFactory.cpp b/Engine/Plugins/Runtime/RigVM/Source/RigVM/Private/RigVMCore/RigVMDispatchFactory.cpp
--- a/Engine/Plugins/Runtime/RigVM/Source/RigVM/Private/RigVMCore/RigVMDispatchFactory.cpp
+++ b/Engine/Plugins/Runtime/RigVM/Source/RigVM/Private/RigVMCore/RigVMDispatchFactory.cpp
@@ -252,7 +252,7 @@ const FRigVMTemplate* FRigVMDispatchFactory::CreateTemplateForArgumentInfos_NoLo
 	});
 	
 	CachedTemplate = InRegistry->AddTemplateFromArguments_NoLock(GetFactoryName(), InArguments, Delegates);
-	(void)UpdateArgumentNameCache_NoLock(InArguments.Num(), OwnerRegistry->Context);
+	(void)UpdateArgumentNameCache_NoLock(InArguments.Num(), OwnerRegistry->ThisHandle);
 
 	CachedArgumentInfos = InArguments;
 	
@@ -412,6 +412,10 @@ const FRigVMTemplate* FRigVMDispatchFactory::GetTemplate_NoLock(FRigVMRegistryHa
 	
 	if(CachedTemplate)
 	{
+		if (!CachedTemplate->IsValid())
+		{
+			return nullptr;
+		}
 		return CachedTemplate;
 	}
 	
diff --git a/Engine/Plugins/Runtime/RigVM/Source/RigVM/Private/RigVMCore/RigVMFunction.cpp b/Engine/Plugins/Runtime/RigVM/Source/RigVM/Private/RigVMCore/RigVMFunction.cpp
--- a/Engine/Plugins/Runtime/RigVM/Source/RigVM/Private/RigVMCore/RigVMFunction.cpp
+++ b/Engine/Plugins/Runtime/RigVM/Source/RigVM/Private/RigVMCore/RigVMFunction.cpp
@@ -185,6 +185,11 @@ const FRigVMTemplate* FRigVMFunction::GetTemplate_NoLock(const FRigVMRegistryHan
 	}
 
 	const FRigVMTemplate* Template = &InRegistry->GetTemplates_NoLock()[TemplateIndex];
+	if (!Template->IsValid())
+	{
+		return nullptr;
+	}
+	
 	if(Template->NumPermutations_NoLock(InRegistry) <= 1)
 	{
 		return nullptr;
diff --git a/Engine/Plugins/Runtime/RigVM/Source/RigVM/Private/RigVMCore/RigVMRegistry.cpp b/Engine/Plugins/Runtime/RigVM/Source/RigVM/Private/RigVMCore/RigVMRegistry.cpp
--- a/Engine/Plugins/Runtime/RigVM/Source/RigVM/Private/RigVMCore/RigVMRegistry.cpp
+++ b/Engine/Plugins/Runtime/RigVM/Source/RigVM/Private/RigVMCore/RigVMRegistry.cpp
@@ -24,7 +24,7 @@ static TAutoConsoleVariable<bool> CVarRigVMUpdateDispatchFactoriesGreedily(TEXT(
 TAutoConsoleVariable<bool> CVarRigVMEnableLocalizedRegistry(TEXT("RigVM.EnableLocalizedRegistry"), false, TEXT("Set this to true to allow VMs to serialize their own localized registries."));
 
 FRigVMRegistry_NoLock::FRigVMRegistry_NoLock(bool InIsGlobalRegistry)
-	: Context(this)
+	: ThisHandle(this)
 	, bIsGlobalRegistry(InIsGlobalRegistry)
 	, bAvoidTypePropagation(false)
 	, bEverRefreshedEngineTypes(false)
@@ -46,46 +46,43 @@ FRigVMRegistry_NoLock::~FRigVMRegistry_NoLock()
 
 void FRigVMRegistry_NoLock::AddReferencedObjects(FReferenceCollector& Collector)
 {
-	// registry should hold strong references to these type objects
-	// otherwise GC may remove them without the registry known it
-	// which can happen during cook time.
-	for (FTypeInfo& Type : Types)
+	auto AddReferenceLambda = [&Collector](TObjectPtr<UObject>& Object)
 	{
 		// the Object needs to be checked for validity since it may be a user defined type (struct or enum)
 		// which is about to get removed. 
-		if (Type.Type.CPPTypeObject)
+		if (Object)
 		{
 #if !UE_BUILD_SHIPPING
 			// in non shipping builds, immediately run IsValidLowLevelFast such that
 			// we can catch invalid types earlier via a direct crash more often
-			if (Type.Type.CPPTypeObject->IsValidLowLevelFast())
+			if (Object->IsValidLowLevelFast())
 			{
 				// By design, hold strong references only to non-native types
-				if (!Type.Type.CPPTypeObject->IsNative())
+				if (!Object->IsNative())
 				{
-					Collector.AddReferencedObject(Type.Type.CPPTypeObject);	
+					Collector.AddReferencedObject(Object);
 				}
 			}
 #else
 			// in shipping builds, try to be as safe as possible
-			if(IsValid(Type.Type.CPPTypeObject))
+			if(IsValid(Object))
 			{
-				if(Type.Type.CPPTypeObject->GetClass())
+				if(Object->GetClass())
 				{
-					if(Type.Type.CPPTypeObject->IsValidLowLevelFast() &&
-						!Type.Type.CPPTypeObject->IsNative() &&
-						!Type.Type.CPPTypeObject->IsUnreachable())
+					if(Object->IsValidLowLevelFast() &&
+						!Object->IsNative() &&
+						!Object->IsUnreachable())
 					{
 						// make sure the object is part of the GUObjectArray and can be retrieved
 						// so that GC doesn't crash after receiving the referenced object
-						const int32 ObjectIndex = GUObjectArray.ObjectToIndex(Type.Type.CPPTypeObject);
+						const int32 ObjectIndex = GUObjectArray.ObjectToIndex(Object);
 						if(ObjectIndex != INDEX_NONE)
 						{
 							if(const FUObjectItem* Item = GUObjectArray.IndexToObject(ObjectIndex))
 							{
-								if(Item->GetObject() == Type.Type.CPPTypeObject)
+								if(Item->GetObject() == Object)
 								{
-									Collector.AddReferencedObject(Type.Type.CPPTypeObject);
+									Collector.AddReferencedObject(Object);
 								}
 							}
 						}
@@ -94,6 +91,32 @@ void FRigVMRegistry_NoLock::AddReferencedObjects(FReferenceCollector& Collector)
 			}
 #endif
 		}
+	};
+	
+	// registry should hold strong references to these type objects
+	// otherwise GC may remove them without the registry known it
+	// which can happen during cook time.
+	for (FTypeInfo& Type : Types)
+	{
+		AddReferenceLambda(Type.Type.CPPTypeObject);
+	}
+
+	for (const FRigVMFunction& Function : Functions)
+	{
+		if (Function.IsValid())
+		{
+			TObjectPtr<UObject> StructPtr(Function.Struct);
+			AddReferenceLambda(StructPtr);
+		}
+	}
+
+	for (const FRigVMDispatchFactory* Factory : Factories)
+	{
+		if (Factory)
+		{
+			TObjectPtr<UObject> FactoryPtr(Factory->GetScriptStruct());
+			AddReferenceLambda(FactoryPtr);
+		}
 	}
 }
 
@@ -260,7 +283,7 @@ void FRigVMRegistry_NoLock::RebuildRegistry_NoLock()
 	DeprecatedTemplateNotationToIndex.Reset();
 	TypesPerCategory.Reset();
 	TemplatesPerCategory.Reset();
-	UserDefinedTypeToIndex.Reset();
+	SoftObjectPathToTypeIndex.Reset();
 	AllowedClasses.Reset();
 
 	InitializeBaseTypes_NoLock();
@@ -282,12 +305,12 @@ void FRigVMRegistry_NoLock::FTypeInfo::Serialize_NoLock(FArchive& Ar)
 
 const FRigVMRegistryHandle& FRigVMRegistry_NoLock::GetHandle_NoLock() const
 {
-	return Context;
+	return ThisHandle;
 }
 
 FRigVMRegistryHandle& FRigVMRegistry_NoLock::GetHandle_NoLock()
 {
-	return Context;
+	return ThisHandle;
 }
 
 void FRigVMRegistry_NoLock::InitializeBaseTypes_NoLock()
@@ -519,7 +542,7 @@ void FRigVMRegistry_NoLock::Serialize_NoLock(FArchive& Ar)
 				if(const FRigVMDispatchFactory* Factory = RegisterFactory_NoLock(FactoryScriptStruct, FlattenedArguments))
 				{
 					// pulling on this caches the template
-					Factory->GetTemplate_NoLock(Context);
+					Factory->GetTemplate_NoLock(ThisHandle);
 				}
 			}
 		}
@@ -531,7 +554,7 @@ void FRigVMRegistry_NoLock::Serialize_NoLock(FArchive& Ar)
 			for(int32 FunctionIndex = 0; FunctionIndex < NumFunctions; FunctionIndex++)
 			{
 				FRigVMFunction Function;
-				Function.Serialize_NoLock(Ar, Context);
+				Function.Serialize_NoLock(Ar, ThisHandle);
 
 				if(Function.Struct)
 				{
@@ -569,7 +592,7 @@ void FRigVMRegistry_NoLock::Serialize_NoLock(FArchive& Ar)
 				Ar << FactoryScriptStruct;
 
 				// find all used permutations of the factory and store their types
-				TArray<FRigVMTemplateArgumentInfo> FlattenedArguments = Factory->GetTemplate_NoLock(Context)->GetFlattenedArgumentInfos_NoLock(Context);
+				TArray<FRigVMTemplateArgumentInfo> FlattenedArguments = Factory->GetTemplate_NoLock(ThisHandle)->GetFlattenedArgumentInfos_NoLock(ThisHandle);
 
 				int32 NumFlattenedArguments = FlattenedArguments.Num();
 				Ar << NumFlattenedArguments;
@@ -577,7 +600,7 @@ void FRigVMRegistry_NoLock::Serialize_NoLock(FArchive& Ar)
 				for(int32 ArgumentIndex = 0; ArgumentIndex < NumFlattenedArguments; ArgumentIndex++)
 				{
 					FRigVMTemplateArgumentInfo& ArgumentInfo = FlattenedArguments[ArgumentIndex];
-					FRigVMTemplateArgument Argument = ArgumentInfo.GetArgument_NoLock(Context);
+					FRigVMTemplateArgument Argument = ArgumentInfo.GetArgument_NoLock(ThisHandle);
 					check(!Argument.bUseCategories);
 					
 					Ar << Argument.Name;
@@ -604,7 +627,7 @@ void FRigVMRegistry_NoLock::Serialize_NoLock(FArchive& Ar)
 			Ar << NumFunctions;
 			for(FRigVMFunction& Function : Functions)
 			{
-				Function.Serialize_NoLock(Ar, Context);
+				Function.Serialize_NoLock(Ar, ThisHandle);
 			}
 		}
 	}
@@ -711,7 +734,7 @@ void FRigVMRegistry_NoLock::RefreshEngineTypes_NoLock()
 				if(!TemplateProcessed[TemplateIndex])
 				{
 					FRigVMTemplate& Template = Templates[TemplateIndex];
-					(void)Template.UpdateAllArgumentTypesSlow(GetHandle_NoLock());
+					(void)Template.UpdateAllArgumentTypesSlow(ThisHandle);
 					TemplateProcessed[TemplateIndex] = true;
 				}
 			}
@@ -744,7 +767,7 @@ bool FRigVMRegistry_NoLock::RefreshFunctionsAndDispatches_NoLock()
 		// the game.
 		if(Factory->CachedTemplate == nullptr)
 		{
-			(void)Factory->GetTemplate_NoLock(Context);
+			(void)Factory->GetTemplate_NoLock(ThisHandle);
 			bRegistryChanged = true;
 		}
 	}
@@ -759,7 +782,7 @@ bool FRigVMRegistry_NoLock::OnCleanupUnloadingObjects_NoLock(const TArrayView<UO
 		const UObject* ObjectOuter = Obj ? Obj->GetOuter() : nullptr;
 		if (ObjectOuter && (ObjectOuter->GetClass() == UPackage::StaticClass()))
 		{
-			if (RemoveType_NoLock(FSoftObjectPath(Obj), Obj->GetClass()))
+			if (RemoveType_NoLock(FSoftObjectPath(Obj)))
 			{
 				bRemoved = true;
 			}
@@ -773,17 +796,17 @@ void FRigVMRegistry_NoLock::OnAssetRenamed_NoLock(const FAssetData& InAssetData,
 {
 	const FSoftObjectPath OldPath(InOldObjectPath);
 	
-	if (const TRigVMTypeIndex* TypeIndexPtr = UserDefinedTypeToIndex.Find(OldPath))
+	if (const TRigVMTypeIndexTriple* TypeIndicesPtr = SoftObjectPathToTypeIndex.Find(OldPath))
 	{
-		const TRigVMTypeIndex TypeIndex = *TypeIndexPtr;
-		UserDefinedTypeToIndex.Remove(OldPath);
-		UserDefinedTypeToIndex.Add(InAssetData.ToSoftObjectPath()) = TypeIndex;
+		const TRigVMTypeIndexTriple TypeIndices = *TypeIndicesPtr;
+		SoftObjectPathToTypeIndex.Remove(OldPath);
+		SoftObjectPathToTypeIndex.Add(InAssetData.ToSoftObjectPath()) = TypeIndices;
 	}
 }
 
 bool FRigVMRegistry_NoLock::OnAssetRemoved_NoLock(const FAssetData& InAssetData)
 {
-	return RemoveType_NoLock(InAssetData.ToSoftObjectPath(), InAssetData.GetClass());
+	return RemoveType_NoLock(InAssetData.ToSoftObjectPath());
 }
 
 bool FRigVMRegistry_NoLock::OnPluginLoaded_NoLock(IPlugin& InPlugin)
@@ -796,32 +819,146 @@ bool FRigVMRegistry_NoLock::OnPluginLoaded_NoLock(IPlugin& InPlugin)
 	return RefreshFunctionsAndDispatches_NoLock();
 }
 
-bool FRigVMRegistry_NoLock::OnPluginUnloaded_NoLock(IPlugin& InPlugin)
+TArray<FString> FRigVMRegistry_NoLock::GetPluginModulePrefixes(const IPlugin& InPlugin)
 {
-	const FString PluginContentPath = InPlugin.GetMountedAssetPath();
+	const FPluginDescriptor Descriptor = InPlugin.GetDescriptor();
+	
+	TArray<FString> ModulePrefixes;
+	ModulePrefixes.Reserve(Descriptor.Modules.Num());
 
-	TSet<FSoftObjectPath> PathsToRemove;
-	for (const TPair<FSoftObjectPath, TRigVMTypeIndex>& Item: UserDefinedTypeToIndex)
+	for (const FModuleDescriptor& Module : Descriptor.Modules)
 	{
-		const FSoftObjectPath ObjectPath = Item.Key;
-		const FString PackageName = ObjectPath.GetLongPackageName();
-		
-		if (PackageName.StartsWith(PluginContentPath))
+		static constexpr TCHAR Format[] = TEXT("/Script/%s.");
+		ModulePrefixes.Add(FString::Printf(Format, *Module.Name.ToString()));
+	}
+
+	return ModulePrefixes;
+}
+
+bool FRigVMRegistry_NoLock::IsWithinPlugin(const IPlugin& InPlugin, const UObject* InObject)
+{
+	return IsWithinPlugin(GetPluginModulePrefixes(InPlugin), InObject);
+}
+
+bool FRigVMRegistry_NoLock::IsWithinPlugin(const TArray<FString>& InModulePrefixes, const UObject* InObject)
+{
+	if (InObject == nullptr || !IsValid(InObject) || InModulePrefixes.IsEmpty())
+	{
+		return false;
+	}
+	
+	const FSoftObjectPath SoftObjectPath(InObject);
+	const FString ObjectPathString = SoftObjectPath.ToString();
+	for (const FString& Prefix : InModulePrefixes)
+	{
+		if (ObjectPathString.StartsWith(Prefix))
 		{
-			PathsToRemove.Add(ObjectPath);
+			return true;
 		}
 	}
 
+	return false;
+}
+
+bool FRigVMRegistry_NoLock::OnPluginUnloaded_NoLock(IPlugin& InPlugin)
+{
+	const TArray<FString> ModulePrefixes = GetPluginModulePrefixes(InPlugin);
+
 	bool bRegistryChanged = false;
-	for (FSoftObjectPath ObjectPath: PathsToRemove)
+
+	// remove the dispatch factories defined in the plugin
+	TArray<UScriptStruct*> FactoriesToRemove;
+	for (int32 Index = 0; Index < Factories.Num(); Index++)
+	{
+		if (!IsWithinPlugin(ModulePrefixes, Factories[Index]->GetScriptStruct()))
+		{
+			continue;
+		}
+		FactoriesToRemove.AddUnique(Factories[Index]->GetScriptStruct());
+	}
+	for (UScriptStruct* FactoryToRemove : FactoriesToRemove)
+	{
+		if (RemoveFactory_NoLock(FactoryToRemove))
+		{
+			bRegistryChanged = true;
+		}
+	}
+
+	// remove all templates defined by solely structs within the plugin
+	for (int32 Index = 0; Index < Templates.Num(); Index++)
+	{
+		if (Templates[Index].UsesDispatch())
+		{
+			continue;
+		}
+
+		bool bHasFunctionOutsideOfPlugin = false;
+		for (int32 FunctionIndex = 0; FunctionIndex < Templates[Index].NumPermutations_NoLock(ThisHandle); FunctionIndex++)
+		{
+			const FRigVMFunction* Function = Templates[Index].GetPermutation_NoLock(FunctionIndex, ThisHandle);
+			if (Function == nullptr)
+			{
+				continue;
+			}
+			
+			check(Function->Struct);
+			check(Function->Factory == nullptr);
+
+			if (!IsWithinPlugin(ModulePrefixes, Function->Struct))
+			{
+				bHasFunctionOutsideOfPlugin = true;
+			}
+		}
+
+		if (bHasFunctionOutsideOfPlugin)
+		{
+			continue;
+		}
+
+		if (RemoveTemplate_NoLock(Index))
+		{
+			bRegistryChanged = true;
+		}
+	}
+
+	// remove the functions defined in the plugin
+	for (int32 Index = 0; Index < Functions.Num(); Index++)
+	{
+		if (!Functions[Index].IsValid())
+		{
+			continue;
+		}
+
+		if (Functions[Index].Struct == nullptr)
+		{
+			continue;
+		}
+
+		if (!IsWithinPlugin(ModulePrefixes, Functions[Index].Struct))
+		{
+			continue;
+		}
+
+		if (RemoveFunction_NoLock(Index))
+		{
+			bRegistryChanged = true;
+		}
+	}	
+
+	// remove the types defined in the plugin
+	for (int32 Index = 0; Index < Types.Num(); Index++)
 	{
-		const UClass* ObjectClass = nullptr;
-		if (const UObject* TypeObject = ObjectPath.ResolveObject())
+		if (Types[Index].IsUnknownType())
 		{
-			ObjectClass = TypeObject->GetClass();
+			continue;
+		}
+
+		if (!IsWithinPlugin(ModulePrefixes, Types[Index].Type.CPPTypeObject))
+		{
+			continue;
 		}
 		
-		if (RemoveType_NoLock(ObjectPath, ObjectClass))
+		if (RemoveType_NoLock(Index))
 		{
 			bRegistryChanged = true;
 		}
@@ -909,6 +1046,23 @@ TRigVMTypeIndex FRigVMRegistry_NoLock::FindOrAddType_NoLock(const FRigVMTemplate
 
 		// 3 is the max array dimension allowed within RigVM
 		TArray<TRigVMTypeIndex, TInlineAllocator<3>> Indices;
+
+		// we need to look up the type indices to reuse if this is a type that we've encountered before
+		// from a plugin that got unloaded - and is being loaded again now.
+		if (ElementType.CPPTypeObject)
+		{
+			const FSoftObjectPath SoftObjectPath(ElementType.CPPTypeObject);
+			if (const TRigVMTypeIndexTriple* ExistingIndices = SoftObjectPathToTypeIndex.Find(SoftObjectPath))
+			{
+				check(Types[ExistingIndices->Get<0>()].IsUnknownType());
+				check(Types[ExistingIndices->Get<1>()].IsUnknownType());
+				check(Types[ExistingIndices->Get<2>()].IsUnknownType());
+				Indices.Add(ExistingIndices->Get<0>());
+				Indices.Add(ExistingIndices->Get<1>());
+				Indices.Add(ExistingIndices->Get<2>());
+			}
+		}
+		
 		TArray<bool, TInlineAllocator<3>> AddedTypeIndex;
 		for (int32 ArrayDimension=0; ArrayDimension<MaxArrayDimension; ++ArrayDimension)
 		{
@@ -938,14 +1092,25 @@ TRigVMTypeIndex FRigVMRegistry_NoLock::FindOrAddType_NoLock(const FRigVMTemplate
 					continue;
 				}
 			}
-			
-			Index = Types.Add(Info);
+
+			if (Indices.IsValidIndex(ArrayDimension))
+			{
+				Index = Indices[ArrayDimension];
+				Types[Index] = Info;
+			}
+			else
+			{
+				Index = Types.Add(Info);
+			}
 #if UE_RIGVM_DEBUG_TYPEINDEX
 			Index.Name = Info.Type.CPPType;
 #endif
 			TypeToIndex.Add(CurType, Index);
 
-			Indices.Add(Index);
+			if (!Indices.IsValidIndex(ArrayDimension))
+			{
+				Indices.Add(Index);
+			}
 			AddedTypeIndex.Add(true);
 		}
 
@@ -1136,7 +1301,7 @@ TRigVMTypeIndex FRigVMRegistry_NoLock::FindOrAddType_NoLock(const FRigVMTemplate
 				{
 					// by creating a template argument for the child property
 					// the type will be added by calling ::FindOrAddType_Internal recursively.
-					(void)FRigVMTemplateArgument::Make_NoLock(Property, Context);
+					(void)FRigVMTemplateArgument::Make_NoLock(Property, ThisHandle);
 				}
 #if WITH_EDITOR
 				else
@@ -1155,8 +1320,14 @@ TRigVMTypeIndex FRigVMRegistry_NoLock::FindOrAddType_NoLock(const FRigVMTemplate
 			if (CPPTypeObject->IsA<UUserDefinedStruct>() || CPPTypeObject->IsA<UUserDefinedEnum>())
 			{
 				TRigVMTypeIndex ElementTypeIndex = GetTypeIndex_NoLock(ElementType);
-				// used to track name changes to user defined types, stores the element type index, see RemoveType()
-				UserDefinedTypeToIndex.FindOrAdd(CPPTypeObject) = ElementTypeIndex;
+
+				TRigVMTypeIndexTriple Triple;
+				Triple.Get<0>() = ElementTypeIndex;
+				Triple.Get<1>() = GetArrayTypeFromBaseTypeIndex_NoLock(Triple.Get<0>());
+				Triple.Get<2>() = GetArrayTypeFromBaseTypeIndex_NoLock(Triple.Get<1>());
+				
+				// used to track name changes to user defined types, stores the element type indices, see RemoveType().
+				SoftObjectPathToTypeIndex.FindOrAdd(CPPTypeObject) = Triple;
 			}
 		}
 		
@@ -1170,7 +1341,7 @@ void FRigVMRegistry_NoLock::RegisterTypeInCategory_NoLock(FRigVMTemplateArgument
 {
 	check(InCategory != FRigVMTemplateArgument::ETypeCategory_Invalid);
 
-	TypesPerCategory.FindChecked(InCategory).Add(InTypeIndex);
+	TypesPerCategory.FindChecked(InCategory).AddUnique(InTypeIndex);
 }
 
 void FRigVMRegistry_NoLock::PropagateTypeAddedToCategory_NoLock(const FRigVMTemplateArgument::ETypeCategory InCategory, const TRigVMTypeIndex InTypeIndex)
@@ -1188,119 +1359,152 @@ void FRigVMRegistry_NoLock::PropagateTypeAddedToCategory_NoLock(const FRigVMTemp
 		for(const int32 TemplateIndex : TemplatesToUseType)
 		{
 			FRigVMTemplate& Template = Templates[TemplateIndex];
-			(void)Template.HandlePropagatedArgumentType(InTypeIndex, GetHandle_NoLock());
+			(void)Template.HandlePropagatedArgumentType(InTypeIndex, ThisHandle);
 		}
 	}
 }
 
-bool FRigVMRegistry_NoLock::RemoveType_NoLock(const FSoftObjectPath& InObjectPath, const UClass* InObjectClass)
+bool FRigVMRegistry_NoLock::RemoveType_NoLock(const FSoftObjectPath& InObjectPath)
 {
-	if (const TRigVMTypeIndex* TypeIndexPtr = UserDefinedTypeToIndex.Find(InObjectPath))
+	if (const TRigVMTypeIndexTriple* TypeIndicesPtr = SoftObjectPathToTypeIndex.Find(InObjectPath))
 	{
-		const TRigVMTypeIndex TypeIndex = *TypeIndexPtr;
-		
-		UserDefinedTypeToIndex.Remove(InObjectPath);
-		
-		if(TypeIndex == INDEX_NONE)
+		const TRigVMTypeIndexTriple& TypeIndices = *TypeIndicesPtr;
+		if(TypeIndices.Get<0>() == INDEX_NONE)
 		{
 			return false;
 		}
 
-		check(!IsArrayType_NoLock(TypeIndex));
+		if (RemoveType_NoLock(TypeIndices.Get<0>()))
+		{
+			return true;
+		}
+	}
+	else if (const UObject* ResolvedObject = InObjectPath.ResolveObject())
+	{
+		for (int32 Index = 0; Index < Types.Num(); Index++)
+		{
+			if (Types[Index].Type.CPPTypeObject == ResolvedObject)
+			{
+				if (RemoveType_NoLock(Index))
+				{
+					return true;
+				}
+			}
+		}
+	}
+
+	return false;
+}
+
+bool FRigVMRegistry_NoLock::RemoveType_NoLock(TRigVMTypeIndex InTypeIndex)
+{
+	if (!Types.IsValidIndex(InTypeIndex))
+	{
+		return false;
+	}
+	
+	if (Types[InTypeIndex].IsUnknownType())
+	{
+		return false;
+	}
+
+	check(!IsArrayType_NoLock(InTypeIndex));
 
-		TArray<TRigVMTypeIndex> Indices;
-		Indices.Init(INDEX_NONE, 3);
-		Indices[0] = TypeIndex;
-		Indices[1] = GetArrayTypeFromBaseTypeIndex_NoLock(Indices[0]);
+	TArray<TRigVMTypeIndex> Indices;
+	Indices.Init(INDEX_NONE, 3);
+	Indices[0] = InTypeIndex;
+	Indices[1] = GetArrayTypeFromBaseTypeIndex_NoLock(Indices[0]);
 
-		// any type that can be removed should have 3 entries in the registry
-		if (ensure(Indices[1] != INDEX_NONE))
+	// any type that can be removed should have 3 entries in the registry
+	if (ensure(Indices[1] != INDEX_NONE))
+	{
+		Indices[2] = GetArrayTypeFromBaseTypeIndex_NoLock(Indices[1]);
+	}
+	
+	for (int32 ArrayDimension=0; ArrayDimension<3; ++ArrayDimension)
+	{
+		const TRigVMTypeIndex Index = Indices[ArrayDimension];
+		
+		if (Index == INDEX_NONE)
 		{
-			Indices[2] = GetArrayTypeFromBaseTypeIndex_NoLock(Indices[1]);
+			break;
 		}
 		
-		for (int32 ArrayDimension=0; ArrayDimension<3; ++ArrayDimension)
+		switch(ArrayDimension)
 		{
-			const TRigVMTypeIndex Index = Indices[ArrayDimension];
-			
-			if (Index == INDEX_NONE)
+			default:
+			case 0:
 			{
+				// we remove all user defined types from the table - but we keep the softobject paths for 
+				// C++ defined types. that way we can look up their type indices later in case the types
+				// come back into view. 
+				if (const UUserDefinedStruct* UserDefinedStruct = Cast<UUserDefinedStruct>(Types[Index].Type.CPPTypeObject))
+				{
+					SoftObjectPathToTypeIndex.Remove(UserDefinedStruct);
+				}
+				else if (const UUserDefinedEnum* UserDefinedEnum = Cast<UUserDefinedEnum>(Types[Index].Type.CPPTypeObject))
+				{
+					SoftObjectPathToTypeIndex.Remove(UserDefinedEnum);
+				}
+				else if (Types[Index].Type.CPPTypeObject && IsValid(Types[Index].Type.CPPTypeObject))
+				{
+					// for C++ based types - let's remember where the types where in the first place
+					// so we can put them back.
+					SoftObjectPathToTypeIndex.FindOrAdd(Types[Index].Type.CPPTypeObject) = {Indices[0], Indices[1], Indices[2]};
+				}
+					
+				RemoveTypeInCategory_NoLock(FRigVMTemplateArgument::ETypeCategory_SingleAnyValue, Index);
+				RemoveTypeInCategory_NoLock(FRigVMTemplateArgument::ETypeCategory_SingleSimpleValue, Index);
+				RemoveTypeInCategory_NoLock(FRigVMTemplateArgument::ETypeCategory_SingleMathStructValue, Index);
+				RemoveTypeInCategory_NoLock(FRigVMTemplateArgument::ETypeCategory_SingleScriptStructValue, Index);
+				RemoveTypeInCategory_NoLock(FRigVMTemplateArgument::ETypeCategory_SingleEnumValue, Index);
+				RemoveTypeInCategory_NoLock(FRigVMTemplateArgument::ETypeCategory_SingleObjectValue, Index);
 				break;
 			}
-			
-			if(InObjectClass == UUserDefinedEnum::StaticClass())
+			case 1:
 			{
-				switch(ArrayDimension)
-				{
-				default:
-				case 0:
-					{
-						RemoveTypeInCategory_NoLock(FRigVMTemplateArgument::ETypeCategory_SingleEnumValue, Index);
-						RemoveTypeInCategory_NoLock(FRigVMTemplateArgument::ETypeCategory_SingleAnyValue, Index);
-						break;
-					}
-				case 1:
-					{
-						RemoveTypeInCategory_NoLock(FRigVMTemplateArgument::ETypeCategory_ArrayEnumValue, Index);
-						RemoveTypeInCategory_NoLock(FRigVMTemplateArgument::ETypeCategory_ArrayAnyValue, Index);
-						break;
-					}
-				case 2:
-					{
-						RemoveTypeInCategory_NoLock(FRigVMTemplateArgument::ETypeCategory_ArrayArrayEnumValue, Index);
-						RemoveTypeInCategory_NoLock(FRigVMTemplateArgument::ETypeCategory_ArrayArrayAnyValue, Index);
-						break;
-					}
-				}
+				RemoveTypeInCategory_NoLock(FRigVMTemplateArgument::ETypeCategory_ArrayAnyValue, Index);
+				RemoveTypeInCategory_NoLock(FRigVMTemplateArgument::ETypeCategory_ArraySimpleValue, Index);
+				RemoveTypeInCategory_NoLock(FRigVMTemplateArgument::ETypeCategory_ArrayMathStructValue, Index);
+				RemoveTypeInCategory_NoLock(FRigVMTemplateArgument::ETypeCategory_ArrayScriptStructValue, Index);
+				RemoveTypeInCategory_NoLock(FRigVMTemplateArgument::ETypeCategory_ArrayEnumValue, Index);
+				RemoveTypeInCategory_NoLock(FRigVMTemplateArgument::ETypeCategory_ArrayObjectValue, Index);
+				break;
 			}
-			else if(InObjectClass == UUserDefinedStruct::StaticClass())
+			case 2:
 			{
-				switch(ArrayDimension)
-				{
-				default:
-				case 0:
-					{
-						RemoveTypeInCategory_NoLock(FRigVMTemplateArgument::ETypeCategory_SingleScriptStructValue, Index);
-						RemoveTypeInCategory_NoLock(FRigVMTemplateArgument::ETypeCategory_SingleAnyValue, Index);
-						break;
-					}
-				case 1:
-					{
-						RemoveTypeInCategory_NoLock(FRigVMTemplateArgument::ETypeCategory_ArrayScriptStructValue, Index);
-						RemoveTypeInCategory_NoLock(FRigVMTemplateArgument::ETypeCategory_ArrayAnyValue, Index);
-						break;
-					}
-				case 2:
-					{
-						RemoveTypeInCategory_NoLock(FRigVMTemplateArgument::ETypeCategory_ArrayArrayScriptStructValue, Index);
-						RemoveTypeInCategory_NoLock(FRigVMTemplateArgument::ETypeCategory_ArrayArrayAnyValue, Index);
-						break;
-					}
-				}
+				RemoveTypeInCategory_NoLock(FRigVMTemplateArgument::ETypeCategory_ArrayArrayAnyValue, Index);
+				RemoveTypeInCategory_NoLock(FRigVMTemplateArgument::ETypeCategory_ArrayArraySimpleValue, Index);
+				RemoveTypeInCategory_NoLock(FRigVMTemplateArgument::ETypeCategory_ArrayArrayMathStructValue, Index);
+				RemoveTypeInCategory_NoLock(FRigVMTemplateArgument::ETypeCategory_ArrayArrayScriptStructValue, Index);
+				RemoveTypeInCategory_NoLock(FRigVMTemplateArgument::ETypeCategory_ArrayArrayEnumValue, Index);
+				RemoveTypeInCategory_NoLock(FRigVMTemplateArgument::ETypeCategory_ArrayArrayObjectValue, Index);
+				break;
 			}
-
-			// remove the type from the registry entirely
-			TypeToIndex.Remove(GetType_NoLock(Index));
-			Types[Index] = FTypeInfo();
 		}
 
-		return true;	
+		// remove the type from the registry entirely
+		TypeToIndex.Remove(Types[Index].Type);
+		Types[Index] = FTypeInfo();
 	}
-	
-	return false;
+
+	return true;	
 }
 
 void FRigVMRegistry_NoLock::RemoveTypeInCategory_NoLock(FRigVMTemplateArgument::ETypeCategory InCategory, TRigVMTypeIndex InTypeIndex)
 {
 	check(InCategory != FRigVMTemplateArgument::ETypeCategory_Invalid);
 
-	TypesPerCategory.FindChecked(InCategory).Remove(InTypeIndex);
+	if (!TypesPerCategory.FindChecked(InCategory).Contains(InTypeIndex))
+	{
+		return;
+	}
 
 	const TArray<int32>& TemplatesToUseType = TemplatesPerCategory.FindChecked(InCategory);
 	for (const int32 TemplateIndex : TemplatesToUseType)
 	{
 		FRigVMTemplate& Template = Templates[TemplateIndex];
-		Template.HandleTypeRemoval(InTypeIndex, GetHandle_NoLock());
+		Template.HandleTypeRemoval(InTypeIndex, ThisHandle);
 	}
 }
 
@@ -1321,7 +1525,7 @@ FRigVMRegistry_RWLock& FRigVMRegistry_RWLock::Get()
 
 void FRigVMRegistry_RWLock::OnCleanupUnloadingObjects(const TArrayView<UObject*> InObjects)
 {
-	bool bAssetRemoved = false;
+	bool bAssetRemoved;
 	{
 		FConditionalWriteScopeLock _(*this);
 		bAssetRemoved = Super::OnCleanupUnloadingObjects_NoLock(InObjects);
@@ -1335,7 +1539,7 @@ void FRigVMRegistry_RWLock::OnCleanupUnloadingObjects(const TArrayView<UObject*>
 
 void FRigVMRegistry_RWLock::OnAssetRemoved(const FAssetData& InAssetData)
 {
-	bool bAssetRemoved = false;
+	bool bAssetRemoved;
 	{
 		FConditionalWriteScopeLock _(*this);
 		bAssetRemoved = Super::OnAssetRemoved_NoLock(InAssetData);
@@ -1349,7 +1553,7 @@ void FRigVMRegistry_RWLock::OnAssetRemoved(const FAssetData& InAssetData)
 
 void FRigVMRegistry_RWLock::OnPluginLoaded(IPlugin& InPlugin)
 {
-	bool bRegistryChanged = false;
+	bool bRegistryChanged;
 	{
 		FConditionalWriteScopeLock _(*this);
 		bRegistryChanged = Super::OnPluginLoaded_NoLock(InPlugin);
@@ -1363,7 +1567,7 @@ void FRigVMRegistry_RWLock::OnPluginLoaded(IPlugin& InPlugin)
 
 void FRigVMRegistry_RWLock::OnPluginUnloaded(IPlugin& InPlugin)
 {
-	bool bRegistryChanged = false;
+	bool bRegistryChanged;
 	{
 		FConditionalWriteScopeLock _(*this);
 		bRegistryChanged = Super::OnPluginUnloaded_NoLock(InPlugin);
@@ -1444,7 +1648,6 @@ TRigVMTypeIndex FRigVMRegistry_NoLock::GetTypeIndexFromCPPType_NoLock(const FStr
 			{
 				UUserDefinedStruct* ScriptStruct = *ScriptStructIt;
 				const FRigVMTemplateArgumentType ArgumentType(ScriptStruct);
-				const FName StuctCppType = *RigVMTypeUtils::GetUniqueStructTypeName(ScriptStruct);
 				if (ArgumentType.CPPType == BaseCPPTypeName)
 				{
 					// this check for example makes sure we don't add structs defined in verse
@@ -1462,7 +1665,6 @@ TRigVMTypeIndex FRigVMRegistry_NoLock::GetTypeIndexFromCPPType_NoLock(const FStr
 				{
 					UUserDefinedEnum* Enum = *EnumIt;
 					const FRigVMTemplateArgumentType ArgumentType(Enum);
-					const FName EnumCppType = *RigVMTypeUtils::CPPTypeFromEnum(Enum);
 					if (ArgumentType.CPPType == BaseCPPTypeName)
 					{
 						// this check for example makes sure we don't add enums defined in verse
@@ -1873,10 +2075,27 @@ void FRigVMRegistry_NoLock::Register_NoLock(const TCHAR* InName, FRigVMFunctionP
 	}
 #endif
 
-	const FRigVMFunction Function(InName, InFunctionPtr, InStruct, Functions.Num(), InArguments);
-	Functions.AddElement(Function);
-	Functions[Function.Index].OwnerRegistry = this;
-	FunctionNameToIndex.Add(InName, Function.Index);
+	const FName FunctionName = InName;
+	int32 FunctionIndex = Functions.Num();
+	if (const int32* ExistingFunctionIndex = PreviousFunctionNameToIndex.Find(FunctionName))
+	{
+		check(Functions.IsValidIndex(*ExistingFunctionIndex));
+		FunctionIndex = *ExistingFunctionIndex;
+	}
+
+	const FRigVMFunction Function(InName, InFunctionPtr, InStruct, FunctionIndex, InArguments);
+	check(Function.Index == FunctionIndex);
+
+	if (FunctionIndex == Functions.Num())
+	{
+		Functions.AddElement(Function);
+	}
+	else
+	{
+		Functions[FunctionIndex] = Function;
+	}
+	Functions[FunctionIndex].OwnerRegistry = this;
+	FunctionNameToIndex.Add(FunctionName, FunctionIndex);
 
 #if !WITH_EDITOR
 	// if we are not in editor and VMs rely on their own localized registry
@@ -1891,7 +2110,7 @@ void FRigVMRegistry_NoLock::Register_NoLock(const TCHAR* InName, FRigVMFunctionP
 	for (TFieldIterator<FProperty> It(InStruct); It; ++It)
 	{
 		// creating the argument causes the registration
-		(void)FRigVMTemplateArgument::Make_NoLock(*It, Context);
+		(void)FRigVMTemplateArgument::Make_NoLock(*It, ThisHandle);
 	}
 
 #if WITH_EDITOR
@@ -1907,10 +2126,10 @@ void FRigVMRegistry_NoLock::Register_NoLock(const TCHAR* InName, FRigVMFunctionP
 		if (FString(InName).Split(TEXT("::"), nullptr, &MethodName))
 		{
 			const FString TemplateName = RigVMStringUtils::JoinStrings(TemplateMetadata, MethodName, TEXT("::"));
-			FRigVMTemplate Template(InStruct, TemplateName, Function.Index, Context);
+			FRigVMTemplate Template(InStruct, TemplateName, Function.Index, ThisHandle);
 			if (Template.IsValid())
 			{
-				Functions[Function.Index].PermutationIndex = Template.NumPermutations_NoLock(Context) - 1;
+				Functions[Function.Index].PermutationIndex = Template.NumPermutations_NoLock(ThisHandle) - 1;
 				
 				bool bWasMerged = false;
 
@@ -1918,19 +2137,19 @@ void FRigVMRegistry_NoLock::Register_NoLock(const TCHAR* InName, FRigVMFunctionP
 				if(ExistingTemplateIndexPtr)
 				{
 					FRigVMTemplate& ExistingTemplate = TemplateArray[*ExistingTemplateIndexPtr];
-					if (ExistingTemplate.Merge(Template, GetHandle_NoLock()))
+					if (ExistingTemplate.Merge(Template, ThisHandle))
 					{
 						if (!bIsDeprecated)
 						{
 							Functions[Function.Index].TemplateIndex = ExistingTemplate.Index;
-							Functions[Function.Index].PermutationIndex = ExistingTemplate.NumPermutations_NoLock(Context) - 1;
+							Functions[Function.Index].PermutationIndex = ExistingTemplate.NumPermutations_NoLock(ThisHandle) - 1;
 #if WITH_EDITOR
 							int32 MaxNumPermutations = 0;
 							for(const FRigVMTemplateArgument& Argument : ExistingTemplate.Arguments)
 							{
-								MaxNumPermutations = FMath::Max(MaxNumPermutations, Argument.GetNumTypes_NoLock(Context));
+								MaxNumPermutations = FMath::Max(MaxNumPermutations, Argument.GetNumTypes_NoLock(ThisHandle));
 							}
-							verify(MaxNumPermutations == ExistingTemplate.NumPermutations_NoLock(Context));
+							verify(MaxNumPermutations == ExistingTemplate.NumPermutations_NoLock(ThisHandle));
 #endif
 						}
 						bWasMerged = true;
@@ -1958,6 +2177,23 @@ void FRigVMRegistry_NoLock::Register_NoLock(const TCHAR* InName, FRigVMFunctionP
 #endif
 }
 
+bool FRigVMRegistry_NoLock::RemoveFunction_NoLock(int32 InFunctionIndex)
+{
+	check(Functions.IsValidIndex(InFunctionIndex));
+
+	if(!Functions[InFunctionIndex].IsValid())
+	{
+		return false;
+	}
+	
+	const FName PermutationName = *Functions[InFunctionIndex].GetName();
+	FunctionNameToIndex.Remove(PermutationName);
+	PreviousFunctionNameToIndex.FindOrAdd(PermutationName, INDEX_NONE) = InFunctionIndex;
+	Functions[InFunctionIndex] = FRigVMFunction();
+
+	return true;
+}
+
 const FRigVMDispatchFactory* FRigVMRegistry_NoLock::RegisterFactory_NoLock(UScriptStruct* InFactoryStruct, const TArray<FRigVMTemplateArgumentInfo>& InArgumentInfos)
 {
 	check(InFactoryStruct);
@@ -1993,16 +2229,74 @@ const FRigVMDispatchFactory* FRigVMRegistry_NoLock::RegisterFactory_NoLock(UScri
 	Factory->FactoryScriptStruct = InFactoryStruct;
 	Factory->FactoryIndex = Factories.Add(Factory);
 	Factory->OwnerRegistry = this;
-	Factory->RegisterDependencyTypes_NoLock(Context);
+	Factory->RegisterDependencyTypes_NoLock(ThisHandle);
 
 	if(!InArgumentInfos.IsEmpty())
 	{
-		Factory->CreateTemplateForArgumentInfos_NoLock(InArgumentInfos, Context);
+		Factory->CreateTemplateForArgumentInfos_NoLock(InArgumentInfos, ThisHandle);
 	}
 	
 	return Factory;
 }
 
+bool FRigVMRegistry_NoLock::RemoveFactory_NoLock(UScriptStruct* InFactoryStruct)
+{
+	const int32 FactoryIndex = Factories.IndexOfByPredicate([InFactoryStruct](const FRigVMDispatchFactory* Factory) -> bool
+	{
+		return Factory->GetScriptStruct() == InFactoryStruct;
+	});
+
+	if (FactoryIndex == INDEX_NONE)
+	{
+		return false;
+	}
+
+	check(Factories.IsValidIndex(FactoryIndex));
+
+	const int32 TemplateIndex = Factories[FactoryIndex]->GetTemplate_NoLock(ThisHandle)->Index;
+	
+	if (!RemoveTemplate_NoLock(TemplateIndex))
+	{
+		return false;
+	}
+
+	InFactoryStruct->DestroyStruct(Factories[FactoryIndex], 1);
+	FMemory::Free(Factories[FactoryIndex]);
+	
+	Factories.RemoveAt(FactoryIndex);
+	return true;
+}
+
+bool FRigVMRegistry_NoLock::RemoveTemplate_NoLock(int32 InTemplateIndex)
+{
+	check(Templates.IsValidIndex(InTemplateIndex));
+
+	if(!Templates[InTemplateIndex].IsValid())
+	{
+		return false;
+	}
+
+	for (const int32& FunctionIndex : Templates[InTemplateIndex].Permutations)
+	{
+		if (FunctionIndex == INDEX_NONE)
+		{
+			continue;
+		}
+		if (!RemoveFunction_NoLock(FunctionIndex))
+		{
+			return false;
+		}
+	}
+	
+	const FName Notation = Templates[InTemplateIndex].GetNotation();
+	TemplateNotationToIndex.Remove(Notation);
+	PreviousTemplateNotationToIndex.FindOrAdd(Notation, INDEX_NONE) = InTemplateIndex;
+	Templates[InTemplateIndex] = FRigVMTemplate();
+
+	return true;
+
+}
+
 void FRigVMRegistry_NoLock::RegisterPredicate_NoLock(UScriptStruct* InStruct, const TCHAR* InName, const TArray<FRigVMFunctionArgument>& InArguments)
 {
 #if !WITH_EDITOR
@@ -2098,15 +2392,15 @@ const FRigVMFunction* FRigVMRegistry_NoLock::FindFunction_NoLock(const TCHAR* In
 		// if the factory has never been registered - FindDispatchFactory will try to look it up and register
 		if(const FRigVMDispatchFactory* Factory = FindDispatchFactory_NoLock(*StructOrFactoryName))
 		{
-			if(const FRigVMTemplate* Template = Factory->GetTemplate_NoLock(Context))
+			if(const FRigVMTemplate* Template = Factory->GetTemplate_NoLock(ThisHandle))
 			{
-				const FRigVMTemplateTypeMap ArgumentTypes = Template->GetArgumentTypesFromString_NoLock(SuffixString, &InTypeResolver, Context);
+				const FRigVMTemplateTypeMap ArgumentTypes = Template->GetArgumentTypesFromString_NoLock(SuffixString, &InTypeResolver, ThisHandle);
 				if(ArgumentTypes.Num() == Template->NumArguments())
 				{
-					const int32 PermutationIndex = Template->FindPermutation_NoLock(ArgumentTypes, Context);
+					const int32 PermutationIndex = Template->FindPermutation_NoLock(ArgumentTypes, ThisHandle);
 					if(PermutationIndex != INDEX_NONE)
 					{
-						return ((FRigVMTemplate*)Template)->GetOrCreatePermutation_NoLock(PermutationIndex, Context);
+						return ((FRigVMTemplate*)Template)->GetOrCreatePermutation_NoLock(PermutationIndex, ThisHandle);
 					}
 				}
 			}
@@ -2185,7 +2479,12 @@ const FRigVMTemplate* FRigVMRegistry_NoLock::FindTemplate_NoLock(const FName& In
 
 	if(const int32* TemplateIndexPtr = TemplateNotationToIndex.Find(InNotation))
 	{
-		return &Templates[*TemplateIndexPtr];
+		const FRigVMTemplate* Template = &Templates[*TemplateIndexPtr];
+		if (!Template->IsValid())
+		{
+			return nullptr;
+		}
+		return Template;
 	}
 
 	const FString NotationString(InNotation.ToString());
@@ -2214,7 +2513,7 @@ const FRigVMTemplate* FRigVMRegistry_NoLock::FindTemplate_NoLock(const FName& In
 		
 		if(const FRigVMDispatchFactory* Factory = FindDispatchFactory_NoLock(*FactoryName))
 		{
-			return Factory->GetTemplate_NoLock(Context);
+			return Factory->GetTemplate_NoLock(ThisHandle);
 		}
 	}
 
@@ -2245,7 +2544,7 @@ const FRigVMTemplate* FRigVMRegistry_NoLock::FindTemplate_NoLock(const FName& In
 				const FString NewDispatchFactoryName = FRigVMDispatchFactory::DispatchPrefix + Redirect->NewName.ObjectName.ToString();
 				if(const FRigVMDispatchFactory* NewDispatchFactory = FindDispatchFactory_NoLock(*NewDispatchFactoryName))
 				{
-					return NewDispatchFactory->GetTemplate_NoLock(Context);
+					return NewDispatchFactory->GetTemplate_NoLock(ThisHandle);
 				}
 			}
 		}
@@ -2310,11 +2609,11 @@ const FRigVMTemplate* FRigVMRegistry_NoLock::AddTemplateFromArguments_NoLock(con
 {
 	// we only support to ask for templates here which provide singleton types
 	int32 NumPermutations = 0;
-	FRigVMTemplate Template(InName, InInfos, GetHandle_NoLock());
+	FRigVMTemplate Template(InName, InInfos, ThisHandle);
 	for(const FRigVMTemplateArgument& Argument : Template.Arguments)
 	{
-		const int32 NumIndices = Argument.GetNumTypes_NoLock(GetHandle_NoLock());
-		if(!Argument.IsSingleton_NoLock({}, GetHandle_NoLock()) && NumPermutations > 1)
+		const int32 NumIndices = Argument.GetNumTypes_NoLock(ThisHandle);
+		if(!Argument.IsSingleton_NoLock({}, ThisHandle) && NumPermutations > 1)
 		{
 			if(NumIndices != NumPermutations)
 			{
@@ -2328,10 +2627,10 @@ const FRigVMTemplate* FRigVMRegistry_NoLock::AddTemplateFromArguments_NoLock(con
 	// if any of the arguments are wildcards we'll need to update the types
 	for(FRigVMTemplateArgument& Argument : Template.Arguments)
 	{
-		const int32 NumTypes = Argument.GetNumTypes_NoLock(GetHandle_NoLock());
+		const int32 NumTypes = Argument.GetNumTypes_NoLock(ThisHandle);
 		if(NumTypes == 1)
 		{
-			const TRigVMTypeIndex FirstTypeIndex = Argument.GetTypeIndex_NoLock(0, GetHandle_NoLock());
+			const TRigVMTypeIndex FirstTypeIndex = Argument.GetTypeIndex_NoLock(0, ThisHandle);
 			if(IsWildCardType_NoLock(FirstTypeIndex))
 			{
 #if WITH_EDITOR
@@ -2348,7 +2647,7 @@ const FRigVMTemplate* FRigVMRegistry_NoLock::AddTemplateFromArguments_NoLock(con
 				Argument.bUseCategories = true;
 				Argument.TypeIndices.Reset();
 		
-				NumPermutations = FMath::Max(NumPermutations, Argument.GetNumTypes_NoLock(GetHandle_NoLock()));
+				NumPermutations = FMath::Max(NumPermutations, Argument.GetNumTypes_NoLock(ThisHandle));
 			}
 		}
 	}
@@ -2376,7 +2675,7 @@ const FRigVMTemplate* FRigVMRegistry_NoLock::AddTemplateFromArguments_NoLock(con
 		bool bAnyArgumentWithZeroTypes = false;
 		for(int32 ArgIndex = 0; ArgIndex < NumArguments; ArgIndex++)
 		{
-			(void)TypeCaches[ArgIndex].UpdateIfRequired(GetHandle_NoLock(), Template.Arguments[ArgIndex]);
+			(void)TypeCaches[ArgIndex].UpdateIfRequired(ThisHandle, Template.Arguments[ArgIndex]);
 			bAnyArgumentWithZeroTypes = bAnyArgumentWithZeroTypes || TypeCaches[ArgIndex].GetNumTypes_NoLock() == 0;
 		}
 
@@ -2434,28 +2733,42 @@ const FRigVMTemplate* FRigVMRegistry_NoLock::AddTemplateFromArguments_NoLock(con
 #if WITH_EDITOR
 	for(FRigVMTemplateArgument& Argument : Template.Arguments)
 	{
-		Argument.UpdateTypeToPermutationsSlow(GetHandle_NoLock());
+		Argument.UpdateTypeToPermutationsSlow(ThisHandle);
 	}
 #endif
 
 	Template.Permutations.Init(INDEX_NONE, NumPermutations);
-	Template.RecomputeTypesHashToPermutations(TypeCaches, GetHandle_NoLock());
+	Template.RecomputeTypesHashToPermutations(TypeCaches, ThisHandle);
+
+	int32 TemplateIndex = Templates.Num();
+	const FName Notation = Template.GetNotation();
+	if (const int32* ExistingTemplateIndex = PreviousTemplateNotationToIndex.Find(Notation))
+	{
+		TemplateIndex = *ExistingTemplateIndex;
+	}
 
-	const int32 Index = Templates.AddElement(Template);
-	Templates[Index].Index = Index;
-	Templates[Index].OwnerRegistry = this;
-	Templates[Index].Delegates = InDelegates;
-	TemplateNotationToIndex.Add(Template.GetNotation(), Index);
+	if (TemplateIndex == Templates.Num())
+	{
+		Templates.AddElement(Template);
+	}
+	else
+	{
+		Templates[TemplateIndex] = Template;
+	}
+	Templates[TemplateIndex].Index = TemplateIndex;
+	Templates[TemplateIndex].OwnerRegistry = this;
+	Templates[TemplateIndex].Delegates = InDelegates;
+	TemplateNotationToIndex.Add(Notation, TemplateIndex);
 
-	for(int32 ArgumentIndex=0; ArgumentIndex < Templates[Index].Arguments.Num(); ArgumentIndex++)
+	for(int32 ArgumentIndex=0; ArgumentIndex < Templates[TemplateIndex].Arguments.Num(); ArgumentIndex++)
 	{
-		for(const FRigVMTemplateArgument::ETypeCategory& ArgumentTypeCategory : Templates[Index].Arguments[ArgumentIndex].TypeCategories)
+		for(const FRigVMTemplateArgument::ETypeCategory& ArgumentTypeCategory : Templates[TemplateIndex].Arguments[ArgumentIndex].TypeCategories)
 		{
-			TemplatesPerCategory.FindChecked(ArgumentTypeCategory).AddUnique(Index);
+			TemplatesPerCategory.FindChecked(ArgumentTypeCategory).AddUnique(TemplateIndex);
 		}
 	}
 	
-	return &Templates[Index];
+	return &Templates[TemplateIndex];
 }
 
 FRigVMDispatchFactory* FRigVMRegistry_NoLock::FindDispatchFactory_NoLock(const FName& InFactoryName) const
@@ -2501,12 +2814,12 @@ FString FRigVMRegistry_NoLock::FindOrAddSingletonDispatchFunction_NoLock(UScript
 	{
 		if(Factory->IsSingleton())
 		{
-			if(const FRigVMTemplate* Template = Factory->GetTemplate_NoLock(Context))
+			if(const FRigVMTemplate* Template = Factory->GetTemplate_NoLock(ThisHandle))
 			{
 				// use the types for the first permutation - since we don't care
 				// for a singleton dispatch
-				const FRigVMTemplateTypeMap TypesForPrimaryPermutation = Template->GetTypesForPermutation_NoLock(0, GetHandle_NoLock());
-				const FString Name = Factory->GetPermutationName_NoLock(TypesForPrimaryPermutation, Context);
+				const FRigVMTemplateTypeMap TypesForPrimaryPermutation = Template->GetTypesForPermutation_NoLock(0, ThisHandle);
+				const FString Name = Factory->GetPermutationName_NoLock(TypesForPrimaryPermutation, ThisHandle);
 				if(const FRigVMFunction* Function = FindFunction_NoLock(*Name))
 				{
 					return Function->Name;
@@ -2541,8 +2854,11 @@ void FRigVMRegistry_RWLock::Initialize_NoLock()
 	AssetRegistryModule.Get().OnAssetRemoved().AddRaw(this, &FRigVMRegistry_RWLock::OnAssetRemoved);
 	AssetRegistryModule.Get().OnAssetRenamed().AddRaw(this, &FRigVMRegistry_RWLock::OnAssetRenamed);
 
-	IPluginManager::Get().OnNewPluginMounted().AddRaw(this, &FRigVMRegistry_RWLock::OnPluginLoaded);
-	IPluginManager::Get().OnPluginUnmounted().AddRaw(this, &FRigVMRegistry_RWLock::OnPluginUnloaded);
+	if (bIsGlobalRegistry)
+	{
+		IPluginManager::Get().OnNewPluginMounted().AddRaw(this, &FRigVMRegistry_RWLock::OnPluginLoaded);
+		IPluginManager::Get().OnPluginUnmounted().AddRaw(this, &FRigVMRegistry_RWLock::OnPluginUnloaded);
+	}
 	
 	UE::Anim::AttributeTypes::GetOnAttributeTypesChanged().AddRaw(this, &FRigVMRegistry_RWLock::OnAnimationAttributeTypesChanged);
 
diff --git a/Engine/Plugins/Runtime/RigVM/Source/RigVM/Private/RigVMCore/RigVMTemplate.cpp b/Engine/Plugins/Runtime/RigVM/Source/RigVM/Private/RigVMCore/RigVMTemplate.cpp
--- a/Engine/Plugins/Runtime/RigVM/Source/RigVM/Private/RigVMCore/RigVMTemplate.cpp
+++ b/Engine/Plugins/Runtime/RigVM/Source/RigVM/Private/RigVMCore/RigVMTemplate.cpp
@@ -1138,6 +1138,13 @@ FRigVMTemplate::FTypeMap FRigVMTemplate::GetArgumentTypesFromString_NoLock(const
 			{
 				if(const FRigVMTemplateArgument* Argument = FindArgument(*ArgumentName))
 				{
+					const FName TypeFName = *TypeName;
+					if (IsIgnoredType(TypeFName))
+					{
+						Types.Reset();
+						return Types;
+					}
+
 					TRigVMTypeIndex TypeIndex = InRegistry->GetTypeIndexFromCPPType_NoLock(TypeName);
 
 					// If the type was not found, check if it's a user-defined type that hasn't been registered yet.
@@ -1145,7 +1152,7 @@ FRigVMTemplate::FTypeMap FRigVMTemplate::GetArgumentTypesFromString_NoLock(const
 					{
 						UObject* CPPTypeObject = RigVMTypeUtils::ObjectFromCPPType(TypeName, true, InTypeResolver);
 						
-						FRigVMTemplateArgumentType ArgType(*TypeName, CPPTypeObject);
+						FRigVMTemplateArgumentType ArgType(TypeFName, CPPTypeObject);
 						TypeIndex = InRegistry->FindOrAddType_NoLock(ArgType, false);
 					}
 					
@@ -1759,7 +1766,12 @@ const FRigVMFunction* FRigVMTemplate::GetOrCreatePermutation_NoLock(int32 InInde
 		FTypeMap Types;
 		for(const FRigVMTemplateArgument& Argument : Arguments)
 		{
-			Types.Add(Argument.GetName(), Argument.GetTypeIndex_NoLock(InIndex, InRegistry));
+			const TRigVMTypeIndex TypeIndex = Argument.GetTypeIndex_NoLock(InIndex, InRegistry);
+			if (IsIgnoredType(TypeIndex))
+			{
+				return nullptr;
+			}
+			Types.Add(Argument.GetName(), TypeIndex);
 		}
 
 		FRigVMDispatchFactory* Factory = Delegates.GetDispatchFactoryDelegate.Execute(InRegistry);
@@ -1776,17 +1788,38 @@ const FRigVMFunction* FRigVMTemplate::GetOrCreatePermutation_NoLock(int32 InInde
 			
 			static constexpr TCHAR Format[] = TEXT("%s::%s");
 			const FString PermutationName = Factory->GetPermutationNameImpl(Types, InRegistry);
-			const int32 FunctionIndex = Permutations[InIndex] = InRegistry->Functions.Num();
-			
-			InRegistry->Functions.AddElement(
-				FRigVMFunction(
+
+			int32 FunctionIndex = InRegistry->Functions.Num();
+			if (const int32* ExistingFunctionIndex = InRegistry->PreviousFunctionNameToIndex.Find(*PermutationName))
+			{
+				FunctionIndex = *ExistingFunctionIndex;
+			}
+
+			Permutations[InIndex] = FunctionIndex;
+
+			if (FunctionIndex  == InRegistry->Functions.Num())
+			{
+				InRegistry->Functions.AddElement(
+					FRigVMFunction(
+						PermutationName,
+						DispatchFunction,
+						Factory,
+						FunctionIndex,
+						FunctionArguments
+					)
+				);
+			}
+			else
+			{
+				InRegistry->Functions[FunctionIndex] = FRigVMFunction(
 					PermutationName,
 					DispatchFunction,
 					Factory,
 					FunctionIndex,
 					FunctionArguments
-				)
-			);
+				);
+			}
+			
 			InRegistry->Functions[FunctionIndex].TemplateIndex = Index;
 			InRegistry->Functions[FunctionIndex].PermutationIndex = InIndex;
 			InRegistry->Functions[FunctionIndex].OwnerRegistry = InRegistry.GetRegistry();
@@ -2214,6 +2247,11 @@ FRigVMTemplateTypeMap FRigVMTemplate::GetTypesForPermutation_NoLock(const int32
 		const TOptional<TRigVMTypeIndex> TypeIndex = Argument->TryToGetTypeIndex_NoLock(InPermutationIndex, InRegistry);
 		if(TypeIndex.IsSet())
 		{
+			if (IsIgnoredType(TypeIndex.GetValue()))
+			{
+				TypeMap.Reset();
+				return TypeMap;
+			}
 			TypeMap.Add(Argument->GetName(), TypeIndex.GetValue());
 		}
 		else
@@ -2440,6 +2478,16 @@ bool FRigVMTemplate::UpdateArgumentTypes_Impl(
 	return true;
 }
 
+bool FRigVMTemplate::IsIgnoredType(const TRigVMTypeIndex& InTypeIndex) const
+{
+	return TypeIndicesToIgnore.Contains(InTypeIndex);
+}
+
+bool FRigVMTemplate::IsIgnoredType(const FName& InTypeName) const
+{
+	return TypeNamesToIgnore.Contains(InTypeName);
+}
+
 bool FRigVMTemplate::HandlePropagatedArgumentType(const TRigVMTypeIndex InTypeIndex, FRigVMRegistryHandle& InRegistry)
 {
 	check(InRegistry.GetRegistry() == OwnerRegistry);
@@ -2455,6 +2503,12 @@ bool FRigVMTemplate::HandlePropagatedArgumentType(const TRigVMTypeIndex InTypeIn
 		return true;
 	}
 
+	// clear out the type to ignore as it is being added back
+	if (TypeIndicesToIgnore.Remove(InTypeIndex) > 0)
+	{
+		TypeNamesToIgnore.Remove(InRegistry->GetType_NoLock(InTypeIndex).CPPType);
+	}
+
 	InvalidateHash();
 
 	const FRigVMTemplateArgument& PrimaryArgument = Arguments[PrimaryArgumentIndex];
@@ -2541,39 +2595,36 @@ bool FRigVMTemplate::HandlePropagatedArgumentType(const TRigVMTypeIndex InTypeIn
 	return bResult;
 }
 
-void FRigVMTemplate::HandleTypeRemoval(TRigVMTypeIndex InTypeIndex, const FRigVMRegistryHandle& InRegistry)
+void FRigVMTemplate::HandleTypeRemoval(TRigVMTypeIndex InTypeIndex, FRigVMRegistryHandle& InRegistry)
 {
 	check(InRegistry.GetRegistry() == OwnerRegistry);
 
 	InvalidateHash();
 
-	TArray<int32> PermutationsToRemove;
-	for (int32 PermutationIndex = 0; PermutationIndex < NumPermutations_NoLock(InRegistry); PermutationIndex++)
+	const int32 CurrentNumPermutations = NumPermutations_NoLock(InRegistry);
+	for (int32 PermutationIndex = 0; PermutationIndex < CurrentNumPermutations; PermutationIndex++)
 	{
-		FRigVMTemplateTypeMap TypeMap = GetTypesForPermutation_NoLock(PermutationIndex, InRegistry);
-
-		TArray<TRigVMTypeIndex> Types;
-		TypeMap.GenerateValueArray(Types);
+		const FRigVMTemplateTypeMap TypeMap = GetTypesForPermutation_NoLock(PermutationIndex, InRegistry);
 
-		if (Types.Contains(InTypeIndex))
+		for (const TPair<FName, TRigVMTypeIndex>& Pair : TypeMap)
 		{
-			PermutationsToRemove.Add(PermutationIndex);
-		}
-	}
-
-	for (FRigVMTemplateArgument& Argument : Arguments)
-	{
-		TArray<TRigVMTypeIndex> NewTypeIndices; 
-		for (int32 PermutationIndex = 0; PermutationIndex < Argument.TypeIndices.Num(); PermutationIndex++)
-		{
-			if (PermutationsToRemove.Contains(PermutationIndex))
+			if (Pair.Value == InTypeIndex)
 			{
-				// invalidate the type index for this permutation
-				Argument.TypeIndices[PermutationIndex] = INDEX_NONE;
+				if (const FRigVMFunction* Permutation = GetPermutation_NoLock(PermutationIndex, InRegistry))
+				{
+					const int32 FunctionIndex = Permutation->Index;
+					InRegistry->RemoveFunction_NoLock(FunctionIndex);
+					Permutations[PermutationIndex] = INDEX_NONE;
+				}
+				break;
 			}
 		}
 	}
 
+	check(!TypeIndicesToIgnore.Contains(InTypeIndex));
+	TypeIndicesToIgnore.Add(InTypeIndex);
+	TypeNamesToIgnore.Add(InRegistry->GetType_NoLock(InTypeIndex).CPPType);
+
 #if WITH_EDITOR
 	for (FRigVMTemplateArgument& Argument : Arguments)
 	{
diff --git a/Engine/Plugins/Runtime/RigVM/Source/RigVM/Public/RigVMCore/RigVMRegistry.h b/Engine/Plugins/Runtime/RigVM/Source/RigVM/Public/RigVMCore/RigVMRegistry.h
--- a/Engine/Plugins/Runtime/RigVM/Source/RigVM/Public/RigVMCore/RigVMRegistry.h
+++ b/Engine/Plugins/Runtime/RigVM/Source/RigVM/Public/RigVMCore/RigVMRegistry.h
@@ -77,9 +77,21 @@ public:
 	// for example "FMyStruct::MyVirtualMethod"
 	UE_API virtual void Register_NoLock(const TCHAR* InName, FRigVMFunctionPtr InFunctionPtr, UScriptStruct* InStruct = nullptr, const TArray<FRigVMFunctionArgument>& InArguments = TArray<FRigVMFunctionArgument>());
 
+	// Unregisters a function given its function index.
+	// This keeps the function index stable - but assigns an invalid function
+	UE_API virtual bool RemoveFunction_NoLock(int32 InFunctionIndex);
+
 	// Registers a dispatch factory given its struct.
 	UE_API virtual const FRigVMDispatchFactory* RegisterFactory_NoLock(UScriptStruct* InFactoryStruct, const TArray<FRigVMTemplateArgumentInfo>& InArgumentInfos = {});
 
+	// Unregisters a factory given its script struct.
+	// This keeps the template index stable
+	UE_API virtual bool RemoveFactory_NoLock(UScriptStruct* InFactoryStruct);
+
+	// Unregisters a template given its index.
+	// This keeps the template index stable - but assigns an invalid template
+	UE_API virtual bool RemoveTemplate_NoLock(int32 InTemplateIndex);
+
 	// Register a predicate contained in the input struct
 	UE_API virtual void RegisterPredicate_NoLock(UScriptStruct* InStruct, const TCHAR* InName, const TArray<FRigVMFunctionArgument>& InArguments);
 
@@ -119,6 +131,15 @@ public:
 	// Update the registry when old types are removed
 	UE_API virtual bool OnAssetRemoved_NoLock(const FAssetData& InAssetData);
 
+	// Returns the prefixes for all modules within a plugin
+	UE_API static TArray<FString> GetPluginModulePrefixes(const IPlugin& InPlugin);
+
+	// Returns true if a given object is part of a plugin
+	UE_API static bool IsWithinPlugin(const IPlugin& InPlugin, const UObject* InObject);
+
+	// Returns true if a given object is part of a plugin object path
+	UE_API static bool IsWithinPlugin(const TArray<FString>& InModulePrefixes, const UObject* InObject);
+
 	// May add factories and unit functions declared in the plugin 
 	UE_API virtual bool OnPluginLoaded_NoLock(IPlugin& InPlugin);
 
@@ -141,7 +162,11 @@ public:
 
 	// Removes a type from the registry, and updates all dependent templates
 	// which also creates invalid permutations in templates that we should ignore
-	UE_API virtual bool RemoveType_NoLock(const FSoftObjectPath& InObjectPath, const UClass* InObjectClass);
+	UE_API virtual bool RemoveType_NoLock(const FSoftObjectPath& InObjectPath);
+
+	// Removes a type from the registry, and updates all dependent templates
+	// which also creates invalid permutations in templates that we should ignore
+	UE_API virtual bool RemoveType_NoLock(TRigVMTypeIndex InTypeIndex);
 
 	// Returns the type index given a type
 	UE_API virtual TRigVMTypeIndex GetTypeIndex_NoLock(const FRigVMTemplateArgumentType& InType) const;
@@ -338,7 +363,7 @@ protected:
 	UE_API FRigVMRegistry_NoLock(bool InIsGlobalRegistry);
 
 	// disable copy constructor
-	FRigVMRegistry_NoLock(const FRigVMRegistry_NoLock&) = delete;
+	FRigVMRegistry_NoLock(const FRigVMRegistry_NoLock&);
 	// disable assignment operator
 	FRigVMRegistry_NoLock& operator= (const FRigVMRegistry_NoLock &InOther) = delete;
 
@@ -361,6 +386,11 @@ protected:
 		mutable uint32 Hash;
 
 		void Serialize_NoLock(FArchive& Ar);
+
+		bool IsUnknownType() const
+		{
+			return Type.IsUnknownType();
+		}
 	};
 
 	// Initialize the base types
@@ -386,8 +416,8 @@ protected:
 	UE_API void PropagateTypeAddedToCategory_NoLock(const FRigVMTemplateArgument::ETypeCategory InCategory, const TRigVMTypeIndex InTypeIndex);
 	UE_API void RemoveTypeInCategory_NoLock(FRigVMTemplateArgument::ETypeCategory InCategory, TRigVMTypeIndex InTypeIndex);
 
-	// the context storing this registry
-	mutable FRigVMRegistryHandle Context;
+	// the handle storing this registry
+	mutable FRigVMRegistryHandle ThisHandle;
 	
 	// true if this registry is the globally shared registry 
 	bool bIsGlobalRegistry;
@@ -412,12 +442,18 @@ protected:
 	// name lookup for functions
 	TMap<FName, int32> FunctionNameToIndex;
 
+	// Previous value for name lookups - this can happen to functions which get deregistered
+	TMap<FName, int32> PreviousFunctionNameToIndex;
+
 	// lookup all the predicate functions of this struct
 	TMap<FName, TArray<FRigVMFunction>> StructNameToPredicates;
 
 	// name lookup for non-deprecated templates
 	TMap<FName, int32> TemplateNotationToIndex;
 
+	// previous name lookup for templates which got removed
+	TMap<FName, int32> PreviousTemplateNotationToIndex;
+
 	// name lookup for deprecated templates
 	TMap<FName, int32> DeprecatedTemplateNotationToIndex;
 
@@ -427,9 +463,12 @@ protected:
 	// Lookup per type category to know which template to keep in sync
 	TMap<FRigVMTemplateArgument::ETypeCategory, TArray<int32>> TemplatesPerCategory;
 
-	// Name loop up for user defined types since they can be deleted.
+	// This type index triple represents the single value, array value, and 2d array value type indices.
+	using TRigVMTypeIndexTriple = TTuple<TRigVMTypeIndex, TRigVMTypeIndex, TRigVMTypeIndex>; 
+
+	// Name lookup for types since they can be deleted.
 	// When that happens, it won't be safe to reload deleted assets so only type names are reliable
-	TMap<FSoftObjectPath, TRigVMTypeIndex> UserDefinedTypeToIndex;
+	TMap<FSoftObjectPath, TRigVMTypeIndexTriple> SoftObjectPathToTypeIndex;
 	
 	// All allowed classes
 	TSet<TObjectPtr<const UClass>> AllowedClasses;
@@ -628,6 +667,14 @@ public:
 		Super::Register_NoLock(InName, InFunctionPtr, InStruct, InArguments);
 	}
 
+	// Unregisters a function given its function index.
+	// This keeps the function index stable - but assigns an invalid function
+	bool RemoveFunction(int32 InFunctionIndex)
+	{
+		FConditionalWriteScopeLock _(*this);
+		return Super::RemoveFunction_NoLock(InFunctionIndex);
+	}
+
 	// Registers a dispatch factory given its struct.
 	const FRigVMDispatchFactory* RegisterFactory(UScriptStruct* InFactoryStruct)
 	{
@@ -635,6 +682,22 @@ public:
 		return Super::RegisterFactory_NoLock(InFactoryStruct);
 	}
 
+	// Unregisters a factory given its script struct.
+	// This keeps the template index stable
+	bool RemoveFactory(UScriptStruct* InFactoryStruct)
+	{
+		FConditionalWriteScopeLock _(*this);
+		return Super::RemoveFactory_NoLock(InFactoryStruct);
+	}
+
+	// Unregisters a template given its index.
+	// This keeps the template index stable - but assigns an invalid template
+	bool RemoveTemplate(int32 InTemplateIndex)
+	{
+		FConditionalWriteScopeLock _(*this);
+		return Super::RemoveTemplate_NoLock(InTemplateIndex);
+	}
+
 	// Register a predicate contained in the input struct
 	void RegisterPredicate(UScriptStruct* InStruct, const TCHAR* InName, const TArray<FRigVMFunctionArgument>& InArguments)
 	{
@@ -722,10 +785,18 @@ public:
 
 	// Removes a type from the registry, and updates all dependent templates
 	// which also creates invalid permutations in templates that we should ignore
-	bool RemoveType(const FSoftObjectPath& InObjectPath, const UClass* InObjectClass)
+	bool RemoveType(const FSoftObjectPath& InObjectPath)
+	{
+		FConditionalWriteScopeLock _(*this);
+		return Super::RemoveType_NoLock(InObjectPath);
+	}
+
+	// Removes a type from the registry, and updates all dependent templates
+	// which also creates invalid permutations in templates that we should ignore
+	bool RemoveType(TRigVMTypeIndex InTypeIndex)
 	{
 		FConditionalWriteScopeLock _(*this);
-		return Super::RemoveType_NoLock(InObjectPath, InObjectClass);
+		return Super::RemoveType_NoLock(InTypeIndex);
 	}
 
 	// Returns the type index given a type
@@ -1008,6 +1079,13 @@ protected:
 		Super::Register_NoLock(InName, InFunctionPtr, InStruct, InArguments);
 	}
 
+	// Unregisters a function given its function index.
+	// This keeps the function index stable - but assigns an invalid function
+	virtual bool RemoveFunction_NoLock(int32 InFunctionIndex) override
+	{
+		return Super::RemoveFunction_NoLock(InFunctionIndex);
+	}
+
 	// Registers a dispatch factory given its struct.
 	// Note: Only call NoLock methods on the FRigVMRegistry_NoLock
 	virtual const FRigVMDispatchFactory* RegisterFactory_NoLock(UScriptStruct* InFactoryStruct, const TArray<FRigVMTemplateArgumentInfo>& InArgumentInfos = {}) override
@@ -1015,6 +1093,20 @@ protected:
 		return Super::RegisterFactory_NoLock(InFactoryStruct, InArgumentInfos);
 	}
 
+	// Unregisters a factory given its script struct.
+	// This keeps the template index stable
+	virtual bool RemoveFactory_NoLock(UScriptStruct* InFactoryStruct) override
+	{
+		return Super::RemoveFactory_NoLock(InFactoryStruct);
+	}
+
+	// Unregisters a template given its index.
+	// This keeps the template index stable - but assigns an invalid template
+	virtual bool RemoveTemplate_NoLock(int32 InTemplateIndex) override
+	{
+		return Super::RemoveTemplate_NoLock(InTemplateIndex);
+	}
+
 	// Register a predicate contained in the input struct
 	// Note: Only call NoLock methods on the FRigVMRegistry_NoLock
 	virtual void RegisterPredicate_NoLock(UScriptStruct* InStruct, const TCHAR* InName, const TArray<FRigVMFunctionArgument>& InArguments) override
@@ -1112,9 +1204,17 @@ protected:
 	// Removes a type from the registry, and updates all dependent templates
 	// which also creates invalid permutations in templates that we should ignore
 	// Note: Only call NoLock methods on the FRigVMRegistry_NoLock
-	virtual bool RemoveType_NoLock(const FSoftObjectPath& InObjectPath, const UClass* InObjectClass) override
+	virtual bool RemoveType_NoLock(const FSoftObjectPath& InObjectPath) override
+	{
+		return Super::RemoveType_NoLock(InObjectPath);
+	}
+
+	// Removes a type from the registry, and updates all dependent templates
+	// which also creates invalid permutations in templates that we should ignore
+	// Note: Only call NoLock methods on the FRigVMRegistry_NoLock
+	virtual bool RemoveType_NoLock(TRigVMTypeIndex InTypeIndex) override
 	{
-		return Super::RemoveType_NoLock(InObjectPath, InObjectClass);
+		return Super::RemoveType_NoLock(InTypeIndex);
 	}
 
 	// Returns the type index given a type
diff --git a/Engine/Plugins/Runtime/RigVM/Source/RigVM/Public/RigVMCore/RigVMTemplate.h b/Engine/Plugins/Runtime/RigVM/Source/RigVM/Public/RigVMCore/RigVMTemplate.h
--- a/Engine/Plugins/Runtime/RigVM/Source/RigVM/Public/RigVMCore/RigVMTemplate.h
+++ b/Engine/Plugins/Runtime/RigVM/Source/RigVM/Public/RigVMCore/RigVMTemplate.h
@@ -87,6 +87,15 @@ struct FRigVMTemplateArgumentType
 		return !CPPType.IsNone();
 	}
 
+	bool IsUnknownType() const
+	{
+		if (CPPTypeObject)
+		{
+			return CPPTypeObject == RigVMTypeUtils::GetWildCardCPPTypeObject();
+		}
+		return false;
+	}
+	
 	void Serialize_NoLock(FArchive& Ar);
 
 	static FRigVMTemplateArgumentType Array()
@@ -567,7 +576,7 @@ public:
 	RIGVM_API bool HandlePropagatedArgumentType(const TRigVMTypeIndex InTypeIndex, FRigVMRegistryHandle& InRegistry);
 
 	// Invalidates template permutations whenever a type such as a user defined struct is removed
-	RIGVM_API void HandleTypeRemoval(TRigVMTypeIndex InTypeIndex, const FRigVMRegistryHandle& InRegistry);
+	RIGVM_API void HandleTypeRemoval(TRigVMTypeIndex InTypeIndex, FRigVMRegistryHandle& InRegistry);
 	
 	// Returns the delegate to be able to react to type changes dynamically
 	// This delegate is deprecated
@@ -621,6 +630,11 @@ private:
 		const FRigVMDispatchFactory* InFactory,
 		TArray<FRigVMTemplateTypeMap, TInlineAllocator<1>>& InOutTypesArray);
 
+	// returns true if this is a type which should be ignored by this template
+	// since it got removed by unloading a plugin
+	bool IsIgnoredType(const TRigVMTypeIndex& InTypeIndex) const;
+	bool IsIgnoredType(const FName& InTypeName) const;
+
 	int32 Index;
 	FName Notation;
 	TArray<FRigVMTemplateArgument> Arguments;
@@ -632,6 +646,9 @@ private:
 
 	FRigVMTemplateDelegates Delegates;
 
+	TArray<TRigVMTypeIndex> TypeIndicesToIgnore;
+	TArray<FName> TypeNamesToIgnore;
+
 	friend struct FRigVMRegistry_NoLock;
 	friend class URigVMController;
 	friend class URigVMLibraryNode;
