Fixed warnings that could pop up when async loading was flushed during PIE, causing rig vm editor to search for assets.

#rnx

diff --git a/Engine/Plugins/Runtime/RigVM/Source/RigVMEditor/Private/EdGraph/NodeSpawners/RigVMEdGraphFunctionRefNodeSpawner.cpp b/Engine/Plugins/Runtime/RigVM/Source/RigVMEditor/Private/EdGraph/NodeSpawners/RigVMEdGraphFunctionRefNodeSpawner.cpp
--- a/Engine/Plugins/Runtime/RigVM/Source/RigVMEditor/Private/EdGraph/NodeSpawners/RigVMEdGraphFunctionRefNodeSpawner.cpp
+++ b/Engine/Plugins/Runtime/RigVM/Source/RigVMEditor/Private/EdGraph/NodeSpawners/RigVMEdGraphFunctionRefNodeSpawner.cpp
@@ -40,8 +40,11 @@ URigVMEdGraphFunctionRefNodeSpawner* URigVMEdGraphFunctionRefNodeSpawner::Create
 	NodeSpawner->ReferencedPublicFunctionHeader = InFunction->GetFunctionHeader();
 	NodeSpawner->NodeClass = URigVMEdGraphNode::StaticClass();
 	NodeSpawner->bIsLocalFunction = true;
-	const FAssetData Asset = UE::RigVM::Editor::Tools::FindAssetFromAnyPath(InFunction->GetPathName(), true);
-	NodeSpawner->AssetPath = Asset.ToSoftObjectPath(); 	
+	if (!GEditor->PlayWorld && !GIsPlayInEditorWorld)
+	{
+		const FAssetData Asset = UE::RigVM::Editor::Tools::FindAssetFromAnyPath(InFunction->GetPathName(), true);
+		NodeSpawner->AssetPath = Asset.ToSoftObjectPath();
+	}
 
 	FBlueprintActionUiSpec& MenuSignature = NodeSpawner->DefaultMenuSignature;
 
diff --git a/Engine/Plugins/Runtime/RigVM/Source/RigVMEditor/Private/RigVMEditorModule.cpp b/Engine/Plugins/Runtime/RigVM/Source/RigVMEditor/Private/RigVMEditorModule.cpp
--- a/Engine/Plugins/Runtime/RigVM/Source/RigVMEditor/Private/RigVMEditorModule.cpp
+++ b/Engine/Plugins/Runtime/RigVM/Source/RigVMEditor/Private/RigVMEditorModule.cpp
@@ -539,6 +539,10 @@ void FRigVMEditorModule::GetInstanceActions(URigVMBlueprint* RigVMBlueprint, FBl
 
 	if (URigVMBlueprintGeneratedClass* GeneratedClass = RigVMBlueprint->GetRigVMBlueprintGeneratedClass())
 	{
+		if (!ActionRegistrar.IsOpenForRegistration(GeneratedClass))
+		{
+			return;
+		}
 		if (URigVMHost* CDO = Cast<URigVMHost>(GeneratedClass->GetDefaultObject()))
 		{
 			static const FText NodeCategory = LOCTEXT("Variables", "Variables");
