RigVM: Defer rebuilding the UI when adding traits or setting the node layout

Speeds up AddTrait and SetNodeLayout calls by avoiding repeated redundant work
Also fixes error reporting when undo-ing setting a pin's display name

diff --git a/Engine/Plugins/Runtime/RigVM/Source/RigVMDeveloper/Private/RigVMModel/RigVMController.cpp b/Engine/Plugins/Runtime/RigVM/Source/RigVMDeveloper/Private/RigVMModel/RigVMController.cpp
--- a/Engine/Plugins/Runtime/RigVM/Source/RigVMDeveloper/Private/RigVMModel/RigVMController.cpp
+++ b/Engine/Plugins/Runtime/RigVM/Source/RigVMDeveloper/Private/RigVMModel/RigVMController.cpp
@@ -8354,7 +8354,7 @@ bool URigVMController::SetPinDisplayName(URigVMPin* InPin, const FString& InDisp
 
 	if(InDisplayName.Equals(OldDisplayName, ESearchCase::CaseSensitive))
 	{
-		return false;
+		return true;
 	}
 
 	FRigVMControllerCompileBracketScope CompileScope(this);
@@ -9525,7 +9525,8 @@ bool URigVMController::SetNodeLayout(const URigVMNode* InNode, FRigVMNodeLayout
 
 	// first remove any pin layout from the node
 	{
-		const TGuardValue<bool> GuardRefreshFunctions(bSuspendRefreshingFunctionReferences, true); 
+		const TGuardValue<bool> GuardRefreshFunctions(bSuspendRefreshingFunctionReferences, true);
+		const TGuardValue<bool> GuardNotifications(bSuspendNotifications, true); 
 
 		if(ClearNodeLayout(InNode->GetFName(), bSetupUndoRedo, bPrintPythonCommand))
 		{
@@ -9536,7 +9537,8 @@ bool URigVMController::SetNodeLayout(const URigVMNode* InNode, FRigVMNodeLayout
 	for(const TPair<FString,FString>& DisplayNamePair : InLayout.DisplayNames)
 	{
 		const TGuardValue<bool> GuardRefreshFunctions(bSuspendRefreshingFunctionReferences, true); 
-
+		const TGuardValue<bool> GuardNotifications(bSuspendNotifications, true);
+		
 		if(const URigVMPin* Pin = InNode->FindPin(DisplayNamePair.Key))
 		{
 			if(SetPinDisplayName(Pin->GetPinPath(), DisplayNamePair.Value, bSetupUndoRedo, bPrintPythonCommand))
@@ -9554,7 +9556,8 @@ bool URigVMController::SetNodeLayout(const URigVMNode* InNode, FRigVMNodeLayout
 	for(const FRigVMPinCategory& Category : InLayout.Categories)
 	{
 		const TGuardValue<bool> GuardRefreshFunctions(bSuspendRefreshingFunctionReferences, true); 
-
+		const TGuardValue<bool> GuardNotifications(bSuspendNotifications, true);
+		
 		for(const FString& SegmentPath : Category.Elements)
 		{
 			if(const URigVMPin* Pin = InNode->FindPin(SegmentPath))
@@ -9575,7 +9578,8 @@ bool URigVMController::SetNodeLayout(const URigVMNode* InNode, FRigVMNodeLayout
 	for(const TPair<FString,int32>& PinIndexPair : InLayout.PinIndexInCategory)
 	{
 		const TGuardValue<bool> GuardRefreshFunctions(bSuspendRefreshingFunctionReferences, true); 
-
+		const TGuardValue<bool> GuardNotifications(bSuspendNotifications, true);
+		
 		if(const URigVMPin* Pin = InNode->FindPin(PinIndexPair.Key))
 		{
 			if(SetPinIndexInCategory(Pin->GetPinPath(), PinIndexPair.Value, bSetupUndoRedo, bPrintPythonCommand))
@@ -9597,7 +9601,12 @@ bool URigVMController::SetNodeLayout(const URigVMNode* InNode, FRigVMNodeLayout
 			RefreshFunctionReferences(LibraryNode, bSetupUndoRedo, false);
 		}
 	}
-	
+
+	if (NumChanges > 0)
+	{
+		Notify(ERigVMGraphNotifType::PinCategoriesChanged, const_cast<URigVMNode*>(InNode));
+	}
+
 	if(bSetupUndoRedo)
 	{
 		if(NumChanges > 0)
@@ -17197,39 +17206,45 @@ FName URigVMController::AddTrait(URigVMNode* InNode, UScriptStruct* InTraitScrip
 	TraitPin->Direction = ERigVMPinDirection::Input;
 	TraitPin->DefaultValue = DefaultValue;
 
-	AddNodePin(InNode, TraitPin);
-	Notify(ERigVMGraphNotifType::PinAdded, TraitPin);
+	{
+		const TGuardValue<bool> GuardNotifications(bSuspendNotifications, true);
 	
-	AddPinsForStruct(InTraitScriptStruct, InNode, TraitPin, TraitPin->GetDirection(), DefaultValue, true);
+		AddNodePin(InNode, TraitPin);
 
-	FRigVMPinInfoArray ProgrammaticPins;
-	Trait->GetProgrammaticPins(this, INDEX_NONE, TraitPin, DefaultValue, ProgrammaticPins);
+		AddPinsForStruct(InTraitScriptStruct, InNode, TraitPin, TraitPin->GetDirection(), DefaultValue, true);
 
-	const FRigVMRegistry& Registry = FRigVMRegistry::Get();
-	const FRigVMPinInfoArray PreviousPins;
+		FRigVMPinInfoArray ProgrammaticPins;
+		Trait->GetProgrammaticPins(this, INDEX_NONE, TraitPin, DefaultValue, ProgrammaticPins);
 
-	for (int32 PinIndex = 0; PinIndex < ProgrammaticPins.Num(); ++PinIndex)
-	{
-		const FString& PinPath = ProgrammaticPins.GetPinPath(PinIndex);
-		FString ParentPinPath, PinName;
-		UObject* OuterForPin = TraitPin;
-		if (URigVMPin::SplitPinPathAtEnd(PinPath, ParentPinPath, PinName))
+		const FRigVMRegistry& Registry = FRigVMRegistry::Get();
+		const FRigVMPinInfoArray PreviousPins;
+
+		for (int32 PinIndex = 0; PinIndex < ProgrammaticPins.Num(); ++PinIndex)
 		{
-			OuterForPin = TraitPin->FindSubPin(ParentPinPath);
+			const FString& PinPath = ProgrammaticPins.GetPinPath(PinIndex);
+			FString ParentPinPath, PinName;
+			UObject* OuterForPin = TraitPin;
+			if (URigVMPin::SplitPinPathAtEnd(PinPath, ParentPinPath, PinName))
+			{
+				OuterForPin = TraitPin->FindSubPin(ParentPinPath);
+			}
+
+			CreatePinFromPinInfo(Registry, PreviousPins, ProgrammaticPins[PinIndex], PinPath, OuterForPin);
 		}
 
-		CreatePinFromPinInfo(Registry, PreviousPins, ProgrammaticPins[PinIndex], PinPath, OuterForPin);
-	}
+		// move the the pin to the right index as required
+		if(TraitPin->GetPinIndex() != InPinIndex &&
+			InPinIndex >=0 && InPinIndex < InNode->GetPins().Num())
+		{
+			URigVMPin* LastPin =  InNode->Pins.Pop();
+			InNode->Pins.Insert(LastPin, InPinIndex);
+		}
 
-	// move the the pin to the right index as required
-	if(TraitPin->GetPinIndex() != InPinIndex &&
-		InPinIndex >=0 && InPinIndex < InNode->GetPins().Num())
-	{
-		URigVMPin* LastPin =  InNode->Pins.Pop();
-		InNode->Pins.Insert(LastPin, InPinIndex);
+		InNode->UpdateTraitRootPinNames();
 	}
 
-	InNode->UpdateTraitRootPinNames();
+	Notify(ERigVMGraphNotifType::PinAdded, TraitPin);
+
 	Trait->OnTraitAdded(this, InNode);
 
 	if (bSetupUndoRedo)
