RigVM: Refactor last use of GetCPPTypeFromProperty

#rnx

diff --git a/Engine/Plugins/Runtime/RigVM/Source/RigVM/Private/RigVMCore/RigVMFunction.cpp b/Engine/Plugins/Runtime/RigVM/Source/RigVM/Private/RigVMCore/RigVMFunction.cpp
--- a/Engine/Plugins/Runtime/RigVM/Source/RigVM/Private/RigVMCore/RigVMFunction.cpp
+++ b/Engine/Plugins/Runtime/RigVM/Source/RigVM/Private/RigVMCore/RigVMFunction.cpp
@@ -189,10 +189,13 @@ const FRigVMTemplate* FRigVMFunction::GetTemplate_NoLock(const FRigVMRegistryHan
 	{
 		return nullptr;
 	}
-	
-	if(Template->NumPermutations_NoLock(InRegistry) <= 1)
+
+	if (!Template->UsesDispatch())
 	{
-		return nullptr;
+		if(Template->NumPermutations_NoLock(InRegistry) <= 1)
+		{
+			return nullptr;
+		}
 	}
 
 	return Template;
diff --git a/Engine/Plugins/Runtime/RigVM/Source/RigVM/Private/RigVMCore/RigVMRegistry.cpp b/Engine/Plugins/Runtime/RigVM/Source/RigVM/Private/RigVMCore/RigVMRegistry.cpp
--- a/Engine/Plugins/Runtime/RigVM/Source/RigVM/Private/RigVMCore/RigVMRegistry.cpp
+++ b/Engine/Plugins/Runtime/RigVM/Source/RigVM/Private/RigVMCore/RigVMRegistry.cpp
@@ -239,9 +239,8 @@ uint32 FRigVMRegistry_NoLock::GetHashForProperty_NoLock(const FProperty* InPrope
 {
 	uint32 Hash = GetTypeHash(InProperty->GetName());
 
-	FString ExtendedCPPType;
-	const FString CPPType = InProperty->GetCPPType(&ExtendedCPPType);
-	Hash = HashCombine(Hash, GetTypeHash(CPPType + ExtendedCPPType));
+	const FString CPPType = RigVMTypeUtils::GetCPPTypeFromProperty(InProperty);
+	Hash = HashCombine(Hash, GetTypeHash(CPPType));
 	
 	if(const FArrayProperty* ArrayProperty = CastField<FArrayProperty>(InProperty))
 	{
