RigVM: Get rid of obsolete default subobjects

#rnx

diff --git a/Engine/Plugins/Runtime/RigVM/Source/RigVMDeveloper/Private/RigVMBlueprint.cpp b/Engine/Plugins/Runtime/RigVM/Source/RigVMDeveloper/Private/RigVMBlueprint.cpp
--- a/Engine/Plugins/Runtime/RigVM/Source/RigVMDeveloper/Private/RigVMBlueprint.cpp
+++ b/Engine/Plugins/Runtime/RigVM/Source/RigVMDeveloper/Private/RigVMBlueprint.cpp
@@ -25,6 +25,7 @@
 #include "RigVMModel/Nodes/RigVMDispatchNode.h"
 #include "UObject/UE5MainStreamObjectVersion.h"
 #include "Stats/StatsHierarchical.h"
+#include "UObject/UObjectIterator.h"
 
 #include UE_INLINE_GENERATED_CPP_BY_NAME(RigVMBlueprint)
 
@@ -2548,6 +2549,23 @@ void URigVMBlueprint::PreDuplicate(FObjectDuplicationParameters& DupParams)
 		PreDuplicateHostPath.Reset();
 	}
 
+	// look for graphs which are no longer reachable and remove them. this can happen in case
+	// one of the default sub objects got renamed.
+	TArray<URigVMGraph*> AllModels = RigVMClient.GetAllModels(true, true);
+	for (TObjectIterator<URigVMGraph> RigVMIt; RigVMIt; ++RigVMIt)
+	{
+		URigVMGraph* Graph = *RigVMIt;
+		if (Graph->IsInOuter(this))
+		{
+			if (!AllModels.Contains(Graph))
+			{
+				if (URigVMController* Controller = RigVMClient.GetOrCreateController(RigVMClient.GetDefaultModel()))
+				{
+					Controller->DestroyObject(Graph);
+				}
+			}
+		}
+	}
 }
 
 void URigVMBlueprint::ReplaceFunctionIdentifiers(const FString& InOldAssetPath, const FString& InNewAssetPath)
diff --git a/Engine/Plugins/Runtime/RigVM/Source/RigVMDeveloper/Public/RigVMModel/RigVMController.h b/Engine/Plugins/Runtime/RigVM/Source/RigVMDeveloper/Public/RigVMModel/RigVMController.h
--- a/Engine/Plugins/Runtime/RigVM/Source/RigVMDeveloper/Public/RigVMModel/RigVMController.h
+++ b/Engine/Plugins/Runtime/RigVM/Source/RigVMDeveloper/Public/RigVMModel/RigVMController.h
@@ -1691,6 +1691,7 @@ private:
 	friend class URigVMActionStack;
 	friend struct FRigVMBaseAction;
 	friend class URigVMCompiler;
+	friend class URigVMBlueprint;
 	friend struct FRigVMControllerObjectFactory;
 	friend struct FRigVMSetPinDefaultValueAction;
 	friend struct FRigVMAddRerouteNodeAction;
