RigVM: Solidify dispatch factory lookup - add error message for serialization problems

#rnx

diff --git a/Engine/Plugins/Runtime/RigVM/Source/RigVM/Private/RigVMCore/RigVMDispatchFactory.cpp b/Engine/Plugins/Runtime/RigVM/Source/RigVM/Private/RigVMCore/RigVMDispatchFactory.cpp
--- a/Engine/Plugins/Runtime/RigVM/Source/RigVM/Private/RigVMCore/RigVMDispatchFactory.cpp
+++ b/Engine/Plugins/Runtime/RigVM/Source/RigVM/Private/RigVMCore/RigVMDispatchFactory.cpp
@@ -10,11 +10,6 @@
 
 FCriticalSection FRigVMDispatchFactory::GetTemplateMutex;
 
-FName FRigVMDispatchFactory::GetFactoryName() const
-{
-	return *(DispatchPrefix + GetScriptStruct()->GetName());
-}
-
 #if WITH_EDITOR
 
 FName FRigVMDispatchFactory::GetNextAggregateName(const FName& InLastAggregatePinName) const
@@ -41,6 +36,26 @@ FString FRigVMDispatchFactory::GetArgumentMetaData(const FName& InArgumentName,
 	return FString();
 }
 
+FName FRigVMDispatchFactory::GetFactoryName() const
+{
+#if WITH_EDITOR
+	// for unit testing we create factories temporarily and they don't
+	// have the factory name / nor factoryname string set
+	if (FactoryName.IsNone())
+	{
+		FactoryName = GetFactoryName(GetScriptStruct());
+	}
+#endif
+	return FactoryName;
+}
+
+FName FRigVMDispatchFactory::GetFactoryName(const UScriptStruct* InFactoryStruct)
+{
+	check(InFactoryStruct);
+	check(InFactoryStruct->IsChildOf(FRigVMDispatchFactory::StaticStruct()));
+	return *(DispatchPrefix + InFactoryStruct->GetName());
+}
+
 FLinearColor FRigVMDispatchFactory::GetNodeColor() const
 {
 	if(const UScriptStruct* ScriptStruct = GetScriptStruct())
@@ -242,13 +257,12 @@ const FRigVMTemplate* FRigVMDispatchFactory::CreateTemplateForArgumentInfos_NoLo
 	check(InRegistry.GetRegistry() == OwnerRegistry);
 	check(CachedTemplate == nullptr);
 
-	const FName FactoryName = GetFactoryName();
-
+	const FName CopyOfFactoryName = GetFactoryName();
 	FRigVMTemplateDelegates Delegates;
 	Delegates.GetDispatchFactoryDelegate = FRigVMTemplate_GetDispatchFactoryDelegate::CreateLambda(
-	[FactoryName](const FRigVMRegistryHandle& InRegistry)
+	[CopyOfFactoryName](const FRigVMRegistryHandle& InRegistry)
 	{
-		return InRegistry->FindDispatchFactory_NoLock(FactoryName);
+		return InRegistry->FindDispatchFactory_NoLock(CopyOfFactoryName);
 	});
 	
 	CachedTemplate = InRegistry->AddTemplateFromArguments_NoLock(GetFactoryName(), InArguments, Delegates);
@@ -296,7 +310,7 @@ FString FRigVMDispatchFactory::GetPermutationName_NoLock(const FRigVMTemplateTyp
 #if WITH_EDITOR
 	const TArray<FRigVMTemplateArgumentInfo>& Arguments = GetArgumentInfos(InRegistry);
 	
-	checkf(InTypes.Num() == Arguments.Num(), TEXT("Failed getting permutation names for '%s' "), *GetFactoryName().ToString());
+	checkf(InTypes.Num() == Arguments.Num(), TEXT("Failed getting permutation names for '%s' "), *FactoryNameString);
 	
 	for(const FRigVMTemplateArgumentInfo& Argument : Arguments)
 	{
@@ -381,7 +395,7 @@ TArray<FRigVMTemplateArgumentInfo> FRigVMDispatchFactory::BuildArgumentListFromP
 FString FRigVMDispatchFactory::GetPermutationNameImpl(const FRigVMTemplateTypeMap& InTypes, const FRigVMRegistryHandle& InRegistry) const
 {
 	const FString TypePairStrings = FRigVMTemplate::GetStringFromArgumentTypes(InTypes, InRegistry);
-	return RigVMStringUtils::JoinStrings(GetFactoryName().ToString(), TypePairStrings, TEXT("::"));
+	return RigVMStringUtils::JoinStrings(FactoryNameString, TypePairStrings, TEXT("::"));
 }
 
 bool FRigVMDispatchFactory::CopyProperty(const FProperty* InTargetProperty, uint8* InTargetPtr,
@@ -399,8 +413,7 @@ const FRigVMTemplate* FRigVMDispatchFactory::GetTemplate() const
 const FRigVMTemplate* FRigVMDispatchFactory::GetTemplate_NoLock(FRigVMRegistryHandle& InRegistry) const
 {
 	// make sure to rely on the instance of this factory that's stored under the registry
-	const FName FactoryName = GetFactoryName();
-	const FRigVMDispatchFactory* ThisFactory = InRegistry->FindDispatchFactory_NoLock(FactoryName);
+	const FRigVMDispatchFactory* ThisFactory = InRegistry->FindDispatchFactory_NoLock(GetFactoryName());
 	if(ThisFactory != this)
 	{
 		return ThisFactory->GetTemplate_NoLock(InRegistry);
@@ -432,16 +445,17 @@ const FRigVMTemplate* FRigVMDispatchFactory::GetTemplate_NoLock(FRigVMRegistryHa
 		if (Index != INDEX_NONE)
 		{
 			UE_LOG(LogRigVM, Error, TEXT("Failed to add template for dispatch '%s'. Argument '%s' is an execute type."),
-				*FactoryName.ToString(), *Info.Name.ToString());
+				*FactoryNameString, *Info.Name.ToString());
 			return nullptr;			
 		}
 	}
 
+	const FName CopyOfFactoryName = GetFactoryName();
 	FRigVMTemplateDelegates Delegates;
 	Delegates.GetDispatchFactoryDelegate = FRigVMTemplate_GetDispatchFactoryDelegate::CreateLambda(
-	[FactoryName](const FRigVMRegistryHandle& InRegistry)
+	[CopyOfFactoryName](const FRigVMRegistryHandle& InRegistry)
 	{
-		return InRegistry->FindDispatchFactory_NoLock(FactoryName);
+		return InRegistry->FindDispatchFactory_NoLock(CopyOfFactoryName);
 	});
 
 	CachedTemplate = InRegistry->AddTemplateFromArguments_NoLock(GetFactoryName(), Infos, Delegates); 
diff --git a/Engine/Plugins/Runtime/RigVM/Source/RigVM/Private/RigVMCore/RigVMRegistry.cpp b/Engine/Plugins/Runtime/RigVM/Source/RigVM/Private/RigVMCore/RigVMRegistry.cpp
--- a/Engine/Plugins/Runtime/RigVM/Source/RigVM/Private/RigVMCore/RigVMRegistry.cpp
+++ b/Engine/Plugins/Runtime/RigVM/Source/RigVM/Private/RigVMCore/RigVMRegistry.cpp
@@ -276,6 +276,8 @@ void FRigVMRegistry_NoLock::RebuildRegistry_NoLock()
 	Templates.Empty();
 	DeprecatedTemplates.Empty();
 	Factories.Reset();
+	FactoryNameToFactory.Reset();
+	FactoryStructToFactory.Reset();
 	FunctionNameToIndex.Reset();
 	StructNameToPredicates.Reset();
 	TemplateNotationToIndex.Reset();
@@ -540,9 +542,13 @@ void FRigVMRegistry_NoLock::Serialize_NoLock(FArchive& Ar)
 
 				if(const FRigVMDispatchFactory* Factory = RegisterFactory_NoLock(FactoryScriptStruct, FlattenedArguments))
 				{
-					// pulling on this caches the template
+					// pulling on the caches the template
 					Factory->GetTemplate_NoLock(ThisHandle);
 				}
+				else
+				{
+					UE_LOG(LogRigVM, Error, TEXT("Failed to register dispatch factory struct '%s' when loading '%s'"), *FactoryScriptStruct->GetName(), *Ar.GetArchiveState().GetArchiveName());
+				}
 			}
 		}
 
@@ -991,6 +997,8 @@ void FRigVMRegistry_NoLock::Reset_NoLock()
 		FMemory::Free(Factory);
 	}
 	Factories.Reset();
+	FactoryNameToFactory.Reset();
+	FactoryStructToFactory.Reset();
 }
 
 TRigVMTypeIndex FRigVMRegistry_NoLock::FindOrAddType_NoLock(const FRigVMTemplateArgumentType& InType, bool bForce)
@@ -2200,20 +2208,9 @@ const FRigVMDispatchFactory* FRigVMRegistry_NoLock::RegisterFactory_NoLock(UScri
 	check(InFactoryStruct->IsChildOf(FRigVMDispatchFactory::StaticStruct()));
 
 	// ensure to register factories only once
-	const FRigVMDispatchFactory* ExistingFactory = nullptr;
-
-	const bool bFactoryAlreadyRegistered = Factories.ContainsByPredicate([InFactoryStruct, &ExistingFactory](const FRigVMDispatchFactory* Factory)
+	if(const FRigVMDispatchFactory* const* ExistingFactory = FactoryStructToFactory.Find(InFactoryStruct))
 	{
-		if(Factory->GetScriptStruct() == InFactoryStruct)
-		{
-			ExistingFactory = Factory;
-			return true;
-		}
-		return false;
-	});
-	if(bFactoryAlreadyRegistered)
-	{
-		return ExistingFactory;
+		return *ExistingFactory;
 	}
 
 #if WITH_EDITOR
@@ -2227,7 +2224,11 @@ const FRigVMDispatchFactory* FRigVMRegistry_NoLock::RegisterFactory_NoLock(UScri
 	InFactoryStruct->InitializeStruct(Factory, 1);
 	Factory->FactoryScriptStruct = InFactoryStruct;
 	Factory->FactoryIndex = Factories.Add(Factory);
+	Factory->FactoryName = FRigVMDispatchFactory::GetFactoryName(InFactoryStruct);
+	Factory->FactoryNameString = Factory->FactoryName.ToString();
 	Factory->OwnerRegistry = this;
+	FactoryNameToFactory.Add(Factory->FactoryName, Factory);
+	FactoryStructToFactory.Add(InFactoryStruct, Factory);
 	Factory->RegisterDependencyTypes_NoLock(ThisHandle);
 
 	if(!InArgumentInfos.IsEmpty())
@@ -2259,10 +2260,13 @@ bool FRigVMRegistry_NoLock::RemoveFactory_NoLock(UScriptStruct* InFactoryStruct)
 		return false;
 	}
 
+	const FName FactoryName = Factories[FactoryIndex]->GetFactoryName();
 	InFactoryStruct->DestroyStruct(Factories[FactoryIndex], 1);
 	FMemory::Free(Factories[FactoryIndex]);
 	
 	Factories.RemoveAt(FactoryIndex);
+	FactoryNameToFactory.Remove(FactoryName);
+	FactoryStructToFactory.Remove(InFactoryStruct);
 	return true;
 }
 
@@ -2772,32 +2776,41 @@ const FRigVMTemplate* FRigVMRegistry_NoLock::AddTemplateFromArguments_NoLock(con
 
 FRigVMDispatchFactory* FRigVMRegistry_NoLock::FindDispatchFactory_NoLock(const FName& InFactoryName) const
 {
-	FRigVMDispatchFactory* const* FactoryPtr = Factories.FindByPredicate([InFactoryName](const FRigVMDispatchFactory* Factory) -> bool
-	{
-		return Factory->GetFactoryName() == InFactoryName;
-	});
-	if(FactoryPtr)
-	{
-		return *FactoryPtr;
-	}
-
-	if(!bIsGlobalRegistry)
+	if (FRigVMDispatchFactory* const* ExistingFactory = FactoryNameToFactory.Find(InFactoryName))
 	{
-		return nullptr;
+		return *ExistingFactory;
 	}
 	
-	FString FactoryName = InFactoryName.ToString();
-	
-	// if the factory has never been registered - we should try to look it up	
-	if(FactoryName.StartsWith(FRigVMDispatchFactory::DispatchPrefix))
+	if (bIsGlobalRegistry)
 	{
-		const FString ScriptStructName = FactoryName.Mid(FCString::Strlen(FRigVMDispatchFactory::DispatchPrefix));
-		if(UScriptStruct* FactoryStruct = FindFirstObject<UScriptStruct>(*ScriptStructName, EFindFirstObjectOptions::NativeFirst | EFindFirstObjectOptions::EnsureIfAmbiguous))
+		if(InFactoryName.ToString().StartsWith(FRigVMDispatchFactory::DispatchPrefix))
 		{
-			FRigVMRegistry_NoLock* MutableThis = const_cast<FRigVMRegistry_NoLock*>(this);
-			return const_cast<FRigVMDispatchFactory*>(MutableThis->RegisterFactory_NoLock(FactoryStruct));
+			FRigVMDispatchFactory* FoundFactory = nullptr;
+		
+			// if the factory has never been registered - we should try to look it up	
+			for (TObjectIterator<UScriptStruct> ScriptStructIt; ScriptStructIt; ++ScriptStructIt)
+			{
+				static const UScriptStruct* SuperStruct = FRigVMDispatchFactory::StaticStruct();
+
+				UScriptStruct* FactoryStruct = *ScriptStructIt;
+				if ((FactoryStruct == SuperStruct) || !FactoryStruct->IsChildOf(SuperStruct))
+				{
+					continue;
+				}
+
+				// for the global registry make sure to register all factories
+				FRigVMRegistry_NoLock* MutableThis = const_cast<FRigVMRegistry_NoLock*>(this);
+				FRigVMDispatchFactory* Factory = const_cast<FRigVMDispatchFactory*>(MutableThis->RegisterFactory_NoLock(FactoryStruct));
+
+				if (Factory && !FoundFactory && Factory->GetFactoryName().IsEqual(InFactoryName, ENameCase::IgnoreCase))
+				{
+					FoundFactory = Factory;
+				}
+			}
+
+			return FoundFactory;
 		}
-	}	
+	}
 	
 	return nullptr;
 }
diff --git a/Engine/Plugins/Runtime/RigVM/Source/RigVM/Public/RigVMCore/RigVMDispatchFactory.h b/Engine/Plugins/Runtime/RigVM/Source/RigVM/Public/RigVMCore/RigVMDispatchFactory.h
--- a/Engine/Plugins/Runtime/RigVM/Source/RigVM/Public/RigVMCore/RigVMDispatchFactory.h
+++ b/Engine/Plugins/Runtime/RigVM/Source/RigVM/Public/RigVMCore/RigVMDispatchFactory.h
@@ -55,6 +55,7 @@ public:
 
 	FRigVMDispatchFactory()
 		: FactoryIndex(INDEX_NONE)
+		, FactoryName(NAME_None)
 		, FactoryScriptStruct(nullptr)
 		, CachedTemplate(nullptr)
 		, OwnerRegistry(nullptr)
@@ -67,12 +68,15 @@ public:
 		delete ArgumentNamesMutex;
 	}
 
-	// returns the name of the factory template
-	RIGVM_API FName GetFactoryName() const;
-
 	// returns the index of this factory in its registry
 	int32 GetFactoryIndex() const { return FactoryIndex; }
 
+	// returns the name of this factory
+	RIGVM_API FName GetFactoryName() const;
+
+	// returns the name of the factory given the factory's struct 
+	RIGVM_API static FName GetFactoryName(const UScriptStruct* InFactoryStruct);
+
 	// returns the struct for this dispatch factory
 	UScriptStruct* GetScriptStruct() const { return FactoryScriptStruct; };
 
@@ -280,6 +284,8 @@ protected:
 	const TArray<FName>* UpdateArgumentNameCache_NoLock(int32 InNumberOperands, FRigVMRegistryHandle& InRegistry) const;
 
 	int32 FactoryIndex;
+	mutable FName FactoryName;
+	FString FactoryNameString;
 	UScriptStruct* FactoryScriptStruct;
 	mutable const FRigVMTemplate* CachedTemplate;
 	const FRigVMRegistry_NoLock* OwnerRegistry;
diff --git a/Engine/Plugins/Runtime/RigVM/Source/RigVM/Public/RigVMCore/RigVMRegistry.h b/Engine/Plugins/Runtime/RigVM/Source/RigVM/Public/RigVMCore/RigVMRegistry.h
--- a/Engine/Plugins/Runtime/RigVM/Source/RigVM/Public/RigVMCore/RigVMRegistry.h
+++ b/Engine/Plugins/Runtime/RigVM/Source/RigVM/Public/RigVMCore/RigVMRegistry.h
@@ -439,6 +439,12 @@ protected:
 	// memory for all dispatch factories
 	TArray<FRigVMDispatchFactory*> Factories;
 
+	// name lookup for factories
+	TMap<FName, FRigVMDispatchFactory*> FactoryNameToFactory;
+
+	// struct lookup for factories
+	TMap<const UScriptStruct*, FRigVMDispatchFactory*> FactoryStructToFactory;
+
 	// name lookup for functions
 	TMap<FName, int32> FunctionNameToIndex;
 
