RigVM: More forgiving code path for empty type arguments. This can happen if the type is no longer known.

#rnx

diff --git a/Engine/Plugins/Runtime/RigVM/Source/RigVM/Private/RigVMCore/RigVMTemplate.cpp b/Engine/Plugins/Runtime/RigVM/Source/RigVM/Private/RigVMCore/RigVMTemplate.cpp
--- a/Engine/Plugins/Runtime/RigVM/Source/RigVM/Private/RigVMCore/RigVMTemplate.cpp
+++ b/Engine/Plugins/Runtime/RigVM/Source/RigVM/Private/RigVMCore/RigVMTemplate.cpp
@@ -468,7 +468,10 @@ TRigVMTypeIndex FRigVMTemplateArgument::GetTypeIndex_NoLock(const int32 InIndex,
 {
 	if (!bUseCategories)
 	{
-		check(!TypeIndices.IsEmpty())
+		if (TypeIndices.IsEmpty())
+		{
+			return INDEX_NONE;
+		}
 		return TypeIndices.IsValidIndex(InIndex) ? TypeIndices[InIndex] : TypeIndices[0];		
 	}
 
