RigVM: More forgiving code path for RemapLinkedPaths

#rnx

diff --git a/Engine/Plugins/Runtime/RigVM/Source/RigVMDeveloper/Private/RigVMModel/RigVMController.cpp b/Engine/Plugins/Runtime/RigVM/Source/RigVMDeveloper/Private/RigVMModel/RigVMController.cpp
--- a/Engine/Plugins/Runtime/RigVM/Source/RigVMDeveloper/Private/RigVMModel/RigVMController.cpp
+++ b/Engine/Plugins/Runtime/RigVM/Source/RigVMDeveloper/Private/RigVMModel/RigVMController.cpp
@@ -22215,7 +22215,6 @@ TArray<URigVMController::FLinkedPath> URigVMController::RemapLinkedPaths(
 		
 		if(InSettings.bRelayToOrphanPins && (SourcePin != nullptr) && (TargetPin != nullptr))
 		{
-			check(SourcePin->IsLinkedTo(TargetPin) == TargetPin->IsLinkedTo(SourcePin));
 			if (!SourcePin->IsLinkedTo(TargetPin))
 			{
 				if (!URigVMPin::CanLink(SourcePin, TargetPin, nullptr, nullptr, ERigVMPinDirection::IO, true))
