RigVM: Improve error message for array bounds issues

#rnx

diff --git a/Engine/Plugins/Runtime/RigVM/Source/RigVM/Private/RigVMCore/RigVMExecuteContext.cpp b/Engine/Plugins/Runtime/RigVM/Source/RigVM/Private/RigVMCore/RigVMExecuteContext.cpp
--- a/Engine/Plugins/Runtime/RigVM/Source/RigVM/Private/RigVMCore/RigVMExecuteContext.cpp
+++ b/Engine/Plugins/Runtime/RigVM/Source/RigVM/Private/RigVMCore/RigVMExecuteContext.cpp
@@ -1,7 +1,7 @@
 // Copyright Epic Games, Inc. All Rights Reserved.
 
 #include "RigVMCore/RigVMExecuteContext.h"
-#include "RigVMObjectVersion.h"
+#include "RigVMCore/RigVM.h"
 
 #include UE_INLINE_GENERATED_CPP_BY_NAME(RigVMExecuteContext)
 
@@ -188,3 +188,13 @@ void FRigVMExtendedExecuteContext::Initialize(const UScriptStruct* InScriptStruc
 	((FRigVMExecuteContext*)PublicDataScope.GetStructMemory())->ExtendedExecuteContext = this;
 	SetDefaultNameCache();
 }
+
+int32 FRigVMExtendedExecuteContext::GetInstructionIndex() const
+{
+	return GetPublicData<>().InstructionIndex;
+}
+
+FString FRigVMExtendedExecuteContext::GetVMPathName() const
+{
+	return VM ? VM->GetPathName() : TEXT("Unknown VM");
+}
diff --git a/Engine/Plugins/Runtime/RigVM/Source/RigVM/Public/RigVMCore/RigVMExecuteContext.h b/Engine/Plugins/Runtime/RigVM/Source/RigVM/Public/RigVMCore/RigVMExecuteContext.h
--- a/Engine/Plugins/Runtime/RigVM/Source/RigVM/Public/RigVMCore/RigVMExecuteContext.h
+++ b/Engine/Plugins/Runtime/RigVM/Source/RigVM/Public/RigVMCore/RigVMExecuteContext.h
@@ -753,7 +753,9 @@ struct FRigVMExtendedExecuteContext
 
 		if(InOutIndex < 0 || InOutIndex >= InArraySize)
 		{
-			GetPublicData<>().Logf(EMessageSeverity::Error, TEXT("Array Index (%d) out of bounds (count %d)."), InOriginalIndex, InArraySize);
+			const int32 InstructionIndex = GetInstructionIndex();
+			const FString ObjectPath = GetVMPathName();
+			GetPublicData<>().Logf(EMessageSeverity::Error, TEXT("Array Index (%d) out of bounds (count %d).\n('%s', instruction %d)"), InOriginalIndex, InArraySize, *ObjectPath, InstructionIndex);
 			return false;
 		}
 		return true;
@@ -763,10 +765,14 @@ struct FRigVMExtendedExecuteContext
 	{
 		if(InSize < 0 || InSize > GetPublicData<>().RuntimeSettings.MaximumArraySize)
 		{
+			const int32 InstructionIndex = GetInstructionIndex();
+			const FString ObjectPath = GetVMPathName();
 			GetPublicData<>().Logf(EMessageSeverity::Error,
-														 TEXT("Array Size (%d) larger than allowed maximum (%d).\nCheck VMRuntimeSettings in class settings."),
+														 TEXT("Array Size (%d) larger than allowed maximum (%d).\nCheck VMRuntimeSettings in class settings.\n('%s', Instruction %d)."),
 														 InSize,
-														 GetPublicData<>().RuntimeSettings.MaximumArraySize);
+														 GetPublicData<>().RuntimeSettings.MaximumArraySize,
+														 *ObjectPath,
+														 InstructionIndex);
 			return false;
 		}
 		return true;
@@ -814,6 +820,9 @@ struct FRigVMExtendedExecuteContext
 		return 0;
 	}
 
+	RIGVM_API int32 GetInstructionIndex() const;
+	RIGVM_API FString GetVMPathName() const;
+
 	uint32 VMHash = 0;
 
 	FRigVMMemoryStorageStruct WorkMemoryStorage;
