RigVM: Defensive code for RigVMEdGraphSchema

#rnx

diff --git a/Engine/Plugins/Runtime/RigVM/Source/RigVMDeveloper/Private/EdGraph/RigVMEdGraphSchema.cpp b/Engine/Plugins/Runtime/RigVM/Source/RigVMDeveloper/Private/EdGraph/RigVMEdGraphSchema.cpp
--- a/Engine/Plugins/Runtime/RigVM/Source/RigVMDeveloper/Private/EdGraph/RigVMEdGraphSchema.cpp
+++ b/Engine/Plugins/Runtime/RigVM/Source/RigVMDeveloper/Private/EdGraph/RigVMEdGraphSchema.cpp
@@ -1368,14 +1368,20 @@ void URigVMEdGraphSchema::SetNodePosition(UEdGraphNode* Node, const FVector2D& P
 
 	if (const URigVMEdGraphNode* RigNode = Cast<URigVMEdGraphNode>(Node))
 	{
-		RigNode->GetController()->SetNodePosition(RigNode->GetModelNode(), Position, bSetupUndo, false, false);
+		if (URigVMController* Controller = RigNode->GetController())
+		{
+			Controller->SetNodePosition(RigNode->GetModelNode(), Position, bSetupUndo, false, false);
+		}
 	}
 	
 	if (const UEdGraphNode_Comment* CommentNode = Cast<UEdGraphNode_Comment>(Node))
 	{
 		if(const URigVMEdGraph* Graph = CommentNode->GetTypedOuter<URigVMEdGraph>())
 		{
-			Graph->GetController()->SetNodePositionByName(CommentNode->GetFName(), Position, bSetupUndo, false, false);
+			if (URigVMController* Controller = Graph->GetController())
+			{
+				Controller->SetNodePositionByName(CommentNode->GetFName(), Position, bSetupUndo, false, false);
+			}
 		}
 	}
 }
