RigVM: Graph Sections first steps

currently hidden behind a cvar

#rnx

diff --git a/Engine/Plugins/Runtime/RigVM/Source/RigVM/Private/RigVMSettings.cpp b/Engine/Plugins/Runtime/RigVM/Source/RigVM/Private/RigVMSettings.cpp
--- a/Engine/Plugins/Runtime/RigVM/Source/RigVM/Private/RigVMSettings.cpp
+++ b/Engine/Plugins/Runtime/RigVM/Source/RigVM/Private/RigVMSettings.cpp
@@ -13,6 +13,7 @@ URigVMEditorSettings::URigVMEditorSettings(const FObjectInitializer& ObjectIniti
 #endif
 {
 #if WITH_EDITORONLY_DATA
+	bHighlightSimilarNodes = true;
 	bAutoLinkMutableNodes = false;
 #endif
 }
diff --git a/Engine/Plugins/Runtime/RigVM/Source/RigVM/Public/RigVMSettings.h b/Engine/Plugins/Runtime/RigVM/Source/RigVM/Public/RigVMSettings.h
--- a/Engine/Plugins/Runtime/RigVM/Source/RigVM/Public/RigVMSettings.h
+++ b/Engine/Plugins/Runtime/RigVM/Source/RigVM/Public/RigVMSettings.h
@@ -28,6 +28,11 @@ public:
 	
 #if WITH_EDITORONLY_DATA
 
+	// with this checked similar nodes in the graph will be highlighted.
+	// this highlighting can then also be used to functionize items
+	UPROPERTY(EditAnywhere, config, Category = Interaction)
+	bool bHighlightSimilarNodes;
+
 	// When this is checked mutable nodes (nodes with an execute pin)
 	// will be hooked up automatically.
 	UPROPERTY(EditAnywhere, config, Category = Interaction)
diff --git a/Engine/Plugins/Runtime/RigVM/Source/RigVMDeveloper/Private/EdGraph/RigVMEdGraphNode.cpp b/Engine/Plugins/Runtime/RigVM/Source/RigVMDeveloper/Private/EdGraph/RigVMEdGraphNode.cpp
--- a/Engine/Plugins/Runtime/RigVM/Source/RigVMDeveloper/Private/EdGraph/RigVMEdGraphNode.cpp
+++ b/Engine/Plugins/Runtime/RigVM/Source/RigVMDeveloper/Private/EdGraph/RigVMEdGraphNode.cpp
@@ -50,17 +50,15 @@ URigVMEdGraphNode::URigVMEdGraphNode()
 , CachedTitleColor(FLinearColor(0.f, 0.f, 0.f, 0.f))
 , CachedNodeColor(FLinearColor(0.f, 0.f, 0.f, 0.f))
 #if WITH_EDITOR
-, bEnableProfiling(false)
 #endif
 , CachedTemplate(nullptr)
-, MicroSeconds(0)
 , bRenameIsPending(false)
 {
 	bHasCompilerMessage = false;
 	ErrorType = (int32)EMessageSeverity::Info + 1;
 	
 #if WITH_EDITOR
-	UpdateProfilingSettings();
+	UpdateVisualSettings();
 #endif
 }
 
@@ -223,7 +221,7 @@ void URigVMEdGraphNode::ReconstructNode_Internal(bool bForce)
 	}
 
 #if WITH_EDITOR
-	UpdateProfilingSettings();
+	UpdateVisualSettings();
 #endif
 
 	// Clear previously set messages
@@ -1534,59 +1532,6 @@ bool URigVMEdGraphNode::CategoryPinRemoved_Internal(const FString& InCategory)
 	return bResult;
 }
 
-FLinearColor URigVMEdGraphNode::GetNodeProfilingColor() const
-{
-#if WITH_EDITOR
-	if(bEnableProfiling)
-	{
-		if(URigVMBlueprint* Blueprint = GetTypedOuter<URigVMBlueprint>())
-		{
-			if(URigVMHost* DebuggedHost = Cast<URigVMHost>(Blueprint->GetObjectBeingDebugged()))
-			{
-				if(const URigVMNode* ModelNode = GetModelNode())
-				{
-					const double CurrentMicroSeconds = ModelNode->GetInstructionMicroSeconds(DebuggedHost->GetRigVMExtendedExecuteContext(), DebuggedHost->GetVM(), FRigVMASTProxy());
-					MicroSeconds = Blueprint->RigGraphDisplaySettings.AggregateAverage(MicroSecondsFrames, MicroSeconds, CurrentMicroSeconds);
-					if(MicroSeconds >= 0.0)
-					{
-						if(Blueprint->RigGraphDisplaySettings.bAutoDetermineRange)
-						{
-							if(CurrentMicroSeconds < Blueprint->RigGraphDisplaySettings.MinMicroSeconds)
-							{
-								Blueprint->RigGraphDisplaySettings.MinMicroSeconds = CurrentMicroSeconds;
-							}
-							if(CurrentMicroSeconds > Blueprint->RigGraphDisplaySettings.MaxMicroSeconds)
-							{
-								Blueprint->RigGraphDisplaySettings.MaxMicroSeconds = CurrentMicroSeconds;
-							}
-						}
-							
-						const double MinMicroSeconds = Blueprint->RigGraphDisplaySettings.LastMinMicroSeconds;
-						const double MaxMicroSeconds = Blueprint->RigGraphDisplaySettings.LastMaxMicroSeconds;
-						if(MaxMicroSeconds <= MinMicroSeconds)
-						{
-							return FLinearColor::Black;
-						}
-			
-						const FLinearColor& MinColor = Blueprint->RigGraphDisplaySettings.MinDurationColor;
-						const FLinearColor& MaxColor = Blueprint->RigGraphDisplaySettings.MaxDurationColor;
-
-						const double T = (MicroSeconds - MinMicroSeconds) / (MaxMicroSeconds - MinMicroSeconds);
-						return FMath::Lerp<FLinearColor>(MinColor, MaxColor, (float)T);
-					}
-				}
-			}
-		}
-	}
-	else
-	{
-		MicroSeconds = 0;
-		MicroSecondsFrames.Reset();
-	}
-#endif
-	return FLinearColor::Black;
-}
-
 void URigVMEdGraphNode::UpdatePinLists()
 {
 	ExecutePins.Reset();
@@ -1720,12 +1665,15 @@ FLinearColor URigVMEdGraphNode::GetNodeTitleColor() const
 FLinearColor URigVMEdGraphNode::GetNodeBodyTintColor() const
 {
 #if WITH_EDITOR
-	if(bEnableProfiling)
+	if(ProfilingColor.IsSet())
+	{
+		return ProfilingColor.GetValue();
+	}
+	if (HighlightingColor.IsSet())
 	{
-		return GetNodeProfilingColor();
+		return HighlightingColor.GetValue();
 	}
 #endif
-	
 	return CachedNodeColor * GetNodeOpacityColor();
 }
 
@@ -2319,19 +2267,75 @@ TArray<URigVMPin*>& URigVMEdGraphNode::PinListForPin(const URigVMPin* InModelPin
 }
 
 #if WITH_EDITOR
-void URigVMEdGraphNode::UpdateProfilingSettings()
-{
-	bEnableProfiling = false;
-	MicroSeconds = 0;
-	MicroSecondsFrames.Reset();
 
+void URigVMEdGraphNode::UpdateVisualSettings()
+{
+	bIsHighlighted.Reset();
+	ProfilingColor.Reset();
+	HighlightingColor.Reset();
+	bEnableProfiling.Reset();
+	
 	if(const URigVMEdGraph* RigGraph = Cast<URigVMEdGraph>(GetOuter()))
 	{
 		if(const URigVMBlueprint* RigBlueprint = RigGraph->GetBlueprint())
 		{
 			bEnableProfiling = RigBlueprint->VMRuntimeSettings.bEnableProfiling;
+			if (bEnableProfiling.Get(false))
+			{
+				if(URigVMBlueprint* Blueprint = GetTypedOuter<URigVMBlueprint>())
+				{
+					if(URigVMHost* DebuggedHost = Cast<URigVMHost>(Blueprint->GetObjectBeingDebugged()))
+					{
+						if(const URigVMNode* ModelNode = GetModelNode())
+						{
+							const double CurrentMicroSeconds = ModelNode->GetInstructionMicroSeconds(DebuggedHost->GetRigVMExtendedExecuteContext(), DebuggedHost->GetVM(), FRigVMASTProxy());
+							MicroSeconds = Blueprint->RigGraphDisplaySettings.AggregateAverage(MicroSecondsFrames, MicroSeconds.Get(0), CurrentMicroSeconds);
+							if(MicroSeconds.GetValue() >= 0.0)
+							{
+								if(Blueprint->RigGraphDisplaySettings.bAutoDetermineRange)
+								{
+									if(CurrentMicroSeconds < Blueprint->RigGraphDisplaySettings.MinMicroSeconds)
+									{
+										Blueprint->RigGraphDisplaySettings.MinMicroSeconds = CurrentMicroSeconds;
+									}
+									if(CurrentMicroSeconds > Blueprint->RigGraphDisplaySettings.MaxMicroSeconds)
+									{
+										Blueprint->RigGraphDisplaySettings.MaxMicroSeconds = CurrentMicroSeconds;
+									}
+								}
+							
+								const double MinMicroSeconds = Blueprint->RigGraphDisplaySettings.LastMinMicroSeconds;
+								const double MaxMicroSeconds = Blueprint->RigGraphDisplaySettings.LastMaxMicroSeconds;
+								if(MaxMicroSeconds > MinMicroSeconds)
+								{
+									const FLinearColor& MinColor = Blueprint->RigGraphDisplaySettings.MinDurationColor;
+									const FLinearColor& MaxColor = Blueprint->RigGraphDisplaySettings.MaxDurationColor;
+
+									const double T = (MicroSeconds.GetValue() - MinMicroSeconds) / (MaxMicroSeconds - MinMicroSeconds);
+									ProfilingColor = FMath::Lerp<FLinearColor>(MinColor, MaxColor, (float)T);
+								}
+							}
+						}
+					}
+				}
+			}
 		}
 	}
+
+	if (!bEnableProfiling.Get(false))
+	{
+		MicroSeconds.Reset();
+		MicroSecondsFrames.Reset();
+	}
+
+	if (const URigVMNode* ModelNode = GetModelNode())
+	{
+		bIsHighlighted = ModelNode->IsHighlighted();
+	}
+	if (bIsHighlighted.Get(false))
+	{
+		HighlightingColor = FSlateColor(EStyleColor::Highlight).GetSpecifiedColor();
+	}
 }
 
 #endif
diff --git a/Engine/Plugins/Runtime/RigVM/Source/RigVMDeveloper/Private/RigVMModel/RigVMController.cpp b/Engine/Plugins/Runtime/RigVM/Source/RigVMDeveloper/Private/RigVMModel/RigVMController.cpp
--- a/Engine/Plugins/Runtime/RigVM/Source/RigVMDeveloper/Private/RigVMModel/RigVMController.cpp
+++ b/Engine/Plugins/Runtime/RigVM/Source/RigVMDeveloper/Private/RigVMModel/RigVMController.cpp
@@ -35,6 +35,7 @@
 #include "RigVMModel/RigVMClient.h"
 #include "RigVMModel/Nodes/RigVMBranchNode.h"
 #include "RigVMModel/Nodes/RigVMArrayNode.h"
+#include "RigVMModel/RigVMGraphSection.h"
 #include "RigVMCore/RigVMObjectArchive.h"
 #include "Logging/LogScopedVerbosityOverride.h"
 #include "ScopedTransaction.h"
@@ -106,6 +107,7 @@ URigVMController::URigVMController()
 	, bRegisterTemplateNodeUsage(true)
 #endif
 	, bEnableSchemaRemoveNodeCheck(true)
+	, UpdateGraphSectionsBracket(0)
 {
 	SetActionStack(URigVMActionStack::GetDisabledActionStack());
 }
@@ -128,6 +130,7 @@ URigVMController::URigVMController(const FObjectInitializer& ObjectInitializer)
 	, bRegisterTemplateNodeUsage(true)
 #endif
 	, bEnableSchemaRemoveNodeCheck(true)
+	, UpdateGraphSectionsBracket(0)
 {
 	SetActionStack(URigVMActionStack::GetDisabledActionStack());
 }
@@ -7401,6 +7404,8 @@ bool URigVMController::SetNodeSelection(const TArray<FName>& InNodeNames, bool b
 		return false;
 	}
 
+	FRigVMControllerGraphSectionsScope GraphSectionsScope(this);
+
 	URigVMGraph* Graph = GetGraph();
 	check(Graph);
 
@@ -7458,6 +7463,8 @@ bool URigVMController::SetNodeSelection(const TArray<FName>& InNodeNames, bool b
 
 	if (bSelectionChanged)
 	{
+		UpdateGraphSectionsIfRequired();
+		
 		Notify(ERigVMGraphNotifType::NodeSelectionChanged, nullptr);
 	}
 
@@ -12400,6 +12407,8 @@ bool URigVMController::AddLink(URigVMPin* OutputPin, URigVMPin* InputPin, bool b
 		return false;
 	}
 
+	FRigVMControllerGraphSectionsScope GraphSectionsScope(this);
+
 	URigVMGraph* Graph = GetGraph();
 	check(Graph);
 
@@ -12702,6 +12711,10 @@ bool URigVMController::AddLink(URigVMPin* OutputPin, URigVMPin* InputPin, bool b
 	Graph->Links.Add(Link);
 	OutputPin->Links.Add(Link);
 	InputPin->Links.Add(Link);
+	OutputPin->IncrementVersion();
+	InputPin->IncrementVersion();
+
+	UpdateGraphSectionsIfRequired();
 
 	if (!bSuspendNotifications)
 	{
@@ -12848,6 +12861,8 @@ bool URigVMController::BreakLink(const FString& InOutputPinPath, const FString&
 		return false;
 	}
 
+	FRigVMControllerGraphSectionsScope GraphSectionsScope(this);
+
 	URigVMGraph* Graph = GetGraph();
 	check(Graph);
 
@@ -12901,6 +12916,7 @@ bool URigVMController::BreakLink(URigVMPin* OutputPin, URigVMPin* InputPin, bool
 		if (Link->GetSourcePin() == OutputPin && Link->GetTargetPin() == InputPin)
 		{
 			FRigVMControllerCompileBracketScope CompileScope(this);
+			FRigVMControllerGraphSectionsScope GraphSectionsScope(this);
 			FRigVMBreakLinkAction Action;
 			if (bSetupUndoRedo)
 			{
@@ -12950,6 +12966,7 @@ bool URigVMController::BreakAllLinks(const FString& InPinPath, bool bAsInput, bo
 		return false;
 	}
 
+
 	URigVMGraph* Graph = GetGraph();
 	check(Graph);
 
@@ -12991,6 +13008,8 @@ bool URigVMController::BreakAllLinks(URigVMPin* Pin, bool bAsInput, bool bSetupU
 	{
 		return false;
 	}
+
+	FRigVMControllerGraphSectionsScope GraphSectionsScope(this);
 	
 	FRigVMControllerCompileBracketScope CompileScope(this);
 	FRigVMBaseAction Action(this);
@@ -13035,6 +13054,11 @@ bool URigVMController::BreakAllLinks(URigVMPin* Pin, bool bAsInput, bool bSetupU
 		}
 	}
 
+	if (LinksBroken > 0)
+	{
+		UpdateGraphSectionsIfRequired();
+	}
+
 	return LinksBroken > 0;
 }
 
@@ -19753,6 +19777,21 @@ URigVMPin* URigVMController::CreatePinFromPinInfo(const FRigVMRegistry& InRegist
 	return Pin;
 }
 
+void URigVMController::UpdateGraphSectionsIfRequired()
+{
+	if (UpdateGraphSectionsBracket > 1)
+	{
+		return;
+	}
+	
+	URigVMGraph* Graph = GetGraph();
+	if (Graph == nullptr)
+	{
+		return;
+	}
+	Graph->UpdateSections();
+}
+
 void URigVMController::ReportInfo(const FString& InMessage) const
 {
 	if (URigVMGraph* Graph = GetGraph())
@@ -20097,6 +20136,8 @@ bool URigVMController::FullyResolveTemplateNode(URigVMTemplateNode* InNode, int3
 		return true;
 	}
 
+	FRigVMControllerGraphSectionsScope GraphSectionsScope(this);
+
 	FRigVMTemplate* Template = const_cast<FRigVMTemplate*>(InNode->GetTemplate());
 	const FRigVMDispatchFactory* Factory = Template->GetDispatchFactory();
 	
@@ -20376,6 +20417,8 @@ bool URigVMController::FullyResolveTemplateNode(URigVMTemplateNode* InNode, int3
 		CloseUndoBracket();
 	}
 
+	UpdateGraphSectionsIfRequired();
+
 	return true;
 }
 
@@ -20486,6 +20529,8 @@ bool URigVMController::ResolveWildCardPin(URigVMPin* InPin, TRigVMTypeIndex InTy
 		return false;
 	}
 
+	FRigVMControllerGraphSectionsScope GraphSectionsScope(this);
+
 	URigVMPin* RootPin = InPin;
 	TRigVMTypeIndex Type = InTypeIndex;	
 	while (RootPin->IsArrayElement())
@@ -20552,6 +20597,8 @@ bool URigVMController::ResolveWildCardPin(URigVMPin* InPin, TRigVMTypeIndex InTy
 	{
 		GetActionStack()->EndAction(Action);
 	}
+
+	UpdateGraphSectionsIfRequired();
 	
 	return true;
 
@@ -21059,6 +21106,8 @@ bool URigVMController::ChangePinType(URigVMPin* InPin, const FString& InCPPType,
 		return false;
 	}
 
+	FRigVMControllerGraphSectionsScope GraphSectionsScope(this);
+
 	UObject* CPPTypeObject = RigVMTypeUtils::FindObjectFromCPPTypeObjectPath(InCPPTypeObjectPath.ToString());
 
 	// always refresh pin type if it is a user defined struct, whose internal layout can change at anytime
diff --git a/Engine/Plugins/Runtime/RigVM/Source/RigVMDeveloper/Private/RigVMModel/RigVMGraph.cpp b/Engine/Plugins/Runtime/RigVM/Source/RigVMDeveloper/Private/RigVMModel/RigVMGraph.cpp
--- a/Engine/Plugins/Runtime/RigVM/Source/RigVMDeveloper/Private/RigVMModel/RigVMGraph.cpp
+++ b/Engine/Plugins/Runtime/RigVM/Source/RigVMDeveloper/Private/RigVMModel/RigVMGraph.cpp
@@ -11,6 +11,10 @@
 #include "RigVMTypeUtils.h"
 #include "RigVMModel/RigVMClient.h"
 
+#if WITH_EDITOR
+#include "RigVMSettings.h"
+#endif
+
 #include UE_INLINE_GENERATED_CPP_BY_NAME(RigVMGraph)
 
 URigVMGraph::URigVMGraph()
@@ -18,6 +22,7 @@ URigVMGraph::URigVMGraph()
 , RuntimeAST(nullptr)
 , bEditable(true)
 , SchemaClass(nullptr)
+, bSectionsEnabled(false)
 {
 	SetExecuteContextStruct(FRigVMExecuteContext::StaticStruct());
 }
@@ -298,6 +303,18 @@ const TArray<FName>& URigVMGraph::GetSelectNodes() const
 	return SelectedNodes;
 }
 
+bool URigVMGraph::IsNodeHighlighted(const FName& InNodeName) const
+{
+	for (const FRigVMGraphSection& MatchingSection : SectionsMatchingTheSelection)
+	{
+		if (MatchingSection.ContainsNode(InNodeName))
+		{
+			return true;
+		}
+	}
+	return false;
+}
+
 bool URigVMGraph::IsTopLevelGraph() const
 {
 	if (GetOuter()->IsA<URigVMLibraryNode>())
@@ -450,6 +467,48 @@ UScriptStruct* URigVMGraph::GetExecuteContextStruct() const
 	return nullptr;
 }
 
+TArray<FRigVMGraphSection> URigVMGraph::GetMatchingSections(const FRigVMGraphSection& InSection) const
+{
+	TArray<FRigVMGraphSection> Matches;
+	for (const FRigVMGraphSection& Match : SectionsMatchingTheSelection)
+	{
+		if (Match.Hash == InSection.Hash)
+		{
+			Matches.Add(Match);
+		}
+	}
+	return Matches;
+}
+
+void URigVMGraph::UpdateSections(bool bForce)
+{
+	SelectedSections.Reset();
+	SectionsMatchingTheSelection.Reset();
+
+	if (SelectedNodes.IsEmpty())
+	{
+		return;
+	}
+
+	if (!AreSectionsEnabled() && !bForce)
+	{
+		return;
+	}
+
+	SelectedSections = FRigVMGraphSection::GetSectionsFromSelection(this);
+	
+	TArray<URigVMNode*> AvailableNodes = Nodes;
+	AvailableNodes.RemoveAll([this]( const URigVMNode* Node) -> bool
+	{
+		return SelectedNodes.Contains(Node->GetFName()) || !FRigVMGraphSection::IsValidNode(Node);
+	});
+
+	for (const FRigVMGraphSection& SelectedSection : SelectedSections)
+	{
+		SectionsMatchingTheSelection.Append(SelectedSection.FindMatches(AvailableNodes));
+	}
+}
+
 void URigVMGraph::Notify(ERigVMGraphNotifType InNotifType, UObject* InSubject)
 {
 	ModifiedEvent.Broadcast(InNotifType, this, InSubject);
@@ -528,6 +587,21 @@ void URigVMGraph::PreSave(FObjectPreSaveContext SaveContext)
 	LastStructureHash = GetStructureHash();
 }
 
+bool URigVMGraph::AreSectionsEnabled() const
+{
+#if WITH_EDITOR
+	if (const IRigVMClientHost* ClientHost = GetImplementingOuter<IRigVMClientHost>())
+	{
+		const URigVMEditorSettings* Settings = GetDefault<URigVMEditorSettings>(ClientHost->GetRigVMEditorSettingsClass());
+		if (Settings && Settings->bHighlightSimilarNodes)
+		{
+			return true;
+		}
+	}
+#endif
+	return bSectionsEnabled;
+}
+
 bool URigVMGraph::IsNameAvailable(const FString& InName) const
 {
 	for (const URigVMNode* Node : Nodes)
diff --git a/Engine/Plugins/Runtime/RigVM/Source/RigVMDeveloper/Private/RigVMModel/RigVMGraphSection.cpp b/Engine/Plugins/Runtime/RigVM/Source/RigVMDeveloper/Private/RigVMModel/RigVMGraphSection.cpp
new file mode 100644
--- /dev/null
+++ b/Engine/Plugins/Runtime/RigVM/Source/RigVMDeveloper/Private/RigVMModel/RigVMGraphSection.cpp
@@ -0,0 +1,710 @@
+// Copyright Epic Games, Inc. All Rights Reserved.
+
+#include "RigVMModel/RigVMGraphSection.h"
+#include "RigVMModel/RigVMGraph.h"
+#include "RigVMModel/Nodes/RigVMCommentNode.h"
+#include "RigVMModel/Nodes/RigVMDispatchNode.h"
+#include "RigVMModel/Nodes/RigVMFunctionReferenceNode.h"
+#include "RigVMModel/Nodes/RigVMRerouteNode.h"
+#include "RigVMModule.h"
+#include "BoneControllers/BoneControllerTypes.h"
+
+#include UE_INLINE_GENERATED_CPP_BY_NAME(RigVMGraphSection)
+
+//UE_DISABLE_OPTIMIZATION
+
+URigVMPin* FRigVMGraphSectionLink::FindSourcePinSkippingReroutes(const URigVMLink* InLink)
+{
+	check(InLink);
+
+	URigVMPin* SourcePin = InLink->GetSourcePin();
+	if (SourcePin == nullptr)
+	{
+		return nullptr;
+	}
+
+	// traverse the reroute and find the source node
+	while (SourcePin->GetNode()->IsA<URigVMRerouteNode>())
+	{
+		TArray<URigVMPin*> SourcePins = SourcePin->GetLinkedSourcePins();
+		if (SourcePins.Num() == 0)
+		{
+			SourcePin = nullptr;
+			break;
+		}
+		SourcePin = SourcePins[0];
+	}
+
+	return SourcePin; // can be nullptr
+}
+
+TArray<URigVMPin*> FRigVMGraphSectionLink::FindTargetPinsSkippingReroutes(const URigVMLink* InLink)
+{
+	TArray<URigVMPin*> AllTargetPins = {InLink->GetTargetPin()};
+	TArray<URigVMPin*> FilteredTargetPins;
+
+	for (int32 TargetPinIndex = 0; TargetPinIndex < AllTargetPins.Num(); ++TargetPinIndex)
+	{
+		URigVMPin* TargetPin = AllTargetPins[TargetPinIndex];
+		if (TargetPin == nullptr)
+		{
+			continue;
+		}
+		
+		if (TargetPin->GetNode()->IsA<URigVMRerouteNode>())
+		{
+			AllTargetPins.Append(TargetPin->GetLinkedTargetPins());
+			continue;
+		}
+
+		FilteredTargetPins.Add(TargetPin);
+	}
+
+	return FilteredTargetPins;
+}
+
+TArray<FRigVMGraphSectionLink::FPinTuple> FRigVMGraphSectionLink::FindLinksSkippingReroutes(const URigVMNode* InNode)
+{
+	check(InNode);
+	const TArray<URigVMLink*> AllLinks = InNode->GetLinks();
+	TArray<FPinTuple> FilteredLinks;
+	FilteredLinks.Reserve(AllLinks.Num());
+
+	for (const URigVMLink* Link : AllLinks)
+	{
+		const URigVMPin* SourcePin = FindSourcePinSkippingReroutes(Link);
+		if (SourcePin == nullptr)
+		{
+			continue;
+		}
+				
+		const TArray<URigVMPin*> TargetPins = FindTargetPinsSkippingReroutes(Link);
+		for (const URigVMPin* TargetPin : TargetPins)
+		{
+			if (TargetPin == nullptr)
+			{
+				continue;
+			}
+			FilteredLinks.Emplace(SourcePin, TargetPin);
+		}
+	}
+
+	return FilteredLinks;
+}
+
+const URigVMNode* FRigVMGraphSectionLink::FindSourceNode(const URigVMNode* InTargetNode) const
+{
+	check(InTargetNode);
+
+	const URigVMPin* TargetPin = InTargetNode->FindPin(TargetPinPath);
+	if (TargetPin == nullptr)
+	{
+		return nullptr;
+	}
+
+	const TArray<URigVMLink*>& Links = TargetPin->GetLinks();
+	for (const URigVMLink* Link : Links)
+	{
+		const URigVMPin* SourcePin = FindSourcePinSkippingReroutes(Link);
+		if (SourcePin == nullptr)
+		{
+			continue;
+		}
+		
+		if (SourcePin->GetSegmentPath(true) != SourcePinPath)
+		{
+			continue;
+		}
+
+		const URigVMNode* SourceNode = SourcePin->GetNode();
+		if (SourceNode == nullptr)
+		{
+			continue;
+		}
+
+		if (SourceNodeHash != FRigVMGraphSection::GetNodeHash(SourceNode))
+		{
+			continue;
+		}
+
+		return SourceNode;
+	}
+
+	return nullptr;
+}
+
+TArray<const URigVMNode*> FRigVMGraphSectionLink::FindTargetNodes(const URigVMNode* InSourceNode) const
+{
+	check(InSourceNode);
+
+	TArray<const URigVMNode*> TargetNodes;
+
+	const URigVMPin* SourcePin = InSourceNode->FindPin(SourcePinPath);
+	if (SourcePin == nullptr)
+	{
+		return TargetNodes;
+	}
+
+	const TArray<URigVMLink*>& Links = SourcePin->GetLinks();
+	for (const URigVMLink* Link : Links)
+	{
+		const TArray<URigVMPin*> TargetPins = FindTargetPinsSkippingReroutes(Link);
+		for (const URigVMPin* TargetPin : TargetPins)
+		{
+			if (TargetPin->GetSegmentPath(true) != TargetPinPath)
+			{
+				continue;
+			}
+
+			const URigVMNode* TargetNode = TargetPin->GetNode();
+			if (TargetNode == nullptr)
+			{
+				continue;
+			}
+
+			if (TargetNodeHash != FRigVMGraphSection::GetNodeHash(TargetNode))
+			{
+				continue;
+			}
+
+			TargetNodes.Add(TargetNode);
+		}
+	}
+
+	return TargetNodes;
+}
+
+FRigVMGraphSection::FRigVMGraphSection()
+: Hash(UINT32_MAX)
+{
+}
+
+FRigVMGraphSection::FRigVMGraphSection(const TArray<URigVMNode*>& InNodes)
+: Hash(UINT32_MAX)
+{
+	if (InNodes.Contains(nullptr))
+	{
+		return;
+	}
+
+	TArray<URigVMNode*> SortedNodes;
+	SortedNodes.Append(InNodes);
+
+	SortedNodes.RemoveAll([](const URigVMNode* Node) -> bool
+	{
+		return !IsValidNode(Node);
+	});
+
+	if(SortedNodes.IsEmpty())
+	{
+		return;
+	}
+
+	const URigVMGraph* Graph = SortedNodes[0]->GetGraph();
+
+	for(const URigVMNode* Node : SortedNodes)
+	{
+		if(Node->GetGraph() != Graph)
+		{
+			UE_LOG(LogRigVM, Error, TEXT("Provided nodes are not on the same graph."));
+			Reset();
+			return;
+		}
+	}
+
+	// find all links within the section.
+	TArray<bool> bNodeHasOutputLink;
+	bNodeHasOutputLink.AddZeroed(SortedNodes.Num());
+
+	TArray<TArray<FRigVMGraphSectionLink::FPinTuple>> PerNodeInputLinks;
+	PerNodeInputLinks.SetNum(SortedNodes.Num());
+
+	for(int32 TargetNodeIndex = 0; TargetNodeIndex < SortedNodes.Num(); TargetNodeIndex++)
+	{
+		const URigVMNode* TargetNode = SortedNodes[TargetNodeIndex];
+
+		// links are sorted from first pin to last pin
+		const TArray<FRigVMGraphSectionLink::FPinTuple> AllLinks = FRigVMGraphSectionLink::FindLinksSkippingReroutes(TargetNode);
+		PerNodeInputLinks[TargetNodeIndex].Reserve(AllLinks.Num());
+		
+		for (const FRigVMGraphSectionLink::FPinTuple& Link : AllLinks)
+		{
+			const URigVMPin* SourcePin = Link.Get<0>();
+			check(SourcePin);
+
+			const int32 SourceNodeIndex = SortedNodes.Find(SourcePin->GetNode());
+			if (SourceNodeIndex == INDEX_NONE)
+			{
+				continue;
+			}
+
+			const URigVMPin* TargetPin = Link.Get<1>();
+			check(TargetPin);
+
+			if (TargetPin->GetNode() == TargetNode)
+			{
+				PerNodeInputLinks[TargetNodeIndex].Emplace(SourcePin, TargetPin);
+			}
+			else
+			{
+				URigVMNode* OtherTargetNode = TargetPin->GetNode();
+				const int32 OtherTargetNodeIndex = SortedNodes.Find(OtherTargetNode);
+				if  (OtherTargetNodeIndex != INDEX_NONE)
+				{
+					bNodeHasOutputLink[SourceNodeIndex] = true;
+				}
+			}
+		}
+	}
+
+	// determine the leaf nodes
+	LeafNodes.Reserve(SortedNodes.Num());
+	for(int32 NodeIndex = 0; NodeIndex < SortedNodes.Num(); NodeIndex++)
+	{
+		LeafNodes.Add(NodeIndex);
+	}
+	LeafNodes.RemoveAll([bNodeHasOutputLink](const int32& NodeIndex) -> bool
+	{
+		return bNodeHasOutputLink[NodeIndex];
+	});
+	
+	if (LeafNodes.IsEmpty())
+	{
+		Reset();
+		return;
+	}
+
+	TMap<int32, int32> NodeTraversalByInt32;
+	struct Local
+	{
+		static void VisitNode(
+			const TArray<URigVMNode*>& InNodes,
+			int32 NodeIndex,
+			int32& TraversalIndex,
+			const TArray<TArray<FRigVMGraphSectionLink::FPinTuple>>& InPerNodeInputLinks,
+			TMap<int32, int32>& NodeTraversalIndex)
+		{
+			if(NodeTraversalIndex.Contains(NodeIndex))
+			{
+				return;
+			}
+			NodeTraversalIndex.Add(NodeIndex, TraversalIndex++);
+
+			const TArray<FRigVMGraphSectionLink::FPinTuple>& LinksForNode = InPerNodeInputLinks[NodeIndex];
+			for(const FRigVMGraphSectionLink::FPinTuple& LinkForNode : LinksForNode)
+			{
+				const URigVMPin* SourcePin = LinkForNode.Get<0>();
+				if (SourcePin == nullptr)
+				{
+					continue;
+				}
+				
+				URigVMNode* SourceNode = SourcePin->GetNode();
+				const int32 SourceNodeIndex = InNodes.Find(SourceNode);
+				if (SourceNodeIndex != INDEX_NONE)
+				{
+					VisitNode(InNodes, SourceNodeIndex, TraversalIndex, InPerNodeInputLinks, NodeTraversalIndex);
+				}
+			}
+		}
+	};
+
+	// visit all leaf nodes
+	int32 TraversalIndex = 0;
+	TArray<FName> LeafNodeNames;
+	for(const int32& LeafNodeIndex : LeafNodes)
+	{
+		LeafNodeNames.Add(SortedNodes[LeafNodeIndex]->GetFName());
+		Local::VisitNode(SortedNodes, LeafNodeIndex, TraversalIndex, PerNodeInputLinks, NodeTraversalByInt32);
+	}
+
+	// convert the map from int32 to object ptr	
+	TMap<const URigVMNode*, int32> NodeTraversalByPtr;
+	for(int32 NodeIndex = 0; NodeIndex < SortedNodes.Num(); NodeIndex++)
+	{
+		NodeTraversalByPtr.Add(SortedNodes[NodeIndex], NodeTraversalByInt32.FindChecked(NodeIndex)); 
+	}
+
+	// sort the nodes
+	Algo::Sort(SortedNodes, [NodeTraversalByPtr](const URigVMNode* InNodeA, const URigVMNode* InNodeB) -> bool
+	{
+		return NodeTraversalByPtr.FindChecked(InNodeA) < NodeTraversalByPtr.FindChecked(InNodeB);  
+	});
+
+	Nodes.Reserve(SortedNodes.Num());
+	NodeHashes.Reserve(SortedNodes.Num());
+	
+	for(URigVMNode* Node : SortedNodes)
+	{
+		Nodes.Add(Node->GetFName());
+		NodeHashes.Add(GetNodeHash(Node));
+	}
+
+	for(URigVMNode* Node : SortedNodes)
+	{
+		const TArray<FRigVMGraphSectionLink::FPinTuple> PerNodeLinks = FRigVMGraphSectionLink::FindLinksSkippingReroutes(Node);
+		for (const FRigVMGraphSectionLink::FPinTuple& Link : PerNodeLinks)
+		{
+			const URigVMPin* SourcePin = Link.Get<0>();
+			const URigVMPin* TargetPin = Link.Get<1>();
+			if(SourcePin == nullptr || TargetPin == nullptr)
+			{
+				continue;
+			}
+
+			URigVMNode* SourceNode = SourcePin->GetNode();
+			URigVMNode* TargetNode = TargetPin->GetNode();
+
+			if(TargetNode == Node)
+			{
+				if(SortedNodes.Contains(SourceNode))
+				{
+					FRigVMGraphSectionLink LinkInfo;
+					LinkInfo.SourceNodeIndex = SortedNodes.Find(SourceNode);
+					LinkInfo.TargetNodeIndex = SortedNodes.Find(TargetNode);
+					LinkInfo.SourceNodeHash = NodeHashes[LinkInfo.SourceNodeIndex];
+					LinkInfo.TargetNodeHash = NodeHashes[LinkInfo.TargetNodeIndex];
+					LinkInfo.SourcePinPath = SourcePin->GetSegmentPath(true);
+					LinkInfo.TargetPinPath = TargetPin->GetSegmentPath(true);
+
+					Links.Add(LinkInfo);
+				}
+			}
+		}
+	}
+
+	LeafNodes.Reset();
+	Algo::Transform(LeafNodeNames, LeafNodes, [this](const FName& LeafNodeName) -> int32
+	{
+		return Nodes.Find(LeafNodeName);
+	});
+
+	Hash = 0;
+	for(uint32 NodeHash : NodeHashes)
+	{
+		Hash = HashCombine(Hash, NodeHash);
+	}
+	for (const FRigVMGraphSectionLink& Link : Links)
+	{
+		Hash = HashCombine(Hash, GetLinkHash(Link));
+	}
+}
+
+TArray<FRigVMGraphSection> FRigVMGraphSection::GetSectionsFromSelection(const URigVMGraph* InGraph)
+{
+	TArray<FRigVMGraphSection> Sections;
+	TArray<URigVMNode*> AllSelectedNodes;
+	Algo::Transform(InGraph->GetSelectNodes(), AllSelectedNodes, [InGraph](const FName& NodeName) -> URigVMNode*
+	{
+		return InGraph->FindNodeByName(NodeName);
+	});
+
+	AllSelectedNodes.Remove(nullptr);
+	if(AllSelectedNodes.Num() == 0)
+	{
+		return Sections;
+	}
+
+	const TArray<TArray<URigVMNode*>> Islands = FRigVMGraphSection::GetNodeIslands(AllSelectedNodes);
+	for (const TArray<URigVMNode*>& Island : Islands)
+	{
+		Sections.Emplace(Island);
+	}
+	
+	return Sections;
+}
+
+bool FRigVMGraphSection::IsValid() const
+{
+	return Hash != UINT32_MAX && !Nodes.IsEmpty();
+}
+
+bool FRigVMGraphSection::IsValidNode(const URigVMNode* InNode)
+{
+	if (InNode)
+	{
+		return !InNode->IsA<URigVMCommentNode>() &&
+			!InNode->IsA<URigVMRerouteNode>();
+	}
+	return true;
+}
+
+void FRigVMGraphSection::Reset()
+{
+	Hash = UINT32_MAX;
+	Nodes.Reset();
+	NodeHashes.Reset();
+	LeafNodes.Reset();
+	Links.Reset();
+}
+
+uint32 FRigVMGraphSection::GetNodeHash(const URigVMNode* InNode)
+{
+	check(InNode);
+
+	uint32 Hash = GetTypeHash(InNode->GetClass());
+	
+	if(const URigVMFunctionReferenceNode* FunctionRefNode = Cast<URigVMFunctionReferenceNode>(InNode))
+	{
+		const FRigVMGraphFunctionIdentifier Identifier = FunctionRefNode->GetFunctionIdentifier();
+		if(Identifier.IsValid())
+		{
+			Hash = HashCombine(Hash, GetTypeHash(Identifier));
+		}
+	}
+	else if(const URigVMTemplateNode* TemplateNode = Cast<URigVMTemplateNode>(InNode))
+	{
+		if(const FRigVMFunction* Function = TemplateNode->GetResolvedFunction())
+		{
+			Hash = HashCombine(Hash, GetTypeHash(Function->Name));
+		}
+		else if(const FRigVMTemplate* Template = TemplateNode->GetTemplate())
+		{
+			Hash = HashCombine(Hash, GetTypeHash(Template->GetNotation()));
+			
+			const FRigVMTemplateTypeMap Types = TemplateNode->GetTemplatePinTypeMap();
+			for(const TPair<FName, TRigVMTypeIndex>& Pair : Types)
+			{
+				Hash = HashCombine(Hash, GetTypeHash(Pair.Key));
+				Hash = HashCombine(Hash, GetTypeHash((int32)Pair.Value));
+			}
+		}
+	}
+	else if(const URigVMVariableNode* VariableNode = Cast<URigVMVariableNode>(InNode))
+	{
+		Hash = HashCombine(Hash, GetTypeHash(VariableNode->GetVariableName()));
+		Hash = HashCombine(Hash, GetTypeHash(VariableNode->IsGetter()));
+	}
+
+	return Hash;
+}
+
+uint32 FRigVMGraphSection::GetLinkHash(const FRigVMGraphSectionLink& InLink)
+{
+	uint32 Hash = HashCombine(GetTypeHash(InLink.SourceNodeIndex), GetTypeHash(InLink.TargetNodeIndex));
+	Hash = HashCombine(Hash, GetTypeHash(InLink.SourcePinPath));
+	Hash = HashCombine(Hash, GetTypeHash(InLink.TargetPinPath));
+	return Hash;
+}
+
+TArray<FRigVMGraphSection> FRigVMGraphSection::FindMatches(const URigVMGraph* InGraph) const
+{
+	return FindMatches(InGraph->GetNodes());
+}
+
+TArray<FRigVMGraphSection> FRigVMGraphSection::FindMatches(const TArray<URigVMNode*>& InAvailableNodes) const
+{
+	TArray<URigVMNode*> AvailableNodes = InAvailableNodes;
+	TArray<FRigVMGraphSection> Matches;
+	TArray<URigVMNode*> MatchedNodes;
+
+	bool bFoundMatch = true;
+	while (bFoundMatch)
+	{
+		bFoundMatch = false;
+		
+		for (const URigVMNode* AvailableNode : AvailableNodes)
+		{
+			MatchedNodes.Reset();
+			if (MatchesNode(AvailableNode, AvailableNodes, &MatchedNodes))
+			{
+				FRigVMGraphSection Match(MatchedNodes);
+				if (Match.IsValid() && Match.Hash == Hash)
+				{
+					Matches.Add(Match);
+					AvailableNodes.RemoveAll([&MatchedNodes]( const URigVMNode* Node) -> bool
+					{
+						return MatchedNodes.Contains(Node);
+					});
+					bFoundMatch = true;
+					break;
+				}
+			}
+		}
+	}
+	return Matches;
+}
+
+TArray<TArray<URigVMNode*>> FRigVMGraphSection::GetNodeIslands(const TArray<URigVMNode*>& InNodes)
+{
+	TArray<TArray<URigVMNode*>> Islands;
+	
+	// compute islands of connected nodes - each subset is going
+	// to represent one island
+	TArray<int32> IslandIndexPerNode;
+	IslandIndexPerNode.AddZeroed(InNodes.Num());
+
+	while (IslandIndexPerNode.Contains(0))
+	{
+		TArray<int32> NodeIndices = {IslandIndexPerNode.Find(0)};
+		Islands.Add({});
+
+		for (int32 Index = 0; Index < NodeIndices.Num(); Index++)
+		{
+			const int32& NodeIndex = NodeIndices[Index];
+			URigVMNode* Node = InNodes[NodeIndex];
+
+			Islands.Last().Add(Node);
+			IslandIndexPerNode[NodeIndex] = Islands.Num();
+
+			TArray<URigVMNode*> LinkedNodes = Node->GetLinkedSourceNodes();
+			LinkedNodes.Append(Node->GetLinkedTargetNodes());
+
+			for (URigVMNode* LinkedNode : LinkedNodes)
+			{
+				const int32 LinkedNodeIndex = InNodes.Find(LinkedNode);
+				if (LinkedNodeIndex == INDEX_NONE)
+				{
+					continue;
+				}
+
+				if (IslandIndexPerNode[LinkedNodeIndex] == 0)
+				{
+					NodeIndices.AddUnique(LinkedNodeIndex);
+				}
+			}
+		}
+	}
+
+	return Islands;
+}
+
+bool FRigVMGraphSection::ContainsNode(const FName& InNodeName) const
+{
+	return Nodes.Contains(InNodeName);
+}
+
+bool FRigVMGraphSection::MatchesNode(const URigVMNode* InNode, const TArray<URigVMNode*>& InAvailableNodes, TArray<URigVMNode*>* OutMatchingNodesForSet) const
+{
+	if (!IsValidNode(InNode))
+	{
+		return false;
+	}
+	
+	const uint32 NodeHash = GetNodeHash(InNode);
+	for (int32 NodeIndex = 0; NodeIndex < Nodes.Num(); NodeIndex++)
+	{
+		if (NodeHash != NodeHashes[NodeIndex])
+		{
+			continue;
+		}
+		
+		TArray<URigVMNode*> VisitedNodes;
+		VisitedNodes.AddZeroed(Nodes.Num());
+		if (MatchesNode_Impl(InNode, InAvailableNodes, NodeIndex, VisitedNodes))
+		{
+			if (OutMatchingNodesForSet)
+			{
+				OutMatchingNodesForSet->Append(VisitedNodes);
+			}
+			return true;
+		}
+	}
+
+	return false;
+}
+
+bool FRigVMGraphSection::MatchesNode_Impl(const URigVMNode* InNode, const TArray<URigVMNode*>& InAvailableNodes, int32 InNodeIndex, TArray<URigVMNode*>& VisitedNodes) const
+{
+	if (!InAvailableNodes.Contains(InNode))
+	{
+		return false;
+	}
+	if (VisitedNodes[InNodeIndex])
+	{
+		return true;
+	}
+	if (VisitedNodes.Contains(InNode))
+	{
+		return false;
+	}
+	VisitedNodes[InNodeIndex] = const_cast<URigVMNode*>(InNode);
+
+	// traverse all right to left
+	for (const FRigVMGraphSectionLink& Link : Links)
+	{
+		if (Link.TargetNodeIndex == InNodeIndex)
+		{
+			if (VisitedNodes[Link.SourceNodeIndex])
+			{
+				continue;
+			}
+
+			if (const URigVMNode* LinkedNode = Link.FindSourceNode(InNode))
+			{
+				if (!MatchesNode_Impl(LinkedNode, InAvailableNodes, Link.SourceNodeIndex, VisitedNodes))
+				{
+					VisitedNodes[InNodeIndex] = nullptr;
+					return false;
+				}
+			}
+			else
+			{
+				VisitedNodes[InNodeIndex] = nullptr;
+				return false;
+			}
+		}
+	}
+
+	// also traverse all links from left to right
+	for (const FRigVMGraphSectionLink& Link : Links)
+	{
+		if (Link.SourceNodeIndex == InNodeIndex)
+		{
+			if (VisitedNodes[Link.TargetNodeIndex])
+			{
+				continue;
+			}
+			
+			TArray<const URigVMNode*> LinkedNodes = Link.FindTargetNodes(InNode);
+			if (LinkedNodes.Num() == 0)
+			{
+				VisitedNodes[InNodeIndex] = nullptr;
+				return false;
+			}
+
+			if (LinkedNodes.Num() == 1)
+			{
+				if (!MatchesNode_Impl(LinkedNodes[0], InAvailableNodes, Link.TargetNodeIndex, VisitedNodes))
+				{
+					VisitedNodes[InNodeIndex] = nullptr;
+					return false;
+				}
+			}
+
+			// if we have more than one match for the linked nodes - we'll need to see which
+			// one may be a potential match. for that we'll need to create a copy of the visited
+			// nodes array since it may be filled up to a point - just to realize that we hit
+			// a mismatching node later.
+			else
+			{
+				const TArray<URigVMNode*> VisitedNodesCopy = VisitedNodes;
+
+				bool bFoundMatch = false;
+				for (const URigVMNode* LinkedNode : LinkedNodes)
+				{
+					if (!MatchesNode_Impl(LinkedNode, InAvailableNodes, Link.TargetNodeIndex, VisitedNodes))
+					{
+						// roll back the visited nodes array
+						VisitedNodes = VisitedNodesCopy;
+					}
+					else
+					{
+						bFoundMatch = true;
+						break;
+					}
+				}
+
+				if (!bFoundMatch)
+				{
+					VisitedNodes[InNodeIndex] = nullptr;
+					return false;
+				}
+			}
+		}
+	}
+
+	return true;
+}
+
+//UE_ENABLE_OPTIMIZATION
diff --git a/Engine/Plugins/Runtime/RigVM/Source/RigVMDeveloper/Private/RigVMModel/RigVMNode.cpp b/Engine/Plugins/Runtime/RigVM/Source/RigVMDeveloper/Private/RigVMModel/RigVMNode.cpp
--- a/Engine/Plugins/Runtime/RigVM/Source/RigVMDeveloper/Private/RigVMModel/RigVMNode.cpp
+++ b/Engine/Plugins/Runtime/RigVM/Source/RigVMDeveloper/Private/RigVMModel/RigVMNode.cpp
@@ -734,14 +734,22 @@ void URigVMNode::IncrementVersion()
 
 bool URigVMNode::IsSelected() const
 {
-	URigVMGraph* Graph = GetGraph();
-	if (Graph)
+	if (const URigVMGraph* Graph = GetGraph())
 	{
 		return Graph->IsNodeSelected(GetFName());
 	}
 	return false;
 }
 
+bool URigVMNode::IsHighlighted() const
+{
+	if (const URigVMGraph* Graph = GetGraph())
+	{
+		return Graph->IsNodeHighlighted(GetFName());
+	}
+	return false;
+}
+
 bool URigVMNode::IsInjected() const
 {
 	return Cast<URigVMInjectionInfo>(GetOuter()) != nullptr;
diff --git a/Engine/Plugins/Runtime/RigVM/Source/RigVMDeveloper/Private/RigVMModel/RigVMPin.cpp b/Engine/Plugins/Runtime/RigVM/Source/RigVMDeveloper/Private/RigVMModel/RigVMPin.cpp
--- a/Engine/Plugins/Runtime/RigVM/Source/RigVMDeveloper/Private/RigVMModel/RigVMPin.cpp
+++ b/Engine/Plugins/Runtime/RigVM/Source/RigVMDeveloper/Private/RigVMModel/RigVMPin.cpp
@@ -647,13 +647,10 @@ void URigVMPin::IncrementVersion(bool bAffectParentPin, bool bAffectSubPins)
 	if(IsRootPin())
 	{
 		// pin changes may affect the event name on a node
-		// so we need to invalidate the node as well.
+		// or the highlight state
 		if(URigVMNode* Node = GetNode())
 		{
-			if(Node->IsEvent())
-			{
-				Node->IncrementVersion();
-			}
+			Node->IncrementVersion();
 		}
 	}
 }
diff --git a/Engine/Plugins/Runtime/RigVM/Source/RigVMDeveloper/Public/EdGraph/RigVMEdGraphNode.h b/Engine/Plugins/Runtime/RigVM/Source/RigVMDeveloper/Public/EdGraph/RigVMEdGraphNode.h
--- a/Engine/Plugins/Runtime/RigVM/Source/RigVMDeveloper/Public/EdGraph/RigVMEdGraphNode.h
+++ b/Engine/Plugins/Runtime/RigVM/Source/RigVMDeveloper/Public/EdGraph/RigVMEdGraphNode.h
@@ -238,7 +238,6 @@ public:
 
 protected:
 
-	UE_API FLinearColor GetNodeProfilingColor() const;
 	UE_API FLinearColor GetNodeOpacityColor() const;
 
 	/** Helper function for AllocateDefaultPins */
@@ -286,15 +285,18 @@ private:
 
 	UE_API TArray<URigVMPin*>& PinListForPin(const URigVMPin* InModelPin);
 
-#if WITH_EDITOR
-	UE_API void UpdateProfilingSettings();
-#endif
-	
 	FLinearColor CachedTitleColor;
 	FLinearColor CachedNodeColor;
 
 #if WITH_EDITOR
-	bool bEnableProfiling;
+	UE_API void UpdateVisualSettings();
+	
+	TOptional<bool> bIsHighlighted;
+	TOptional<FLinearColor> ProfilingColor;
+	TOptional<FLinearColor> HighlightingColor;
+	TOptional<bool> bEnableProfiling;
+	mutable TOptional<double> MicroSeconds;
+	mutable TArray<double> MicroSecondsFrames;
 #endif
 
 	TArray<URigVMPin*> ExecutePins;
@@ -316,8 +318,6 @@ private:
 
 	mutable const FRigVMTemplate* CachedTemplate;
 	mutable TOptional<bool> DrawAsCompactNodeCache;
-	mutable double MicroSeconds;
-	mutable TArray<double> MicroSecondsFrames;
 	FSimpleDelegate RequestRenameDelegate;
 	mutable bool bRenameIsPending;
 
diff --git a/Engine/Plugins/Runtime/RigVM/Source/RigVMDeveloper/Public/RigVMModel/RigVMController.h b/Engine/Plugins/Runtime/RigVM/Source/RigVMDeveloper/Public/RigVMModel/RigVMController.h
--- a/Engine/Plugins/Runtime/RigVM/Source/RigVMDeveloper/Public/RigVMModel/RigVMController.h
+++ b/Engine/Plugins/Runtime/RigVM/Source/RigVMDeveloper/Public/RigVMModel/RigVMController.h
@@ -1533,6 +1533,7 @@ protected:
 
 	UE_API IRigVMClientHost* GetClientHost_Internal(const URigVMGraph* InGraph) const;
 	UE_API URigVMPin* CreatePinFromPinInfo(const FRigVMRegistry& InRegistry, const FRigVMPinInfoArray& InPreviousPinInfos, const FRigVMPinInfo& InPinInfo, const FString& InPinPath, UObject* InOuter);
+	void UpdateGraphSectionsIfRequired();
 
 	// backwards compatibility code
 	UE_API FRigVMClientPatchResult PatchRerouteNodesOnLoad();
@@ -1686,6 +1687,8 @@ private:
 
 	bool bEnableSchemaRemoveNodeCheck;
 
+	int32 UpdateGraphSectionsBracket;
+
 	friend class URigVMGraph;
 	friend class URigVMPin;
 	friend class URigVMActionStack;
@@ -1708,6 +1711,7 @@ private:
 	friend struct FRigVMActionWrapper;
 	friend class URigVMSchema;
 	friend class URigVMEdGraphFunctionRefNodeSpawner;
+	friend class FRigVMControllerGraphSectionsScope;
 };
 
 class FRigVMControllerNotifGuard
@@ -1790,4 +1794,29 @@ public:
 	TMap<FName,FRigVMController_CommonTypePerTemplate> TemplateDefaultTypes;
 };
 
+class FRigVMControllerGraphSectionsScope
+{
+public:
+	FRigVMControllerGraphSectionsScope(URigVMController* InController)
+	: Controller(InController)
+	{
+		if (Controller)
+		{
+			Controller->UpdateGraphSectionsBracket++;
+		}
+	}
+
+	~FRigVMControllerGraphSectionsScope()
+	{
+		if (Controller)
+		{
+			Controller->UpdateGraphSectionsBracket--;
+		}
+	}
+
+private:
+
+	URigVMController* Controller;
+};
+
 #undef UE_API
diff --git a/Engine/Plugins/Runtime/RigVM/Source/RigVMDeveloper/Public/RigVMModel/RigVMGraph.h b/Engine/Plugins/Runtime/RigVM/Source/RigVMDeveloper/Public/RigVMModel/RigVMGraph.h
--- a/Engine/Plugins/Runtime/RigVM/Source/RigVMDeveloper/Public/RigVMModel/RigVMGraph.h
+++ b/Engine/Plugins/Runtime/RigVM/Source/RigVMDeveloper/Public/RigVMModel/RigVMGraph.h
@@ -2,7 +2,7 @@
 
 #pragma once
 
-#include "RigVMNode.h"
+#include "RigVMGraphSection.h"
 #include "RigVMLink.h"
 #include "RigVMNotifications.h"
 #include "Nodes/RigVMVariableNode.h"
@@ -18,6 +18,7 @@
 class URigVMFunctionLibrary;
 struct FRigVMClient;
 
+
 /**
  * The Graph represents a Function definition
  * using Nodes as statements.
@@ -121,6 +122,10 @@ public:
 	UFUNCTION(BlueprintCallable, Category = RigVMGraph)
 	UE_API const TArray<FName>& GetSelectNodes() const;
 
+	// Returns true if a node matches the currently selected subset
+	UFUNCTION(BlueprintCallable, Category = RigVMGraph)
+	UE_API bool IsNodeHighlighted(const FName& InNodeName) const;
+
 	// Returns true if this graph is the top level graph
 	UFUNCTION(BlueprintCallable, Category = RigVMGraph)
 	UE_API bool IsTopLevelGraph() const;
@@ -179,6 +184,12 @@ public:
 
 	UE_API virtual void PreSave(FObjectPreSaveContext SaveContext) override;
 
+   	UE_API bool AreSectionsEnabled() const;
+	UE_API const TArray<FRigVMGraphSection>& GetSelectedSections() const { return SelectedSections; }
+	UE_API const TArray<FRigVMGraphSection>& GetSectionsMatchingTheSelection() const { return SectionsMatchingTheSelection; }
+	UE_API TArray<FRigVMGraphSection> GetMatchingSections(const FRigVMGraphSection& InSection) const;
+	UE_API void UpdateSections(bool bForce = false);
+
 private:
 
 	FRigVMGraphModifiedEvent ModifiedEvent;
@@ -222,6 +233,15 @@ private:
 	UPROPERTY()
 	TSubclassOf<URigVMSchema> SchemaClass;
 
+	UPROPERTY(transient)
+	bool bSectionsEnabled;
+
+	UPROPERTY(transient)
+	TArray<FRigVMGraphSection> SelectedSections;
+
+	UPROPERTY(transient)
+	TArray<FRigVMGraphSection> SectionsMatchingTheSelection;
+
 	UE_API bool IsNameAvailable(const FString& InName) const;
 
 	friend class URigVMController;
diff --git a/Engine/Plugins/Runtime/RigVM/Source/RigVMDeveloper/Public/RigVMModel/RigVMGraphSection.h b/Engine/Plugins/Runtime/RigVM/Source/RigVMDeveloper/Public/RigVMModel/RigVMGraphSection.h
new file mode 100644
--- /dev/null
+++ b/Engine/Plugins/Runtime/RigVM/Source/RigVMDeveloper/Public/RigVMModel/RigVMGraphSection.h
@@ -0,0 +1,93 @@
+// Copyright Epic Games, Inc. All Rights Reserved.
+
+#pragma once
+
+#include "RigVMNode.h"
+#include "RigVMGraphSection.generated.h"
+
+#define UE_API RIGVMDEVELOPER_API
+
+class URigVMGraph;
+
+USTRUCT()
+struct FRigVMGraphSectionLink
+{
+	GENERATED_BODY()
+
+	UPROPERTY()
+	int32 SourceNodeIndex = INDEX_NONE;
+	
+	UPROPERTY()
+	int32 TargetNodeIndex = INDEX_NONE;
+
+	UPROPERTY()
+	uint32 SourceNodeHash = UINT32_MAX;
+	
+	UPROPERTY()
+	uint32 TargetNodeHash = UINT32_MAX;
+
+	UPROPERTY()
+	FString SourcePinPath;
+	
+	UPROPERTY()
+	FString TargetPinPath;
+
+	using FPinTuple = TTuple<const URigVMPin*, const URigVMPin*>;
+	static URigVMPin* FindSourcePinSkippingReroutes(const URigVMLink* InLink);
+	static TArray<URigVMPin*> FindTargetPinsSkippingReroutes(const URigVMLink* InLink);
+	static TArray<FPinTuple> FindLinksSkippingReroutes(const URigVMNode* InNode);
+	const URigVMNode* FindSourceNode(const URigVMNode* InTargetNode) const;
+	TArray<const URigVMNode*> FindTargetNodes(const URigVMNode* InSourceNode) const;
+};
+
+USTRUCT()
+struct FRigVMGraphSection
+{
+	GENERATED_BODY()
+
+public:
+
+	UE_API FRigVMGraphSection();
+	UE_API FRigVMGraphSection(const TArray<URigVMNode*>& InNodes);
+
+	UE_API static TArray<FRigVMGraphSection> GetSectionsFromSelection(const URigVMGraph* InGraph);
+	
+	UE_API bool IsValid() const;
+	UE_API static bool IsValidNode(const URigVMNode* InNode);
+
+	UPROPERTY()
+	uint32 Hash;
+
+	UPROPERTY()
+	TArray<FName> Nodes;
+
+	UPROPERTY()
+	TArray<uint32> NodeHashes;
+
+	UPROPERTY()
+	TArray<int32> LeafNodes;
+
+	UPROPERTY()
+	TArray<FRigVMGraphSectionLink> Links;
+
+	UE_API TArray<FRigVMGraphSection> FindMatches(const URigVMGraph* InGraph) const;
+	UE_API TArray<FRigVMGraphSection> FindMatches(const TArray<URigVMNode*>& InAvailableNodes) const;
+	
+	UE_API bool ContainsNode(const FName& InNodeName) const;
+	UE_API bool MatchesNode(const URigVMNode* InNode, const TArray<URigVMNode*>& InAvailableNodes, TArray<URigVMNode*>* OutMatchingNodesForSet = nullptr) const;
+
+private:
+
+	void Reset();
+
+	static TArray<TArray<URigVMNode*>> GetNodeIslands(const TArray<URigVMNode*>& InNodes);
+	
+	static uint32 GetNodeHash(const URigVMNode* InNode);
+	static uint32 GetLinkHash(const FRigVMGraphSectionLink& InLink);
+	
+	bool MatchesNode_Impl(const URigVMNode* InNode, const TArray<URigVMNode*>& InAvailableNodes, int32 InNodeIndex, TArray<URigVMNode*>& VisitedNodes) const;
+
+	friend struct FRigVMGraphSectionLink;
+};
+
+#undef UE_API
diff --git a/Engine/Plugins/Runtime/RigVM/Source/RigVMDeveloper/Public/RigVMModel/RigVMNode.h b/Engine/Plugins/Runtime/RigVM/Source/RigVMDeveloper/Public/RigVMModel/RigVMNode.h
--- a/Engine/Plugins/Runtime/RigVM/Source/RigVMDeveloper/Public/RigVMModel/RigVMNode.h
+++ b/Engine/Plugins/Runtime/RigVM/Source/RigVMDeveloper/Public/RigVMModel/RigVMNode.h
@@ -209,6 +209,10 @@ public:
 	UFUNCTION(BlueprintCallable, Category = RigVMNode)
 	UE_API bool IsSelected() const;
 
+	// Returns true if this Node should be highlighted
+	UFUNCTION(BlueprintCallable, Category = RigVMNode)
+	UE_API bool IsHighlighted() const;
+
 	// Returns true if this is an injected node.
 	// Injected nodes are managed by pins are are not visible to the user.
 	UFUNCTION(BlueprintCallable, Category = RigVMNode)
diff --git a/Engine/Plugins/Runtime/RigVM/Source/RigVMEditor/Private/RigVMEditorModule.cpp b/Engine/Plugins/Runtime/RigVM/Source/RigVMEditor/Private/RigVMEditorModule.cpp
--- a/Engine/Plugins/Runtime/RigVM/Source/RigVMEditor/Private/RigVMEditorModule.cpp
+++ b/Engine/Plugins/Runtime/RigVM/Source/RigVMEditor/Private/RigVMEditorModule.cpp
@@ -1068,7 +1068,7 @@ void FRigVMEditorModule::GetNodeTemplatesContextMenuActions(IRigVMClientHost* Ri
 void FRigVMEditorModule::GetNodeOrganizationContextMenuActions(IRigVMClientHost* RigVMClientHost, const URigVMEdGraphNode* EdGraphNode, URigVMNode* ModelNode, UToolMenu* Menu) const
 {
 	URigVMBlueprint* RigVMBlueprint = Cast<URigVMBlueprint>(RigVMClientHost);
-	const URigVMGraph* Model = ModelNode->GetGraph();
+	URigVMGraph* Model = ModelNode->GetGraph();
 	URigVMController* Controller = RigVMClientHost->GetRigVMClient()->GetController(Model);
 
 	FToolMenuSection& OrganizationSection = Menu->AddSection("RigVMEditorContextMenuOrganization", LOCTEXT("OrganizationHeader", "Organization"));
@@ -1126,6 +1126,37 @@ void FRigVMEditorModule::GetNodeOrganizationContextMenuActions(IRigVMClientHost*
 		})
 	));
 
+	if (Model->AreSectionsEnabled() && !Model->GetSectionsMatchingTheSelection().IsEmpty())
+	{
+		OrganizationSection.AddMenuEntry(
+			"Collapse to Function (All Matches)",
+			LOCTEXT("CollapseNodesToFunctionAllMatches", "Collapse to Function (All Matches)"),
+			LOCTEXT("CollapseNodesToFunctionAllMatches_Tooltip", "Turns the selected nodes into a new Function while refactoring all matches"),
+			FSlateIcon(),
+			FUIAction(FExecuteAction::CreateLambda([Model, Controller]() {
+				// todo: this is doing the same thing as collapse to function right now
+				TArray<FName> Nodes = Model->GetSelectNodes();
+				Controller->OpenUndoBracket(TEXT("Collapse to Function"));
+				URigVMCollapseNode* CollapseNode = Controller->CollapseNodes(Nodes, TEXT("New Function"), true, true);
+				if(CollapseNode)
+				{
+					if (Controller->PromoteCollapseNodeToFunctionReferenceNode(CollapseNode->GetFName(), true, true).IsNone())
+					{
+						Controller->CancelUndoBracket();
+					}
+					else
+					{
+						Controller->CloseUndoBracket();
+					}
+				}
+				else
+				{
+					Controller->CancelUndoBracket();
+				}
+			})
+		));
+	}
+
 	if (URigVMCollapseNode* CollapseNode = Cast<URigVMCollapseNode>(ModelNode))
 	{
 		OrganizationSection.AddMenuEntry(
@@ -1196,73 +1227,50 @@ void FRigVMEditorModule::GetNodeOrganizationContextMenuActions(IRigVMClientHost*
 		));
 	}
 
+	if (!Model->AreSectionsEnabled())
+	{
+		OrganizationSection.AddMenuEntry(
+			"Highlight Occurrences",
+			LOCTEXT("HighlightOccurrences", "Highlight Occurrences"),
+			LOCTEXT("HighlightOccurrences_Tooltip", "Highlights all occurrences of the selected section in this graph"),
+			FSlateIcon(),
+			FUIAction(FExecuteAction::CreateLambda([Model]() {
+				Model->UpdateSections(true /* force */);
+			})
+		));
+	}
+
 	OrganizationSection.AddMenuEntry(
-		"Select Similar Nodes",
-		LOCTEXT("SelectSimilarNodes", "Select Similar Nodes"),
-		LOCTEXT("SelectSimilarNodes_Tooltip", "Selects all matching nodes to the selection in this graph"),
+		"Select Occurrences",
+		LOCTEXT("SelectOccurrences", "Select Occurrences"),
+		LOCTEXT("SelectOccurrences_Tooltip", "Selects all occurrences of the selected sections in this graph"),
 		FSlateIcon(),
 		FUIAction(FExecuteAction::CreateLambda([Model, Controller]() {
-			TArray<FName> Nodes = Model->GetSelectNodes();
 
-			TArray<const FRigVMFunction*> RigVMFunctions;
-			TArray<FRigVMGraphFunctionIdentifier> FunctionIdentifiers;
-
-			for(const FName& NodeName : Nodes)
+			const TArray<FName>& SelectedNodes = Model->GetSelectNodes();
+			if (SelectedNodes.IsEmpty())
 			{
-				const URigVMNode* Node = Model->FindNodeByName(NodeName);
-				if(Node == nullptr)
-				{
-					continue;
-				}
-
-				if(const URigVMFunctionReferenceNode* FunctionRefNode = Cast<URigVMFunctionReferenceNode>(Node))
-				{
-					const FRigVMGraphFunctionIdentifier Identifier = FunctionRefNode->GetFunctionIdentifier();
-					if(Identifier.IsValid())
-					{
-						FunctionIdentifiers.AddUnique(Identifier);
-					}
-				}
-				else if(const URigVMTemplateNode* TemplateNode = Cast<URigVMTemplateNode>(Node))
-				{
-					if(const FRigVMFunction* Function = TemplateNode->GetResolvedFunction())
-					{
-						RigVMFunctions.AddUnique(Function);
-					}
-				}
+				return;
 			}
 
-			if(RigVMFunctions.IsEmpty() && FunctionIdentifiers.IsEmpty())
+			const TArray<FRigVMGraphSection> SelectedSections = FRigVMGraphSection::GetSectionsFromSelection(Model);
+			TArray<FRigVMGraphSection> AllSections = SelectedSections;
+
+			TArray<URigVMNode*> AvailableNodes = Model->GetNodes();
+			AvailableNodes.RemoveAll([SelectedNodes]( const URigVMNode* Node) -> bool
 			{
-				return;
+				return SelectedNodes.Contains(Node->GetFName());
+			});
+
+			for (const FRigVMGraphSection& SelectedSection : SelectedSections)
+			{
+				AllSections.Append(SelectedSection.FindMatches(AvailableNodes));
 			}
-			
-			TArray<FName> SimilarNodes;
-			SimilarNodes.Append(Nodes);
 
-			for(const URigVMNode* Node : Model->GetNodes())
+			TArray<FName> SimilarNodes;
+			for (const FRigVMGraphSection& Section : AllSections)
 			{
-				if(const URigVMFunctionReferenceNode* FunctionRefNode = Cast<URigVMFunctionReferenceNode>(Node))
-				{
-					const FRigVMGraphFunctionIdentifier Identifier = FunctionRefNode->GetFunctionIdentifier();
-					if(Identifier.IsValid())
-					{
-						if(FunctionIdentifiers.Contains(Identifier))
-						{
-							SimilarNodes.AddUnique(Node->GetFName());
-						}
-					}
-				}
-				else if(const URigVMTemplateNode* TemplateNode = Cast<URigVMTemplateNode>(Node))
-				{
-					if(const FRigVMFunction* Function = TemplateNode->GetResolvedFunction())
-					{
-						if(RigVMFunctions.Contains(Function))
-						{
-							SimilarNodes.AddUnique(Node->GetFName());
-						}
-					}
-				}
+				SimilarNodes.Append(Section.Nodes);
 			}
 
 			Controller->SetNodeSelection(SimilarNodes, true);
@@ -1324,26 +1332,6 @@ void FRigVMEditorModule::GetNodeVariantContextMenuActions(IRigVMClientHost* RigV
 					Controller->UpgradeNodes(Nodes, true, true);
 				})
 			));
-
-			VariantSection.AddMenuEntry(
-				"Upgrade Nodes All Occurences",
-				LOCTEXT("UpgradeNodesAllOccurences", "Upgrade Nodes (All Occurences)"),
-				LOCTEXT("UpgradeNodesAllOccurences_Tooltip", "Upgrades all occurences of selected outdated nodes to their current implementation"),
-				FSlateIcon(),
-				FUIAction(FExecuteAction::CreateLambda([RigVMClientHost, Model]() {
-					if(FRigVMClient* Client = RigVMClientHost->GetRigVMClient())
-					{
-						TArray<FName> NodeNames = Model->GetSelectNodes();
-						TArray<URigVMNode*> Nodes;
-						Nodes.Reserve(NodeNames.Num());
-						for(const FName& NodeName : NodeNames)
-						{
-							Nodes.Add(Model->FindNodeByName(NodeName));
-						}
-						Client->UpgradeAllOccurencesOfNodes(Nodes, Model->GetOutermost(), true);
-					}
-				})
-			));
 		}
 
 		if(bIsFunctionReference)
diff --git a/Engine/Plugins/Runtime/RigVM/Source/RigVMEditor/Private/Widgets/SRigVMGraphNode.cpp b/Engine/Plugins/Runtime/RigVM/Source/RigVMEditor/Private/Widgets/SRigVMGraphNode.cpp
--- a/Engine/Plugins/Runtime/RigVM/Source/RigVMEditor/Private/Widgets/SRigVMGraphNode.cpp
+++ b/Engine/Plugins/Runtime/RigVM/Source/RigVMEditor/Private/Widgets/SRigVMGraphNode.cpp
@@ -653,14 +653,18 @@ void SRigVMGraphNode::CreateStandardPinWidget(UEdGraphPin* CurPin)
 
 const FSlateBrush * SRigVMGraphNode::GetNodeBodyBrush() const
 {
-	if(URigVMEdGraphNode* RigNode = Cast<URigVMEdGraphNode>(GetNodeObj()))
+	static const FSlateBrush* Brush = FAppStyle::GetBrush("Graph.Node.Body");
+	static const FSlateBrush* TintedBrush = FAppStyle::GetBrush("Graph.Node.TintedBody");
+	
+	if(const URigVMEdGraphNode* RigNode = Cast<URigVMEdGraphNode>(GetNodeObj()))
 	{
-		if(RigNode->bEnableProfiling)
+		if(RigNode->ProfilingColor.IsSet() ||
+			RigNode->HighlightingColor.IsSet())
 		{
-			return FAppStyle::GetBrush("Graph.Node.TintedBody");
+			return TintedBrush;
 		}
 	}
-	return FAppStyle::GetBrush("Graph.Node.Body");
+	return Brush;
 }
 
 FReply SRigVMGraphNode::OnMouseButtonDown(const FGeometry& MyGeometry, const FPointerEvent& MouseEvent)
@@ -1192,6 +1196,11 @@ void SRigVMGraphNode::Tick(const FGeometry& AllottedGeometry, const double InCur
 			PinToDelete->MarkAsGarbage();
 		}
 		PinsToDelete.Reset();
+
+		if (URigVMEdGraphNode* RigVMEdGraphNode = Cast<URigVMEdGraphNode>(GraphNode))
+		{
+			RigVMEdGraphNode->UpdateVisualSettings();
+		}
 	}
 
 	if((!UseLowDetailNodeContent()) && (LeftNodeBox != nullptr))
@@ -1413,7 +1422,7 @@ FText SRigVMGraphNode::GetInstructionDurationText() const
 		{
 			if(const URigVMEdGraphNode* RigGraphNode = Cast<URigVMEdGraphNode>(GraphNode))
 			{
-				const double MicroSeconds = RigGraphNode->MicroSeconds;
+				const double MicroSeconds = RigGraphNode->MicroSeconds.Get(0);
 				if(MicroSeconds >= 0)
 				{
 					if (const SGraphPanel* MyOwnerPanel = GetOwnerPanel().Get())
@@ -1426,7 +1435,7 @@ FText SRigVMGraphNode::GetInstructionDurationText() const
 							{
 								if(const URigVMEdGraphNode* SelectedRigGraphNode = Cast<URigVMEdGraphNode>(SelectedNode))
 								{
-									OverallMicroSeconds += SelectedRigGraphNode->MicroSeconds;
+									OverallMicroSeconds += SelectedRigGraphNode->MicroSeconds.Get(0);
 								}
 							}
 							
diff --git a/Engine/Source/Editor/GraphEditor/Private/SGraphNode.cpp b/Engine/Source/Editor/GraphEditor/Private/SGraphNode.cpp
--- a/Engine/Source/Editor/GraphEditor/Private/SGraphNode.cpp
+++ b/Engine/Source/Editor/GraphEditor/Private/SGraphNode.cpp
@@ -649,9 +649,10 @@ FSlateColor SGraphNode::GetNodeBodyColor() const
 	return ReturnBodyColor;
 }
 
-const FSlateBrush *  SGraphNode::GetNodeBodyBrush() const
+const FSlateBrush* SGraphNode::GetNodeBodyBrush() const
 {
-	return FAppStyle::GetBrush("Graph.Node.Body");
+	static const FSlateBrush* Brush = FAppStyle::GetBrush("Graph.Node.Body");
+	return Brush;
 }
 
 FSlateColor SGraphNode::GetNodeTitleIconColor() const
@@ -1017,7 +1018,7 @@ void SGraphNode::UpdateGraphNode()
 				.Padding(Settings->GetNonPinNodeBodyPadding())
 				[
 					SNew(SImage)
-					.Image(GetNodeBodyBrush())
+					.Image(this, &SGraphNode::GetNodeBodyBrush)
 					.ColorAndOpacity(this, &SGraphNode::GetNodeBodyColor)
 				]
 				+SOverlay::Slot()
