Allow the RigVM registry to unregister types when modules from an never-loaded plugin are unloaded

Modules are statically initialized, and register their UObjects way before the plugin can be loaded or unloaded.
As a result we have live types to unregister while their owning plugin isn't in the picture.

#tests RigVM tests
#rnx

diff --git a/Engine/Plugins/Runtime/RigVM/Source/RigVM/Private/RigVMCore/RigVMRegistry.cpp b/Engine/Plugins/Runtime/RigVM/Source/RigVM/Private/RigVMCore/RigVMRegistry.cpp
--- a/Engine/Plugins/Runtime/RigVM/Source/RigVM/Private/RigVMCore/RigVMRegistry.cpp
+++ b/Engine/Plugins/Runtime/RigVM/Source/RigVM/Private/RigVMCore/RigVMRegistry.cpp
@@ -435,6 +435,7 @@ void FRigVMRegistry_NoLock::InitializeBaseTypes_NoLock()
 
 			IPluginManager::Get().OnNewPluginMounted().RemoveAll(this);
 			IPluginManager::Get().OnPluginUnmounted().RemoveAll(this);
+			FModuleManager::Get().OnModulesUnloaded().RemoveAll(this);
 
 			UE::Anim::AttributeTypes::GetOnAttributeTypesChanged().RemoveAll(this);
 
@@ -868,7 +869,24 @@ bool FRigVMRegistry_NoLock::IsWithinPlugin(const TArray<FString>& InModulePrefix
 bool FRigVMRegistry_NoLock::OnPluginUnloaded_NoLock(IPlugin& InPlugin)
 {
 	const TArray<FString> ModulePrefixes = GetPluginModulePrefixes(InPlugin);
+	return OnModulesUnloadedInternal(ModulePrefixes);
+}
+
+bool FRigVMRegistry_NoLock::OnModulesUnloaded_NoLock(TConstArrayView<FName> ModuleNames)
+{
+	TArray<FString> ModulePrefixes;
+
+	for (FName ModuleName : ModuleNames)
+	{
+		static constexpr TCHAR Format[] = TEXT("/Script/%s.");
+		ModulePrefixes.Add(FString::Printf(Format, *ModuleName.ToString()));
+	}
+
+	return OnModulesUnloadedInternal(ModulePrefixes);
+}
 
+bool FRigVMRegistry_NoLock::OnModulesUnloadedInternal(const TArray<FString>& ModulePrefixes)
+{
 	bool bRegistryChanged = false;
 
 	// remove the dispatch factories defined in the plugin
@@ -1586,6 +1604,20 @@ void FRigVMRegistry_RWLock::OnPluginUnloaded(IPlugin& InPlugin)
 	}
 }
 
+void FRigVMRegistry_RWLock::OnModulesUnloaded(TConstArrayView<FName> ModuleNames)
+{
+	bool bRegistryChanged;
+	{
+		FConditionalWriteScopeLock _(*this);
+		bRegistryChanged = Super::OnModulesUnloaded_NoLock(ModuleNames);
+	}
+		
+	if (bRegistryChanged)
+	{
+		OnRigVMRegistryChangedDelegate.Broadcast();
+	}
+}
+
 void FRigVMRegistry_RWLock::OnAnimationAttributeTypesChanged(const UScriptStruct* InStruct, bool bIsAdded)
 {
 	{
@@ -2870,6 +2902,7 @@ void FRigVMRegistry_RWLock::Initialize_NoLock()
 	{
 		IPluginManager::Get().OnNewPluginMounted().AddRaw(this, &FRigVMRegistry_RWLock::OnPluginLoaded);
 		IPluginManager::Get().OnPluginUnmounted().AddRaw(this, &FRigVMRegistry_RWLock::OnPluginUnloaded);
+		FModuleManager::Get().OnModulesUnloaded().AddRaw(this, &FRigVMRegistry_RWLock::OnModulesUnloaded);
 	}
 	
 	UE::Anim::AttributeTypes::GetOnAttributeTypesChanged().AddRaw(this, &FRigVMRegistry_RWLock::OnAnimationAttributeTypesChanged);
diff --git a/Engine/Plugins/Runtime/RigVM/Source/RigVM/Public/RigVMCore/RigVMRegistry.h b/Engine/Plugins/Runtime/RigVM/Source/RigVM/Public/RigVMCore/RigVMRegistry.h
--- a/Engine/Plugins/Runtime/RigVM/Source/RigVM/Public/RigVMCore/RigVMRegistry.h
+++ b/Engine/Plugins/Runtime/RigVM/Source/RigVM/Public/RigVMCore/RigVMRegistry.h
@@ -145,6 +145,9 @@ public:
 
 	// Removes all types associated with a plugin that's being unloaded. 
 	UE_API virtual bool OnPluginUnloaded_NoLock(IPlugin& InPlugin);
+
+	// Removes all types associated with modules that are being unloaded.
+	UE_API virtual bool OnModulesUnloaded_NoLock(TConstArrayView<FName> ModuleNames);
 	
 	// Update the registry when new types are added to the attribute system so that they can be selected
 	// on Attribute Nodes
@@ -406,6 +409,8 @@ protected:
 		return RF_Public;
 	}
 
+	UE_API bool OnModulesUnloadedInternal(const TArray<FString>& ModulePrefixes);
+
 	UE_API bool IsAllowedType_NoLock(const FProperty* InProperty) const;
 	UE_API bool IsAllowedType_NoLock(const UEnum* InEnum) const;
 	UE_API bool IsAllowedType_NoLock(const UStruct* InStruct) const;
@@ -767,6 +772,9 @@ public:
 	// Removes all types associated with a plugin that's being unloaded. 
 	RIGVM_API void OnPluginUnloaded(IPlugin& InPlugin);
 
+	// Removes all types associated with modules that are being unloaded.
+	RIGVM_API void OnModulesUnloaded(TConstArrayView<FName> ModuleNames);
+
 	// Update the registry when new types are added to the attribute system so that they can be selected
 	// on Attribute Nodes
 	RIGVM_API void OnAnimationAttributeTypesChanged(const UScriptStruct* InStruct, bool bIsAdded);
diff --git a/Engine/Source/Runtime/Core/Private/Modules/ModuleManager.cpp b/Engine/Source/Runtime/Core/Private/Modules/ModuleManager.cpp
--- a/Engine/Source/Runtime/Core/Private/Modules/ModuleManager.cpp
+++ b/Engine/Source/Runtime/Core/Private/Modules/ModuleManager.cpp
@@ -640,7 +640,7 @@ void FModuleManager::InternalFreeLibrary(FName ModuleName, void* Handle)
 
 		// If the module to unload has UObjects in it, then the UObject system need be informed.
 		// UObjects should be removed and garbage collection should run.
-		RemoveLoadedObjectsCallback.Broadcast(UnloadedModules);
+		OnModulesUnloadCallback.Broadcast(UnloadedModules);
 	}
 
 	// Identify if any library is now useless
diff --git a/Engine/Source/Runtime/Core/Public/Modules/ModuleManager.h b/Engine/Source/Runtime/Core/Public/Modules/ModuleManager.h
--- a/Engine/Source/Runtime/Core/Public/Modules/ModuleManager.h
+++ b/Engine/Source/Runtime/Core/Public/Modules/ModuleManager.h
@@ -542,9 +542,9 @@ public:
 	 *
 	 * @return The delegate.
 	 */
-	TMulticastDelegateRegistration<void(TConstArrayView<FName>)>& OnRemoveLoadedObjectsCallback()
+	TMulticastDelegateRegistration<void(TConstArrayView<FName>)>& OnModulesUnloaded()
 	{
-		return RemoveLoadedObjectsCallback;
+		return OnModulesUnloadCallback;
 	}
 
 	/**
@@ -753,7 +753,7 @@ private:
 	TMulticastDelegate<void(FName, bool)> ProcessLoadedObjectsCallback;
 
 	/** Multicast delegate called to remove any loaded objects for a list of modules. */
-	TMulticastDelegate<void(TConstArrayView<FName>)> RemoveLoadedObjectsCallback;
+	TMulticastDelegate<void(TConstArrayView<FName>)> OnModulesUnloadCallback;
 
 	/** Delegate called to check whether any module to delay-unload still has live objects. */
 	TDelegate<bool(TConstArrayView<FName>)> CheckLiveObjectsInModulesCallback;
diff --git a/Engine/Source/Runtime/CoreUObject/Private/UObject/UObjectBase.cpp b/Engine/Source/Runtime/CoreUObject/Private/UObject/UObjectBase.cpp
--- a/Engine/Source/Runtime/CoreUObject/Private/UObject/UObjectBase.cpp
+++ b/Engine/Source/Runtime/CoreUObject/Private/UObject/UObjectBase.cpp
@@ -999,7 +999,7 @@ void RegisterModularObjectsProcessing()
 	UE_CALL_ONCE([]()
 	{
 		FModuleManager::Get().OnProcessLoadedObjectsCallback().AddStatic(ProcessNewlyLoadedUObjects);
-		FModuleManager::Get().OnRemoveLoadedObjectsCallback().AddStatic(RemoveLoadedUObjects);
+		FModuleManager::Get().OnModulesUnloaded().AddStatic(RemoveLoadedUObjects);
 		FModuleManager::Get().OnCheckLiveObjectsInModulesCallback().BindStatic(CheckLiveObjectsInModules);
 
 		FCoreUObjectDelegates::GarbageCollectComplete.AddLambda([]()
