RigVM: Fix CIS breakage due to superflous UE_API statement

#rnx

diff --git a/Engine/Plugins/Runtime/RigVM/Source/RigVMDeveloper/Public/RigVMModel/RigVMGraph.h b/Engine/Plugins/Runtime/RigVM/Source/RigVMDeveloper/Public/RigVMModel/RigVMGraph.h
--- a/Engine/Plugins/Runtime/RigVM/Source/RigVMDeveloper/Public/RigVMModel/RigVMGraph.h
+++ b/Engine/Plugins/Runtime/RigVM/Source/RigVMDeveloper/Public/RigVMModel/RigVMGraph.h
@@ -185,8 +185,8 @@ public:
 	UE_API virtual void PreSave(FObjectPreSaveContext SaveContext) override;
 
    	UE_API bool AreSectionsEnabled() const;
-	UE_API const TArray<FRigVMGraphSection>& GetSelectedSections() const { return SelectedSections; }
-	UE_API const TArray<FRigVMGraphSection>& GetSectionsMatchingTheSelection() const { return SectionsMatchingTheSelection; }
+	const TArray<FRigVMGraphSection>& GetSelectedSections() const { return SelectedSections; }
+	const TArray<FRigVMGraphSection>& GetSectionsMatchingTheSelection() const { return SectionsMatchingTheSelection; }
 	UE_API TArray<FRigVMGraphSection> GetMatchingSections(const FRigVMGraphSection& InSection) const;
 	UE_API void UpdateSections(bool bForce = false);
 
