RigVM: Rely on cloning of localized registries

#rnx

diff --git a/Engine/Plugins/Animation/ControlRig/Source/ControlRig/Private/ControlRigBlueprintGeneratedClass.cpp b/Engine/Plugins/Animation/ControlRig/Source/ControlRig/Private/ControlRigBlueprintGeneratedClass.cpp
--- a/Engine/Plugins/Animation/ControlRig/Source/ControlRig/Private/ControlRigBlueprintGeneratedClass.cpp
+++ b/Engine/Plugins/Animation/ControlRig/Source/ControlRig/Private/ControlRigBlueprintGeneratedClass.cpp
@@ -41,6 +41,7 @@ void UControlRigBlueprintGeneratedClass::Serialize(FArchive& Ar)
 	{
 		if (Ar.IsSaving() && CDO->VM)
 		{
+			VM->CreateLocalizedRegistryIfRequired();
 			VM->CopyDataForSerialization(CDO->VM);
 		}
 	}
diff --git a/Engine/Plugins/Runtime/RigVM/Source/RigVM/Private/RigVMCore/RigVM.cpp b/Engine/Plugins/Runtime/RigVM/Source/RigVM/Private/RigVMCore/RigVM.cpp
--- a/Engine/Plugins/Runtime/RigVM/Source/RigVM/Private/RigVMCore/RigVM.cpp
+++ b/Engine/Plugins/Runtime/RigVM/Source/RigVM/Private/RigVMCore/RigVM.cpp
@@ -365,7 +365,11 @@ void URigVM::CopyDataForSerialization(URigVM* InVM)
 	DefaultWorkMemoryStorage = InVM->DefaultWorkMemoryStorage;
 	DefaultDebugMemoryStorage = InVM->DefaultDebugMemoryStorage;
 
-	LocalizedRegistry = InVM->LocalizedRegistry;
+	LocalizedRegistry.Reset();
+	if (InVM->LocalizedRegistry.IsValid())
+	{
+		LocalizedRegistry = FRigVMRegistry_NoLock::CloneLocalizedRegistry(InVM->LocalizedRegistry.Get());
+	}
 }
 
 void URigVM::PostLoad()
@@ -1140,6 +1144,21 @@ bool URigVM::ResolveFunctionsIfRequired()
 			{
 				// We cannot recover from missing functions.
 				UE_LOG(LogRigVM, Error, TEXT("No handler found for function '%s' for VM '%s'"), *FunctionNameString, *GetPathName());
+
+				// Print more information to help with debugging
+				if (LocalizedRegistry.IsValid())
+				{
+					for (int32 FactoryIndex = 0; FactoryIndex < LocalizedRegistry->GetFactories_NoLock().Num(); FactoryIndex++)
+					{
+						const FRigVMDispatchFactory* AvailableFactory = LocalizedRegistry->GetFactories_NoLock()[FactoryIndex];
+						UE_LOG(LogRigVM, Error, TEXT("Available factory [%03d] '%s' for VM '%s'"), FactoryIndex, *AvailableFactory->GetFactoryName().ToString(), *GetPathName());
+					}
+					for (const FRigVMFunction& AvailableFunction : LocalizedRegistry->GetFunctions_NoLock())
+					{
+						UE_LOG(LogRigVM, Error, TEXT("Available function [%03d] '%s' for VM '%s'"), AvailableFunction.Index, *AvailableFunction.Name, *GetPathName());
+					}
+				}
+				
 				bSuccess = false;
 			}
 		}
@@ -2525,7 +2544,7 @@ bool URigVM::CreateLocalizedRegistryIfRequired(bool bForce)
 	// create a new registry
 	TSharedPtr<FRigVMRegistry_NoLock> TemporaryRegistryPtr = FRigVMRegistry_NoLock::CreateLocalizedRegistry();
 	FRigVMRegistry_NoLock& TemporaryRegistry = *TemporaryRegistryPtr.Get();
-	FRigVMRegistryHandle TemporaryContext(&TemporaryRegistry);
+	FRigVMRegistryHandle TemporaryHandle(&TemporaryRegistry);
 
 	RefreshInstructionsIfRequired();
 	if(ExternalPropertyPaths.IsEmpty())
@@ -2598,7 +2617,7 @@ bool URigVM::CreateLocalizedRegistryIfRequired(bool bForce)
 			}
 
 			// just creating the argument will register the type as needed
-			FRigVMTemplateArgument::Make_NoLock(const_cast<FProperty*>(Property), TemporaryContext);
+			FRigVMTemplateArgument::Make_NoLock(const_cast<FProperty*>(Property), TemporaryHandle);
 		}
 	}
 
@@ -2661,14 +2680,26 @@ bool URigVM::CreateLocalizedRegistryIfRequired(bool bForce)
 		// register the factory
 		FRigVMDispatchFactory* TemporaryFactory = const_cast<FRigVMDispatchFactory*>(TemporaryRegistry.RegisterFactory_NoLock(FactoryStruct, TemporaryArgumentInfos));
 
-		// update all of the argument name caches
-		if(const FRigVMTemplate* Template = TemporaryFactory->GetTemplate_NoLock(TemporaryContext))
+		// update all argument name caches,
+		// force create all function pointers
+		if(FRigVMTemplate* Template = const_cast<FRigVMTemplate*>(TemporaryFactory->GetTemplate_NoLock(TemporaryHandle)))
 		{
-			(void)TemporaryFactory->UpdateArgumentNameCache_NoLock(Template->NumArguments(), TemporaryContext);
+			(void)TemporaryFactory->UpdateArgumentNameCache_NoLock(Template->NumArguments(), TemporaryHandle);
+
+			for (int32 PermutationIndex = 0; PermutationIndex < Template->NumPermutations_NoLock(TemporaryHandle); PermutationIndex++)
+			{
+				if (!Template->GetOrCreatePermutation_NoLock(PermutationIndex, TemporaryHandle))
+				{
+					FRigVMTemplateTypeMap Types = Template->GetTypesForPermutation_NoLock(PermutationIndex, TemporaryHandle);
+					const FString TypesString = Template->GetStringFromArgumentTypes(Types, TemporaryHandle);
+					UE_LOG(LogRigVM, Error, TEXT("Cannot create permutation [%03d] of template '%s' types '%s' for VM '%s'\nCannot localize registry."), PermutationIndex, *Template->GetNotation().ToString(), *TypesString, *GetPathName());
+					return false;
+				}
+			}
 		}
 	}
 
-	// register all struct based functions
+	// register all functions
 	for(const FRigVMFunction* Function : FunctionsStorage)
 	{
 		if(Function == nullptr)
@@ -2678,7 +2709,22 @@ bool URigVM::CreateLocalizedRegistryIfRequired(bool bForce)
 		
 		if(Function->Factory)
 		{
-			continue;
+			// pulling on the function makes sure the permutation exists
+			if (!TemporaryRegistry.FindFunction_NoLock(*Function->Name))
+			{
+				for (int32 FactoryIndex = 0; FactoryIndex < LocalizedRegistry->GetFactories_NoLock().Num(); FactoryIndex++)
+				{
+					const FRigVMDispatchFactory* AvailableFactory = LocalizedRegistry->GetFactories_NoLock()[FactoryIndex];
+					UE_LOG(LogRigVM, Error, TEXT("Available factory [%03d] '%s' for VM '%s'"), FactoryIndex, *AvailableFactory->GetFactoryName().ToString(), *GetPathName());
+				}
+				for (const FRigVMFunction& AvailableFunction : LocalizedRegistry->GetFunctions_NoLock())
+				{
+					UE_LOG(LogRigVM, Error, TEXT("Available function [%03d] '%s' for VM '%s'"), AvailableFunction.Index, *AvailableFunction.Name, *GetPathName());
+				}
+
+				UE_LOG(LogRigVM, Error, TEXT("Cannot find expected permutation '%s' for VM '%s'\nCannot localize registry."), *Function->Name, *GetPathName());
+				return false;
+			}
 		}
 		else
 		{
diff --git a/Engine/Plugins/Runtime/RigVM/Source/RigVM/Private/RigVMCore/RigVMRegistry.cpp b/Engine/Plugins/Runtime/RigVM/Source/RigVM/Private/RigVMCore/RigVMRegistry.cpp
--- a/Engine/Plugins/Runtime/RigVM/Source/RigVM/Private/RigVMCore/RigVMRegistry.cpp
+++ b/Engine/Plugins/Runtime/RigVM/Source/RigVM/Private/RigVMCore/RigVMRegistry.cpp
@@ -39,6 +39,45 @@ TSharedRef<FRigVMRegistry_NoLock> FRigVMRegistry_NoLock::CreateLocalizedRegistry
 	return LocalizedRegistry;
 }
 
+TSharedRef<FRigVMRegistry_NoLock> FRigVMRegistry_NoLock::CloneLocalizedRegistry(FRigVMRegistry_NoLock* InRegistry)
+{
+	TSharedRef<FRigVMRegistry_NoLock> ClonedRegistry = MakeShareable(new FRigVMRegistry_NoLock(false));
+	ClonedRegistry->Types = InRegistry->Types;
+	ClonedRegistry->TypeToIndex = InRegistry->TypeToIndex;
+
+	ClonedRegistry->Factories.Reserve(InRegistry->Factories.Num());
+	ClonedRegistry->FactoryNameToFactory.Reserve(InRegistry->Factories.Num());
+	ClonedRegistry->FactoryStructToFactory.Reserve(InRegistry->Factories.Num());
+	
+	for (const FRigVMDispatchFactory* Factory : InRegistry->Factories)
+	{
+		const FRigVMTemplate* Template = Factory->GetTemplate_NoLock(InRegistry->GetHandle_NoLock());
+		const TArray<FRigVMTemplateArgumentInfo> FlattenedArguments = Template->GetFlattenedArgumentInfos_NoLock(InRegistry->GetHandle_NoLock());
+
+		const FRigVMDispatchFactory* ClonedFactory = ClonedRegistry->RegisterFactory_NoLock(Factory->GetScriptStruct(), FlattenedArguments);
+		const FRigVMTemplate* ClonedTemplate = ClonedFactory->GetTemplate_NoLock(ClonedRegistry->GetHandle_NoLock());
+		const_cast<FRigVMTemplate*>(ClonedTemplate)->Permutations = Template->Permutations;
+	}
+
+	for (FRigVMFunction& Function : InRegistry->Functions)
+	{
+		if(Function.Struct)
+		{
+			ClonedRegistry->Register_NoLock(*Function.Name, Function.FunctionPtr, Function.Struct, Function.Arguments);
+		}
+		else if(Function.Factory)
+		{
+			// finding function will pull on the factory to create it
+			(void)ClonedRegistry->FindFunction_NoLock(*Function.Name);
+		}
+	}
+
+	ClonedRegistry->AllowedClasses = InRegistry->AllowedClasses;
+	ClonedRegistry->AllowedStructs = InRegistry->AllowedStructs;
+
+	return ClonedRegistry;
+}
+
 FRigVMRegistry_NoLock::~FRigVMRegistry_NoLock()
 {
 	FRigVMRegistry_NoLock::Reset_NoLock();
diff --git a/Engine/Plugins/Runtime/RigVM/Source/RigVM/Public/RigVMCore/RigVMRegistry.h b/Engine/Plugins/Runtime/RigVM/Source/RigVM/Public/RigVMCore/RigVMRegistry.h
--- a/Engine/Plugins/Runtime/RigVM/Source/RigVM/Public/RigVMCore/RigVMRegistry.h
+++ b/Engine/Plugins/Runtime/RigVM/Source/RigVM/Public/RigVMCore/RigVMRegistry.h
@@ -62,7 +62,10 @@ public:
 	UE_API FRigVMRegistryHandle& GetHandle_NoLock();
 
 	// Constructor for a localized, independent registry
-	UE_API static TSharedRef<FRigVMRegistry_NoLock> CreateLocalizedRegistry(); 
+	UE_API static TSharedRef<FRigVMRegistry_NoLock> CreateLocalizedRegistry();
+
+	// Constructor for a localized, independent registry based on an existing registry
+	UE_API static TSharedRef<FRigVMRegistry_NoLock> CloneLocalizedRegistry(FRigVMRegistry_NoLock* InRegistry);
 
 	DECLARE_MULTICAST_DELEGATE(FOnRigVMRegistryChanged);
 
