RigVM: Improve CollapseNodes workflow

#rnx

diff --git a/Engine/Plugins/Runtime/RigVM/Source/RigVMDeveloper/Private/RigVMModel/Nodes/RigVMFunctionInterfaceNode.cpp b/Engine/Plugins/Runtime/RigVM/Source/RigVMDeveloper/Private/RigVMModel/Nodes/RigVMFunctionInterfaceNode.cpp
--- a/Engine/Plugins/Runtime/RigVM/Source/RigVMDeveloper/Private/RigVMModel/Nodes/RigVMFunctionInterfaceNode.cpp
+++ b/Engine/Plugins/Runtime/RigVM/Source/RigVMDeveloper/Private/RigVMModel/Nodes/RigVMFunctionInterfaceNode.cpp
@@ -36,11 +36,76 @@ FText URigVMFunctionInterfaceNode::GetToolTipText() const
 
 FText URigVMFunctionInterfaceNode::GetToolTipTextForPin(const URigVMPin* InPin) const
 {
+	FText ToolTip;
 	if(const URigVMPin* OuterPin = FindReferencedPin(InPin))
 	{
-		return OuterPin->GetToolTipText();
+		ToolTip = OuterPin->GetToolTipText();
 	}
-	return Super::GetToolTipTextForPin(InPin);
+	else
+	{
+		ToolTip = Super::GetToolTipTextForPin(InPin);
+	}
+
+	if (!IsInterfacePinUsed(InPin->GetRootPin()->GetFName()))
+	{
+		static const FText Format = NSLOCTEXT("RigVMFunctionInterfaceNode", "PinIsUnused", "{0}\n\nThis pin is not used in this function (grayed out).");
+		ToolTip = FText::Format(Format, ToolTip);
+	}
+	
+	return ToolTip;
+}
+
+bool URigVMFunctionInterfaceNode::IsInterfacePinUsed(const FName& InInterfacePinName) const
+{
+	const URigVMPin* RootPin = FindRootPinByName(InInterfacePinName);
+	if (RootPin == nullptr)
+	{
+		return false;
+	}
+	
+	// if we have links - early exit
+	if (RootPin->IsLinked(false))
+	{
+		return true;
+	}
+
+	// if we potentially have links anywhere - stick with normal color
+	if (RootPin->IsLinked(true))
+	{
+		return true;
+	}
+
+	// only input arguments need to check bound variables
+	// (input arguments are output pins in the contained graph)
+	if (RootPin->GetDirection() != ERigVMPinDirection::Output)
+	{
+		return false;
+	}
+
+	// if there are any pins bound to this one...
+	if (const URigVMGraph* Graph = RootPin->GetGraph())
+	{
+		for (const URigVMNode* Node : Graph->GetNodes())
+		{
+			const URigVMVariableNode* VariableNode = Cast<URigVMVariableNode>(Node);
+			if (!VariableNode)
+			{
+				continue;
+			}
+			if (VariableNode->GetVariableName() != RootPin->GetName())
+			{
+				continue;
+			}
+
+			// we found a variable node that's bound to our input pin
+			if (VariableNode->GetValuePin()->IsLinked(true))
+			{
+				return true;
+			}
+		}
+	}
+
+	return false;
 }
 
 const URigVMPin* URigVMFunctionInterfaceNode::FindReferencedPin(const URigVMPin* InPin) const
diff --git a/Engine/Plugins/Runtime/RigVM/Source/RigVMDeveloper/Private/RigVMModel/RigVMController.cpp b/Engine/Plugins/Runtime/RigVM/Source/RigVMDeveloper/Private/RigVMModel/RigVMController.cpp
--- a/Engine/Plugins/Runtime/RigVM/Source/RigVMDeveloper/Private/RigVMModel/RigVMController.cpp
+++ b/Engine/Plugins/Runtime/RigVM/Source/RigVMDeveloper/Private/RigVMModel/RigVMController.cpp
@@ -5207,9 +5207,9 @@ URigVMCollapseNode* URigVMController::CollapseNodes(const TArray<URigVMNode*>& I
   	FVector2D Diagonal = Bounds.Max - Bounds.Min;
 	FVector2D Center = (Bounds.Min + Bounds.Max) * 0.5f;
 
-	TArray<URigVMPin*> PinsToCollapse;
-	TMap<URigVMPin*, URigVMPin*> CollapsedPins;
-	TArray<URigVMLink*> LinksToRewire;
+	TMap<const URigVMPin*, const URigVMPin*> CollapsedPins;
+	TArray<URigVMLink*> InputLinksToRewire;
+	TArray<URigVMLink*> OutputLinksToRewire;
 	TArray<URigVMLink*> AllLinks = Graph->GetLinks();
 
 	auto NodeToBeCollapsed = [&Nodes](URigVMNode* InNode) -> bool
@@ -5232,6 +5232,9 @@ URigVMCollapseNode* URigVMController::CollapseNodes(const TArray<URigVMNode*>& I
 
 		return false;
 	};
+
+	TArray<const URigVMNode*> BoundaryNodes;
+	
 	// find all pins to collapse. we need this to find out if
 	// we might have a parent pin of a given linked pin already 
 	// collapsed.
@@ -5246,26 +5249,207 @@ URigVMCollapseNode* URigVMController::CollapseNodes(const TArray<URigVMNode*>& I
 			continue;
 		}
 
-		URigVMPin* PinToCollapse = SourcePin;
-		PinsToCollapse.AddUnique(PinToCollapse);
-		LinksToRewire.Add(Link);
+		if (bSourceToBeCollapsed)
+		{
+			OutputLinksToRewire.Add(Link);
+		}
+		else
+		{
+			InputLinksToRewire.Add(Link);
+		}
+
+		BoundaryNodes.AddUnique(SourcePin->GetNode());
+		BoundaryNodes.AddUnique(TargetPin->GetNode());
 	}
 
-	// sort the links so that the links on the same node are in the right order
-	Algo::Sort(LinksToRewire, [&AllLinks](URigVMLink* A, URigVMLink* B) -> bool
+	auto ExecuteLinkCountPredicate = [](const URigVMLink* Link) { return Link->GetSourcePin()->IsExecuteContext(); };
+	const int32 NumInputExecuteLinks = static_cast<int32>(Algo::CountIf(InputLinksToRewire, ExecuteLinkCountPredicate));
+	const int32 NumOutputExecuteLinks = static_cast<int32>(Algo::CountIf(OutputLinksToRewire, ExecuteLinkCountPredicate));
+
+	if (NumInputExecuteLinks > 1 || NumOutputExecuteLinks > 1)
 	{
-		if(A->GetSourcePin()->GetNode() == B->GetSourcePin()->GetNode())
+		ReportAndNotifyWarning(TEXT("Cannot collapse nodes. More than one chain of execute links selected."));
+		return nullptr;
+	}
+
+	// introduce an order for the node based on the screen position
+	TMap<const URigVMNode*, int64> NodeScores;
+	int32 MaxNumPinsPerNode = 0;
+	for (const URigVMNode* BoundaryNode : BoundaryNodes)
+	{
+		NodeScores.Add(BoundaryNode, static_cast<int64>(BoundaryNode->GetPosition().Y));
+		MaxNumPinsPerNode = FMath::Max<int32>(MaxNumPinsPerNode, BoundaryNode->GetAllPinsRecursively().Num());
+	}
+	
+	// introduce an order for the pins on all nodes based on screen position
+	TMap<const URigVMPin*, int64> PinScores;
+	for (const URigVMNode* BoundaryNode : BoundaryNodes)
+	{
+		const TArray<URigVMPin*> AllPins = BoundaryNode->GetAllPinsRecursively();
+		TArray<int64> MaxScorePerDirection;
+		MaxScorePerDirection.AddZeroed(static_cast<int32>(ERigVMPinDirection::Invalid));
+
+		const int64 NodeScore = NodeScores.FindChecked(BoundaryNode) * static_cast<int64>(MaxNumPinsPerNode);
+		
+		for (const URigVMPin* Pin : AllPins)
+		{
+			int64& MaxScore = MaxScorePerDirection[static_cast<int64>(Pin->GetDirection())];
+			PinScores.Add(Pin, NodeScore + (MaxScore++));
+		}
+	}
+
+	// sort both link arrays based on the score of the target pin
+	Algo::Sort(InputLinksToRewire, [&PinScores](URigVMLink* A, URigVMLink* B) -> bool
+	{
+		return PinScores.FindChecked(A->GetTargetPin()) < PinScores.FindChecked(B->GetTargetPin());
+	});
+	Algo::Sort(OutputLinksToRewire, [&PinScores](URigVMLink* A, URigVMLink* B) -> bool
+	{
+		return PinScores.FindChecked(A->GetTargetPin()) < PinScores.FindChecked(B->GetTargetPin());
+	});
+
+	struct FLinkCollapseInfoPin
+	{
+		FLinkCollapseInfoPin()
+		: Pin(nullptr)
+		{
+		}
+
+		FLinkCollapseInfoPin(const FString& InSegmentPath, const URigVMPin* InPin)
+		: SegmentPath(InSegmentPath)
+		, Pin(InPin)
 		{
-			return A->GetSourcePin()->GetAbsolutePinIndex() < B->GetSourcePin()->GetAbsolutePinIndex();
 		}
 		
-		if(A->GetTargetPin()->GetNode() == B->GetTargetPin()->GetNode())
+		FString SegmentPath;
+		const URigVMPin* Pin;
+	};
+	struct FLinkCollapseInfo
+	{
+		FLinkCollapseInfo()
+		: Pin(nullptr)
 		{
-			return A->GetTargetPin()->GetAbsolutePinIndex() < B->GetTargetPin()->GetAbsolutePinIndex();
 		}
 
-		return AllLinks.Find(A) < AllLinks.Find(B);
-	});
+		FLinkCollapseInfo(const URigVMPin* InPin, const URigVMPin* InLinkedPin)
+		: Pin(InPin)
+		, LinkedPins({{FString(), InLinkedPin}})
+		{
+		}
+		
+		const URigVMPin* Pin;
+		TArray<FLinkCollapseInfoPin> LinkedPins;
+	};
+
+	TArray<FLinkCollapseInfo> InputLinkInfos;
+	for (const URigVMLink* Link : InputLinksToRewire)
+	{
+		InputLinkInfos.Emplace(Link->GetSourcePin(), Link->GetTargetPin());
+	}
+
+	TArray<FLinkCollapseInfo> OutputLinkInfos;
+	for (const URigVMLink* Link : OutputLinksToRewire)
+	{
+		OutputLinkInfos.Emplace(Link->GetSourcePin(), Link->GetTargetPin());
+	}
+
+	auto CombineLinkInfos = [](const TArray<FLinkCollapseInfo>& InFlatInfos) -> TArray<FLinkCollapseInfo>
+	{
+		TArray<FLinkCollapseInfo> Result;
+
+		for (const FLinkCollapseInfo& FlatInfo : InFlatInfos)
+		{
+			check(FlatInfo.LinkedPins.Num() == 1);
+			check(FlatInfo.LinkedPins[0].SegmentPath.IsEmpty());
+			
+			const URigVMPin* Pin = FlatInfo.Pin;
+
+			// try to find a mapped pin using this pin exactly
+			FLinkCollapseInfo* MatchingInfo = Result.FindByPredicate([Pin](const FLinkCollapseInfo& Info)
+			{
+				return Info.Pin == Pin;
+			});
+			
+			if (MatchingInfo)
+			{
+				MatchingInfo->LinkedPins.Emplace(FString(), FlatInfo.LinkedPins[0].Pin);
+				continue;
+			}
+
+			// get all pins in this pin tree
+ 			URigVMPin* RootPin = Pin->GetRootPin();
+			TArray<URigVMPin*> ChildPins = {RootPin};
+			ChildPins.Append(RootPin->GetAllSubPinsRecursively());
+			ChildPins.Remove(const_cast<URigVMPin*>(Pin));
+
+			bool bCombinedWithExistingLinkInfo = false;
+			for (const URigVMPin* ChildPin : ChildPins)
+			{
+				MatchingInfo = Result.FindByPredicate([ChildPin](const FLinkCollapseInfo& Info)
+				{
+					return Info.Pin == ChildPin;
+				});
+
+				if (!MatchingInfo)
+				{
+					continue;
+				}
+
+				// if we found a match we are either a sub pin of the pin in question,
+				// or a parent pin. ajust out mapping accordingly
+				if (ChildPin->IsInOuter(Pin))
+				{
+					FString SegmentPath = ChildPin->GetName();
+					ChildPin = ChildPin->GetParentPin();
+					while (ChildPin && ChildPin != Pin)
+					{
+						SegmentPath = URigVMPin::JoinPinPath(ChildPin->GetName(), SegmentPath);
+						ChildPin = ChildPin->GetParentPin();
+					}
+
+					MatchingInfo->Pin = Pin;
+					for (FLinkCollapseInfoPin& PinInfo : MatchingInfo->LinkedPins)
+					{
+						if (PinInfo.SegmentPath.IsEmpty())
+						{
+							PinInfo.SegmentPath = SegmentPath;
+						}
+						else
+						{
+							PinInfo.SegmentPath = URigVMPin::JoinPinPath(SegmentPath, PinInfo.SegmentPath);
+						}
+					}
+
+					MatchingInfo->LinkedPins.Emplace(SegmentPath, FlatInfo.LinkedPins[0].Pin);
+				}
+				else
+				{
+					FString SegmentPath = Pin->GetName();
+					Pin = Pin->GetParentPin();
+					while (Pin && ChildPin != Pin)
+					{
+						SegmentPath = URigVMPin::JoinPinPath(Pin->GetName(), SegmentPath);
+						Pin = Pin->GetParentPin();
+					}
+
+					MatchingInfo->LinkedPins.Emplace(SegmentPath, FlatInfo.LinkedPins[0].Pin);
+				}
+
+				bCombinedWithExistingLinkInfo = true;
+				break;
+			}
+
+			if (!bCombinedWithExistingLinkInfo)
+			{
+				Result.Add(FlatInfo);
+			}
+		}
+
+		return Result;
+	};
+
+	InputLinkInfos = CombineLinkInfos(InputLinkInfos);
+	OutputLinkInfos = CombineLinkInfos(OutputLinkInfos);
 
 	FRigVMControllerCompileBracketScope CompileScope(this);
 	FRigVMCollapseNodesAction CollapseAction;
@@ -5322,20 +5506,21 @@ URigVMCollapseNode* URigVMController::CollapseNodes(const TArray<URigVMNode*>& I
 		return nullptr;
 	}
 
-	// now looper over the links to be rewired
-	for (URigVMLink* Link : LinksToRewire)
+	auto CreateExposedPins = [this, &CollapsedPins, CollapseNode](const FLinkCollapseInfo& CollapseInfo, bool bIsInput)
 	{
-		bool bSourceToBeCollapsed = NodeToBeCollapsed(Link->GetSourcePin()->GetNode());
-
-		URigVMPin* PinToCollapse = bSourceToBeCollapsed ? Link->GetSourcePin() : Link->GetTargetPin();
+		check(!CollapseInfo.LinkedPins.IsEmpty());
+		
+		const URigVMPin* PinToCollapse = CollapseInfo.Pin;
+		const URigVMPin* PinForConfiguration = bIsInput ? CollapseInfo.LinkedPins[0].Pin : CollapseInfo.Pin;
+		
 		if (CollapsedPins.Contains(PinToCollapse))
 		{
-			continue;
+			return;
 		}
 
 		if(PinToCollapse->IsExecuteContext() && PinToCollapse->GetDirection() == ERigVMPinDirection::IO)
 		{
-			for(const TPair<URigVMPin*, URigVMPin*>& Pair : CollapsedPins)
+			for(const TPair<const URigVMPin*, const URigVMPin*>& Pair : CollapsedPins)
 			{
 				if(Pair.Key->IsExecuteContext() && Pair.Key->GetDirection() == ERigVMPinDirection::IO)
 				{
@@ -5345,44 +5530,23 @@ URigVMCollapseNode* URigVMController::CollapseNodes(const TArray<URigVMNode*>& I
 			}
 			if (CollapsedPins.Contains(PinToCollapse))
 			{
-				continue;
-			}
-		}
-
-		// for links that connect to the right side of the collapse
-		// node, we need to skip sub pins of already exposed pins
-		if (bSourceToBeCollapsed)
-		{
-			bool bParentPinCollapsed = false;
-			URigVMPin* ParentPin = PinToCollapse->GetParentPin();
-			while (ParentPin != nullptr)
-			{
-				if (PinsToCollapse.Contains(ParentPin))
-				{
-					bParentPinCollapsed = true;
-					break;
-				}
-				ParentPin = ParentPin->GetParentPin();
-			}
-
-			if (bParentPinCollapsed)
-			{
-				continue;
+				return;
 			}
 		}
 
-		FName PinName = URigVMSchema::GetUniqueName(PinToCollapse->GetFName(), [CollapseNode](const FName& InName) {
+		const FName DesiredPinName = PinForConfiguration->GetFName();
+		const FName PinName = URigVMSchema::GetUniqueName(DesiredPinName, [CollapseNode](const FName& InName) {
 			return CollapseNode->FindPin(InName.ToString()) == nullptr;
 		}, false, true);
 
 		URigVMPin* CollapsedPin = NewObject<URigVMPin>(CollapseNode, PinName);
-		ConfigurePinFromPin(CollapsedPin, PinToCollapse, true);
+		ConfigurePinFromPin(CollapsedPin, PinForConfiguration, false);
 
 		if (CollapsedPin->GetDirection() == ERigVMPinDirection::IO)
 		{
 			if(!CollapsedPin->IsExecuteContext())
 			{
-				CollapsedPin->Direction = bSourceToBeCollapsed ? ERigVMPinDirection::Output : ERigVMPinDirection::Input;
+				CollapsedPin->Direction = bIsInput ? ERigVMPinDirection::Input : ERigVMPinDirection::Output;
 			}
 		}
 
@@ -5394,10 +5558,19 @@ URigVMCollapseNode* URigVMController::CollapseNodes(const TArray<URigVMNode*>& I
 
 		AddNodePin(CollapseNode, CollapsedPin);
 
-		FPinState PinState = GetPinState(PinToCollapse);
+		FPinState PinState = GetPinState(PinForConfiguration);
 		ApplyPinState(CollapsedPin, PinState);
 
 		CollapsedPins.Add(PinToCollapse, CollapsedPin);
+	};
+
+	for (const FLinkCollapseInfo& InputLinkInfo : InputLinkInfos)
+	{
+		CreateExposedPins(InputLinkInfo, true);
+	}
+	for (const FLinkCollapseInfo& OutputLinkInfo : OutputLinkInfos)
+	{
+		CreateExposedPins(OutputLinkInfo, false);
 	}
 
 	Notify(ERigVMGraphNotifType::NodeAdded, CollapseNode);
@@ -5435,125 +5608,122 @@ URigVMCollapseNode* URigVMController::CollapseNodes(const TArray<URigVMNode*>& I
 		}
 	}
 
-	// create the new nodes within the collapse node
-	TArray<FName> ContainedNodeNames;
-	if(URigVMController* CollapseController = GetControllerForGraph(CollapseNode->GetContainedGraph()))
+	auto RewireLinksForInfo = [CollapsedPins, EntryNode, ReturnNode](URigVMController* InController, const FLinkCollapseInfo& CollapseInfo, bool bIsInput, bool bIsInterior)
 	{
-		FString TextContent = ExportNodesToText(NodeNames);
+		const URigVMPin* OriginalPin = CollapseInfo.Pin;
+		OriginalPin = CollapsedPins.FindChecked(OriginalPin);
 
-		ContainedNodeNames = CollapseController->ImportNodesFromText(TextContent, false);
-
-		// move the nodes to the right place
-		for (const FName& ContainedNodeName : ContainedNodeNames)
+		if (bIsInterior)
 		{
-			if (URigVMNode* ContainedNode = CollapseNode->GetContainedGraph()->FindNodeByName(ContainedNodeName))
+			if (bIsInput)
 			{
-				if(!ContainedNode->IsInjected())
+				OriginalPin = EntryNode->FindPin(OriginalPin->GetSegmentPath(true));
+			}
+			else
+			{
+				OriginalPin = ReturnNode->FindPin(OriginalPin->GetSegmentPath(true));
+				if (OriginalPin == nullptr)
 				{
-					CollapseController->SetNodePosition(ContainedNode, ContainedNode->Position - Center, false, false);
+					return;
 				}
-			}
-		}
 
-		for (URigVMLink* LinkToRewire : LinksToRewire)
-		{
-			URigVMPin* SourcePin = LinkToRewire->GetSourcePin();
-			URigVMPin* TargetPin = LinkToRewire->GetTargetPin();
-
-			if (NodeToBeCollapsed(SourcePin->GetNode()))
-			{
-				// if the parent pin of this was collapsed
-				// it's possible that the child pin wasn't.
-				if (!CollapsedPins.Contains(SourcePin))
+				URigVMPin* SourcePin = InController->GetGraph()->FindPin(CollapseInfo.Pin->GetPinPath());
+				if (SourcePin == nullptr)
 				{
-					continue;
+					return;
 				}
 
-				URigVMPin* CollapsedPin = CollapsedPins.FindChecked(SourcePin);
-				SourcePin = CollapseNode->ContainedGraph->FindPin(SourcePin->GetPinPath());
-				TargetPin = ReturnNode->FindPin(CollapsedPin->GetName());
-			}
-			else
-			{
-				URigVMPin* CollapsedPin = CollapsedPins.FindChecked(TargetPin);
-				SourcePin = EntryNode->FindPin(CollapsedPin->GetName());
-				TargetPin = CollapseNode->ContainedGraph->FindPin(TargetPin->GetPinPath());
-			}
+				URigVMPin* TargetPin = const_cast<URigVMPin*>(OriginalPin);
 
-			if (SourcePin && TargetPin)
-			{
 				if (!SourcePin->IsLinkedTo(TargetPin))
 				{
-					CollapseController->AddLink(SourcePin, TargetPin, false);
+					InController->AddLink(SourcePin, TargetPin, false);
 				}
+				return;
 			}
 		}
-	}
-
-	TArray<URigVMLink*> RewiredLinks;
-	for (URigVMLink* LinkToRewire : LinksToRewire)
-	{
-		if (RewiredLinks.Contains(LinkToRewire))
+		else if (bIsInput)
 		{
-			continue;
+			URigVMPin* SourcePin = const_cast<URigVMPin*>(CollapseInfo.Pin);
+			URigVMPin* TargetPin = const_cast<URigVMPin*>(OriginalPin);
+			if (!SourcePin->IsLinkedTo(TargetPin))
+			{
+				InController->AddLink(SourcePin, TargetPin, false);
+			}
+			return;
 		}
 
-		URigVMPin* SourcePin = LinkToRewire->GetSourcePin();
-		URigVMPin* TargetPin = LinkToRewire->GetTargetPin();
-
-		if (NodeToBeCollapsed(SourcePin->GetNode()))
+		if (OriginalPin == nullptr)
 		{
-			FString SegmentPath;
-			URigVMPin* PinToCheck = SourcePin;
+			return;
+		}
 
-			URigVMPin** CollapsedPinPtr = CollapsedPins.Find(PinToCheck);
-			while (CollapsedPinPtr == nullptr)
+		for (const FLinkCollapseInfoPin& LinkedPinInfo : CollapseInfo.LinkedPins)
+		{
+			const URigVMPin* Pin = OriginalPin;
+			if (!LinkedPinInfo.SegmentPath.IsEmpty())
 			{
-				if (SegmentPath.IsEmpty())
+				Pin = Pin->FindSubPin(LinkedPinInfo.SegmentPath);
+				if (!Pin)
 				{
-					SegmentPath = PinToCheck->GetName();
+					continue;
 				}
-				else
+			}
+				
+			const URigVMPin* LinkedPin = LinkedPinInfo.Pin;
+			if (bIsInterior)
+			{
+				LinkedPin = InController->GetGraph()->FindPin(LinkedPin->GetPinPath());
+				if (LinkedPin == nullptr)
 				{
-					SegmentPath = URigVMPin::JoinPinPath(PinToCheck->GetName(), SegmentPath);
+					continue;
 				}
-
-				PinToCheck = PinToCheck->GetParentPin();
-				check(PinToCheck);
-
-				CollapsedPinPtr = CollapsedPins.Find(PinToCheck);
 			}
 
-			URigVMPin* CollapsedPin = *CollapsedPinPtr;
-			check(CollapsedPin);
-
-			if (!SegmentPath.IsEmpty())
+			URigVMPin* SourcePin = const_cast<URigVMPin*>(Pin);
+			URigVMPin* TargetPin = const_cast<URigVMPin*>(LinkedPin);
+			if (!SourcePin->IsLinkedTo(TargetPin))
 			{
-				CollapsedPin = CollapsedPin->FindSubPin(SegmentPath);
-				check(CollapsedPin);
+				InController->AddLink(SourcePin, TargetPin, false);
 			}
+		}
+	};
+
+	// create the new nodes within the collapse node
+	if(URigVMController* CollapseController = GetControllerForGraph(CollapseNode->GetContainedGraph()))
+	{
+		const FString TextContent = ExportNodesToText(NodeNames);
+		const TArray<FName> ContainedNodeNames = CollapseController->ImportNodesFromText(TextContent, false);
 
-			TArray<URigVMLink*> TargetLinks = SourcePin->GetTargetLinks(false);
-			for (URigVMLink* TargetLink : TargetLinks)
+		// move the nodes to the right place
+		for (const FName& ContainedNodeName : ContainedNodeNames)
+		{
+			if (URigVMNode* ContainedNode = CollapseNode->GetContainedGraph()->FindNodeByName(ContainedNodeName))
 			{
-				TargetPin = TargetLink->GetTargetPin();
-				if (!CollapsedPin->IsLinkedTo(TargetPin))
+				if(!ContainedNode->IsInjected())
 				{
-					AddLink(CollapsedPin, TargetPin, false);
+					CollapseController->SetNodePosition(ContainedNode, ContainedNode->Position - Center, false, false);
 				}
 			}
-			RewiredLinks.Append(TargetLinks);
 		}
-		else
+
+		for (const FLinkCollapseInfo& InputLinkInfo : InputLinkInfos)
 		{
-			URigVMPin* CollapsedPin = CollapsedPins.FindChecked(TargetPin);
-			if (!SourcePin->IsLinkedTo(CollapsedPin))
-			{
-				AddLink(SourcePin, CollapsedPin, false);
-			}
+			RewireLinksForInfo(CollapseController, InputLinkInfo, true, true);
 		}
+		for (const FLinkCollapseInfo& OutputLinkInfo : OutputLinkInfos)
+		{
+			RewireLinksForInfo(CollapseController, OutputLinkInfo, false, true);
+		}
+	}
 
-		RewiredLinks.Add(LinkToRewire);
+	for (const FLinkCollapseInfo& InputLinkInfo : InputLinkInfos)
+	{
+		RewireLinksForInfo(this, InputLinkInfo, true, false);
+	}
+	for (const FLinkCollapseInfo& OutputLinkInfo : OutputLinkInfos)
+	{
+		RewireLinksForInfo(this, OutputLinkInfo, false, false);
 	}
 
 	if (ReturnNode)
@@ -18071,7 +18241,7 @@ void URigVMController::ConfigurePinFromProperty(FProperty* InProperty, URigVMPin
 	}
 }
 
-void URigVMController::ConfigurePinFromPin(URigVMPin* InOutPin, URigVMPin* InPin, bool bCopyDisplayName)
+void URigVMController::ConfigurePinFromPin(URigVMPin* InOutPin, const URigVMPin* InPin, bool bCopyDisplayName)
 {
 	// it is important we copy things that define the identity of the pin
 	// things that defines the state of the pin is copied during GetPinState()
@@ -19558,7 +19728,7 @@ TMap<FString, FString> URigVMController::GetRedirectedPinPaths(URigVMNode* InNod
 	return RedirectedPinPaths;
 }
 
-URigVMController::FPinState URigVMController::GetPinState(URigVMPin* InPin, bool bStoreWeakInjectionInfos) const
+URigVMController::FPinState URigVMController::GetPinState(const URigVMPin* InPin, bool bStoreWeakInjectionInfos) const
 {
 	FPinState State;
 	State.Direction = InPin->GetDirection();
diff --git a/Engine/Plugins/Runtime/RigVM/Source/RigVMDeveloper/Public/RigVMModel/Nodes/RigVMFunctionInterfaceNode.h b/Engine/Plugins/Runtime/RigVM/Source/RigVMDeveloper/Public/RigVMModel/Nodes/RigVMFunctionInterfaceNode.h
--- a/Engine/Plugins/Runtime/RigVM/Source/RigVMDeveloper/Public/RigVMModel/Nodes/RigVMFunctionInterfaceNode.h
+++ b/Engine/Plugins/Runtime/RigVM/Source/RigVMDeveloper/Public/RigVMModel/Nodes/RigVMFunctionInterfaceNode.h
@@ -27,7 +27,10 @@ public:
 	
 	// URigVMNode interface
 	UE_API virtual  FText GetToolTipText() const override;
-	UE_API virtual FText GetToolTipTextForPin(const URigVMPin* InPin) const override;	
+	UE_API virtual FText GetToolTipTextForPin(const URigVMPin* InPin) const override;
+
+	// returns true if the interface pin is used for anything within the function / collapse graph
+	UE_API virtual bool IsInterfacePinUsed(const FName& InInterfacePinName) const;
 
 protected:
 
diff --git a/Engine/Plugins/Runtime/RigVM/Source/RigVMDeveloper/Public/RigVMModel/RigVMController.h b/Engine/Plugins/Runtime/RigVM/Source/RigVMDeveloper/Public/RigVMModel/RigVMController.h
--- a/Engine/Plugins/Runtime/RigVM/Source/RigVMDeveloper/Public/RigVMModel/RigVMController.h
+++ b/Engine/Plugins/Runtime/RigVM/Source/RigVMDeveloper/Public/RigVMModel/RigVMController.h
@@ -1319,7 +1319,7 @@ private:
 	UE_API void AddPinsForArray(FArrayProperty* InArrayProperty, URigVMNode* InNode, URigVMPin* InParentPin, ERigVMPinDirection InPinDirection, const TArray<FString>& InDefaultValues, bool bAutoExpandArrays);
 	UE_API void AddPinsForTemplate(const FRigVMTemplate* InTemplate, const FRigVMTemplateTypeMap& InPinTypeMap, URigVMNode* InNode);
 	UE_API void ConfigurePinFromProperty(FProperty* InProperty, URigVMPin* InOutPin, ERigVMPinDirection InPinDirection = ERigVMPinDirection::Invalid) const;
-	UE_API void ConfigurePinFromPin(URigVMPin* InOutPin, URigVMPin* InPin, bool bCopyDisplayName = false);
+	UE_API void ConfigurePinFromPin(URigVMPin* InOutPin, const URigVMPin* InPin, bool bCopyDisplayName = false);
 	UE_API void ConfigurePinFromArgument(URigVMPin* InOutPin, const FRigVMGraphFunctionArgument& InArgument, bool bCopyDisplayName = false);
 	UE_API bool ResetPinDefaultValue(URigVMPin* InPin, bool bSetupUndoRedo);
 	static UE_API FString GetPinInitialDefaultValue(const URigVMPin* InPin);
@@ -1404,7 +1404,7 @@ public:
 	};
 
 	UE_API TMap<FString, FString> GetRedirectedPinPaths(URigVMNode* InNode) const;
-	UE_API FPinState GetPinState(URigVMPin* InPin, bool bStoreWeakInjectionInfos = false) const;
+	UE_API FPinState GetPinState(const URigVMPin* InPin, bool bStoreWeakInjectionInfos = false) const;
 	UE_API TMap<FString, FPinState> GetPinStates(URigVMNode* InNode, bool bStoreWeakInjectionInfos = false) const;
 	UE_API void ApplyPinState(URigVMPin* InPin, const FPinState& InPinState, bool bSetupUndoRedo = false);
 	UE_API void ApplyPinStates(URigVMNode* InNode, const TMap<FString, FPinState>& InPinStates, const TMap<FString, FString>& InRedirectedPinPaths = TMap<FString, FString>(), bool bSetupUndoRedo = false);
diff --git a/Engine/Plugins/Runtime/RigVM/Source/RigVMEditor/Private/Editor/RigVMGraphDetailCustomization.cpp b/Engine/Plugins/Runtime/RigVM/Source/RigVMEditor/Private/Editor/RigVMGraphDetailCustomization.cpp
--- a/Engine/Plugins/Runtime/RigVM/Source/RigVMEditor/Private/Editor/RigVMGraphDetailCustomization.cpp
+++ b/Engine/Plugins/Runtime/RigVM/Source/RigVMEditor/Private/Editor/RigVMGraphDetailCustomization.cpp
@@ -54,6 +54,169 @@ static_assert(ESlateTransformComponent::Scale == 2);
 
 static const FText RigVMGraphDetailCustomizationMultipleValues = LOCTEXT("MultipleValues", "Multiple Values");
 
+FRigVMFunctionArgumentReorderDragDropOp::FRigVMFunctionArgumentReorderDragDropOp(const TWeakObjectPtr<URigVMPin>& InPinPtr)
+: PinPtr(InPinPtr)
+{
+}
+
+void FRigVMFunctionArgumentReorderDragDropOp::Init()
+{
+	SetValidTarget(false);
+	SetupDefaults();
+	Construct();
+}
+
+void FRigVMFunctionArgumentReorderDragDropOp::SetValidTarget(bool IsValidTarget)
+{
+	if (IsValidTarget)
+	{
+		CurrentHoverText = LOCTEXT("PlaceArgumentHere", "Place Argument Here");
+		CurrentIconBrush = FAppStyle::GetBrush("Graph.ConnectorFeedback.OK");
+	}
+	else
+	{
+		CurrentHoverText = LOCTEXT("CannotPlaceArgumentHere", "Cannot Place Argument Here");
+		CurrentIconBrush = FAppStyle::GetBrush("Graph.ConnectorFeedback.Error");
+	}
+}
+
+const URigVMPin* FRigVMFunctionArgumentReorderDragDropOp::GetPin() const
+{
+	if (PinPtr.IsValid())
+	{
+		return PinPtr.Get();
+	}
+	return nullptr;
+}
+
+TSharedPtr<FDragDropOperation> FRigVMFunctionArgumentReorderDragDropHandler::CreateDragDropOperation() const
+{
+	TSharedPtr<FRigVMFunctionArgumentReorderDragDropOp> DragOp = MakeShared<FRigVMFunctionArgumentReorderDragDropOp>(PinPtr);
+	DragOp->Init();
+	return DragOp;
+}
+
+bool FRigVMFunctionArgumentReorderDragDropHandler::AcceptDrop(const FDragDropEvent& DragDropEvent, EItemDropZone DropZone) const
+{
+	if (!WeakRigVMClientHost.IsValid())
+	{
+		return false;
+	}
+
+	IRigVMClientHost* RigVMClientHost = WeakRigVMClientHost.Get();
+	if (!RigVMClientHost)
+	{
+		return false;
+	}
+
+	FRigVMClient* Client = RigVMClientHost->GetRigVMClient();
+	if (!Client)
+	{
+		return false;
+	}
+
+	const TSharedPtr<FRigVMFunctionArgumentReorderDragDropOp> DragOp = DragDropEvent.GetOperationAs<FRigVMFunctionArgumentReorderDragDropOp>();
+	check(DragOp.IsValid());
+	
+	const URigVMPin* DraggedPin = DragOp->GetPin();
+	const URigVMPin* ThisPin = GetPin();
+	check(DraggedPin && ThisPin && DraggedPin != ThisPin);
+	check(DraggedPin->GetDirection() == ThisPin->GetDirection());
+
+	URigVMLibraryNode* LibraryNode = Cast<URigVMLibraryNode>(ThisPin->GetNode());
+	if (!LibraryNode)
+	{
+		return false;
+	}
+
+	URigVMGraph* ContainedGraph = LibraryNode->GetContainedGraph();
+	if (!ContainedGraph)
+	{
+		return false;
+	}
+
+	URigVMController* Controller = Client->GetOrCreateController(ContainedGraph);
+	if (!Controller)
+	{
+		return false;
+	}
+
+	int32 NewPinIndex = ThisPin->GetPinIndex();
+	if (DropZone == EItemDropZone::BelowItem)
+	{
+		if (NewPinIndex < DraggedPin->GetPinIndex())
+		{
+			NewPinIndex++;
+		}
+	}
+	else if (DropZone == EItemDropZone::AboveItem)
+	{
+		if (NewPinIndex > DraggedPin->GetPinIndex())
+		{
+			NewPinIndex--;
+		}
+	}
+
+	Controller->SetExposedPinIndex(DraggedPin->GetFName(), NewPinIndex);
+	return true;
+}
+
+TOptional<EItemDropZone> FRigVMFunctionArgumentReorderDragDropHandler::CanAcceptDrop(const FDragDropEvent& DragDropEvent,
+	EItemDropZone DropZone) const
+{
+	const TSharedPtr<FRigVMFunctionArgumentReorderDragDropOp> DragOp = DragDropEvent.GetOperationAs<FRigVMFunctionArgumentReorderDragDropOp>();
+	if (!DragOp.IsValid())
+	{
+		return TOptional<EItemDropZone>();
+	}
+
+	DragOp->SetValidTarget(false);
+	
+	const URigVMPin* DraggedPin = DragOp->GetPin();
+	const URigVMPin* ThisPin = GetPin();
+	if(DraggedPin == nullptr || ThisPin == nullptr || DraggedPin == ThisPin)
+	{
+		return TOptional<EItemDropZone>();
+	}
+
+	if (DraggedPin->GetDirection() != ThisPin->GetDirection())
+	{
+		return TOptional<EItemDropZone>();
+	}
+
+	if (DropZone == EItemDropZone::OntoItem)
+	{
+		DropZone = EItemDropZone::BelowItem;
+	}
+
+	if (DropZone == EItemDropZone::BelowItem)
+	{
+		if (DraggedPin->GetPinIndex() == ThisPin->GetPinIndex() + 1)
+		{
+			return TOptional<EItemDropZone>();
+		}
+	}
+	else if (DropZone == EItemDropZone::AboveItem)
+	{
+		if (DraggedPin->GetPinIndex() == ThisPin->GetPinIndex() - 1)
+		{
+			return TOptional<EItemDropZone>();
+		}
+	}
+
+	DragOp->SetValidTarget(true);
+	return DropZone;
+}
+
+const URigVMPin* FRigVMFunctionArgumentReorderDragDropHandler::GetPin() const
+{
+	if (PinPtr.IsValid())
+	{
+		return PinPtr.Get();
+	}
+	return nullptr;
+}
+
 FRigVMFunctionArgumentGroupLayout::FRigVMFunctionArgumentGroupLayout(
 	const TWeakObjectPtr<URigVMGraph>& InGraph,
 	const TWeakInterfacePtr<IRigVMClientHost>& InRigVMClientHost,
@@ -236,13 +399,18 @@ void FRigVMFunctionArgumentLayout::GenerateHeaderRowContent(FDetailWidgetRow& No
 	{
 		CustomPinTypeFilters.Add(MakeShared<FRigVMFunctionArgumentPinTypeSelectorFilter>(RigVMEditorPtr, GraphPtr));
 	}
+
+	TAttribute<int32> ArgumentIndex = TAttribute<int32>::CreateSP(this, &FRigVMFunctionArgumentLayout::OnGetArgumentIndex); 
 	
 	NodeRow
 		.NameContent()
+		.HAlign(HAlign_Fill)
 		[
 			SNew(SHorizontalBox)
+			
 			+ SHorizontalBox::Slot()
 			.FillWidth(1)
+			.HAlign(HAlign_Fill)
 			.VAlign(VAlign_Center)
 			[
 				SAssignNew(ArgumentNameWidget, SEditableTextBox)
@@ -292,45 +460,15 @@ void FRigVMFunctionArgumentLayout::GenerateHeaderRowContent(FDetailWidgetRow& No
 				.Font(IDetailLayoutBuilder::GetDetailFont())
 			]
 			+ SHorizontalBox::Slot()
-			.AutoWidth()
-			[
-				SNew(SButton)
-				.ButtonStyle(FAppStyle::Get(), TEXT("SimpleButton"))
-				.ContentPadding(0.f)
-				.IsEnabled_Raw(this, &FRigVMFunctionArgumentLayout::CanArgumentBeMoved, true)
-				.OnClicked(this, &FRigVMFunctionArgumentLayout::OnArgMoveUp)
-				.ToolTipText(LOCTEXT("FunctionArgDetailsArgMoveUpTooltip", "Move this parameter up in the list."))
-				[
-					SNew(SImage)
-					.Image(FAppStyle::GetBrush("Icons.ChevronUp"))
-				.ColorAndOpacity(FSlateColor::UseForeground())
-				]
-			]
-			+ SHorizontalBox::Slot()
-			.AutoWidth()
-			.Padding(2, 0)
-			[
-				SNew(SButton)
-				.ButtonStyle(FAppStyle::Get(), TEXT("SimpleButton"))
-				.ContentPadding(0.f)
-				.IsEnabled_Raw(this, &FRigVMFunctionArgumentLayout::CanArgumentBeMoved, false)
-				.OnClicked(this, &FRigVMFunctionArgumentLayout::OnArgMoveDown)
-				.ToolTipText(LOCTEXT("FunctionArgDetailsArgMoveDownTooltip", "Move this parameter down in the list."))
-				[
-					SNew(SImage)
-					.Image(FAppStyle::GetBrush("Icons.ChevronDown"))
-					.ColorAndOpacity(FSlateColor::UseForeground())
-				]
-			]
-			+ SHorizontalBox::Slot()
 			.HAlign(HAlign_Right)
 			.VAlign(VAlign_Center)
 			.Padding(10, 0, 0, 0)
 			.AutoWidth()
 			[
-				PropertyCustomizationHelpers::MakeClearButton(FSimpleDelegate::CreateSP(this, &FRigVMFunctionArgumentLayout::OnRemoveClicked), LOCTEXT("FunctionArgDetailsClearTooltip", "Remove this parameter."), !IsPinEditingReadOnly())
+				PropertyCustomizationHelpers::MakeDeleteButton(FSimpleDelegate::CreateSP(this, &FRigVMFunctionArgumentLayout::OnRemoveClicked), LOCTEXT("FunctionArgDetailsClearTooltip", "Remove this parameter."), !IsPinEditingReadOnly())
 			]
-		];
+		]
+		.DragDropHandler(MakeShared<FRigVMFunctionArgumentReorderDragDropHandler>(PinPtr, GraphPtr, WeakRigVMClientHost));
 }
 
 void FRigVMFunctionArgumentLayout::GenerateChildContent(IDetailChildrenBuilder& ChildrenBuilder)
@@ -354,90 +492,16 @@ void FRigVMFunctionArgumentLayout::OnRemoveClicked()
 	}
 }
 
-FReply FRigVMFunctionArgumentLayout::OnArgMoveUp()
-{
-	if (PinPtr.IsValid() && WeakRigVMClientHost.IsValid())
-	{
-		URigVMPin* Pin = PinPtr.Get();
-		IRigVMClientHost* RigVMClientHost = WeakRigVMClientHost.Get();
-		if (URigVMLibraryNode* LibraryNode = Cast<URigVMLibraryNode>(Pin->GetNode()))
-		{
-			if (URigVMController* Controller = RigVMClientHost->GetController(LibraryNode->GetContainedGraph()))
-			{
-				bool bIsInput = Pin->GetDirection() == ERigVMPinDirection::Input || Pin->GetDirection() == ERigVMPinDirection::IO;
-				
-				int32 NewPinIndex = Pin->GetPinIndex() - 1;
-				while (NewPinIndex != INDEX_NONE)
-				{
-					URigVMPin* OtherPin = LibraryNode->GetPins()[NewPinIndex];
-					if (bIsInput)
-					{
-						if (OtherPin->GetDirection() == ERigVMPinDirection::Input || OtherPin->GetDirection() == ERigVMPinDirection::IO)
-						{
-							break;
-						}
-					}	
-					else
-					{
-						if (OtherPin->GetDirection() == ERigVMPinDirection::Output || OtherPin->GetDirection() == ERigVMPinDirection::IO)
-						{
-							break;
-						}
-					}
-					--NewPinIndex;
-				}
-				if (NewPinIndex != INDEX_NONE)
-				{
-					Controller->SetExposedPinIndex(Pin->GetFName(), NewPinIndex, true, true);
-				}
-				return FReply::Handled();
-			}
-		}
-	}
-	return FReply::Unhandled();
-}
-
-FReply FRigVMFunctionArgumentLayout::OnArgMoveDown()
+int32 FRigVMFunctionArgumentLayout::OnGetArgumentIndex() const
 {
-	if (PinPtr.IsValid() && WeakRigVMClientHost.IsValid())
+	if (PinPtr.IsValid())
 	{
-		URigVMPin* Pin = PinPtr.Get();
-		IRigVMClientHost* RigVMClientHost = WeakRigVMClientHost.Get();
-		if (URigVMLibraryNode* LibraryNode = Cast<URigVMLibraryNode>(Pin->GetNode()))
+		if(const URigVMPin* Pin = PinPtr.Get())
 		{
-			if (URigVMController* Controller = RigVMClientHost->GetController(LibraryNode->GetContainedGraph()))
-			{
-				bool bIsInput = Pin->GetDirection() == ERigVMPinDirection::Input || Pin->GetDirection() == ERigVMPinDirection::IO;
-				
-				int32 NewPinIndex = Pin->GetPinIndex() + 1;
-				while (NewPinIndex < LibraryNode->GetPins().Num())
-				{
-					URigVMPin* OtherPin = LibraryNode->GetPins()[NewPinIndex];
-					if (bIsInput)
-					{
-						if (OtherPin->GetDirection() == ERigVMPinDirection::Input || OtherPin->GetDirection() == ERigVMPinDirection::IO)
-						{
-							break;
-						}
-					}	
-					else
-					{
-						if (OtherPin->GetDirection() == ERigVMPinDirection::Output || OtherPin->GetDirection() == ERigVMPinDirection::IO)
-						{
-							break;
-						}
-					}
-					++NewPinIndex;
-				}
-				if (NewPinIndex < LibraryNode->GetPins().Num())
-				{
-					Controller->SetExposedPinIndex(Pin->GetFName(), NewPinIndex, true, true);
-				}
-				return FReply::Handled();
-			}
+			return Pin->GetPinIndex();
 		}
 	}
-	return FReply::Unhandled();
+	return INDEX_NONE;
 }
 
 bool FRigVMFunctionArgumentLayout::ShouldPinBeReadOnly(bool bIsEditingPinType/* = false*/) const
@@ -466,63 +530,6 @@ bool FRigVMFunctionArgumentLayout::IsPinEditingReadOnly(bool bIsEditingPinType/*
 	return false;
 }
 
-bool FRigVMFunctionArgumentLayout::CanArgumentBeMoved(bool bMoveUp) const
-{
-	if(IsPinEditingReadOnly())
-	{
-		return false;
-	}
-	if (PinPtr.IsValid())
-	{
-		URigVMPin* Pin = PinPtr.Get();
-		if(Pin->IsExecuteContext())
-		{
-			return false;
-		}
-
-		if(URigVMNode* Node = Pin->GetNode())
-		{
-			auto IsInput = [](URigVMPin* InPin) -> bool
-			{
-				return InPin->GetDirection() == ERigVMPinDirection::Input ||
-					InPin->GetDirection() == ERigVMPinDirection::Visible;
-			};
-
-			const bool bLookForInput = IsInput(Pin);
-
-			if(bMoveUp)
-			{
-				// if this is the first pin of its type
-				for(int32 Index = 0; Index < Node->GetPins().Num(); Index++)
-				{
-					URigVMPin* OtherPin = Node->GetPins()[Index];
-					if(OtherPin->IsExecuteContext())
-					{
-						continue;
-					}
-					if(IsInput(OtherPin) == bLookForInput)
-					{
-						return OtherPin != Pin;
-					}
-				}
-			}
-			else if(!bMoveUp)
-			{
-				// if this is the last pin of its type
-				for(int32 Index = Node->GetPins().Num() - 1; Index >= 0; Index--)
-				{
-					URigVMPin* OtherPin = Node->GetPins()[Index];
-					if(IsInput(OtherPin) == bLookForInput)
-					{
-						return OtherPin != Pin;
-					}
-				}
-			}
-		}
-	}
-	return true;
-}
-
 FText FRigVMFunctionArgumentLayout::OnGetArgNameText() const
 {
 	if (PinPtr.IsValid())
@@ -539,7 +546,7 @@ FText FRigVMFunctionArgumentLayout::OnGetArgToolTipText() const
 
 void FRigVMFunctionArgumentLayout::OnArgNameTextCommitted(const FText& NewText, ETextCommit::Type InTextCommit)
 {
-	if (InTextCommit == ETextCommit::OnEnter)
+	if (InTextCommit == ETextCommit::OnEnter || InTextCommit == ETextCommit::OnUserMovedFocus)
 	{
 		if (!NewText.IsEmpty() && PinPtr.IsValid() && WeakRigVMClientHost.IsValid() && !ShouldPinBeReadOnly())
 		{
diff --git a/Engine/Plugins/Runtime/RigVM/Source/RigVMEditor/Private/RigVMEditorModule.cpp b/Engine/Plugins/Runtime/RigVM/Source/RigVMEditor/Private/RigVMEditorModule.cpp
--- a/Engine/Plugins/Runtime/RigVM/Source/RigVMEditor/Private/RigVMEditorModule.cpp
+++ b/Engine/Plugins/Runtime/RigVM/Source/RigVMEditor/Private/RigVMEditorModule.cpp
@@ -1170,6 +1170,44 @@ void FRigVMEditorModule::GetNodeOrganizationContextMenuActions(IRigVMClientHost*
 		));
 	}
 
+	if (URigVMFunctionInterfaceNode* FunctionInterfaceNode = Cast<URigVMFunctionInterfaceNode>(ModelNode))
+	{
+		TArray<FName> PinsToRemove;
+		for (const URigVMPin* Pin : FunctionInterfaceNode->GetPins())
+		{
+			if (!FunctionInterfaceNode->IsInterfacePinUsed(Pin->GetFName()))
+			{
+				if (FunctionInterfaceNode->IsA<URigVMFunctionEntryNode>() && Pin->GetDirection() == ERigVMPinDirection::Output)
+				{
+					PinsToRemove.Add(Pin->GetFName());
+				}
+				else if (FunctionInterfaceNode->IsA<URigVMFunctionReturnNode>() && Pin->GetDirection() == ERigVMPinDirection::Input)
+				{
+					PinsToRemove.Add(Pin->GetFName());
+				}
+			}
+		}
+
+		OrganizationSection.AddMenuEntry(
+			"RemoveUnusedExposedPins",
+			LOCTEXT("RemoveUnusedExposedPins", "Remove unused pins"),
+				LOCTEXT("RemoveUnusedExposedPins_Tooltip", "Removes all unused exposed pins from this node"),
+			FSlateIcon(),
+			FUIAction(FExecuteAction::CreateLambda([Controller, PinsToRemove]() {
+				Controller->OpenUndoBracket(TEXT("Remove Unused Exposed Pins"));
+				for (const FName& ExposedPinName : PinsToRemove)
+				{
+					Controller->RemoveExposedPin(ExposedPinName);
+				}
+				Controller->CloseUndoBracket();
+			}),
+			FCanExecuteAction::CreateLambda([PinsToRemove]()
+			{
+				return !PinsToRemove.IsEmpty();
+			})
+		));
+	}
+
 	if(RigVMBlueprint)
 	{
 		if (URigVMFunctionReferenceNode* FunctionReferenceNode = Cast<URigVMFunctionReferenceNode>(ModelNode))
@@ -1480,6 +1518,7 @@ void FRigVMEditorModule::GetNodeDisplayContextMenuActions(IRigVMClientHost* RigV
 
 void FRigVMEditorModule::GetPinContextMenuActions(IRigVMClientHost* RigVMClientHost, const UEdGraphPin* EdGraphPin, URigVMPin* ModelPin, UToolMenu* Menu) const
 {
+	GetExposedPinContextMenuActions(RigVMClientHost, EdGraphPin, ModelPin, Menu);
 	GetPinWorkflowContextMenuActions(RigVMClientHost, EdGraphPin, ModelPin, Menu);
 	GetPinDebugContextMenuActions(RigVMClientHost, EdGraphPin, ModelPin, Menu);
 	GetPinArrayContextMenuActions(RigVMClientHost, EdGraphPin, ModelPin, Menu);
@@ -1491,6 +1530,32 @@ void FRigVMEditorModule::GetPinContextMenuActions(IRigVMClientHost* RigVMClientH
 	GetPinInjectedNodesContextMenuActions(RigVMClientHost, EdGraphPin, ModelPin, Menu);
 }
 
+void FRigVMEditorModule::GetExposedPinContextMenuActions(IRigVMClientHost* RigVMClientHost, const UEdGraphPin* EdGraphPin, URigVMPin* ModelPin, UToolMenu* Menu) const
+{
+	if (!ModelPin->IsRootPin() || !ModelPin->GetNode()->IsA<URigVMFunctionInterfaceNode>())
+	{
+		return;
+	}
+
+	URigVMController* Controller = RigVMClientHost->GetRigVMClient()->GetController(ModelPin->GetGraph());
+	if (Controller == nullptr)
+	{
+		return;
+	}
+
+	FToolMenuSection& Section = Menu->AddSection("RigVMEditorContextMenuExposedPin", LOCTEXT("ExposedPinHeader", "Exposed Pin"));
+
+	Section.AddMenuEntry(
+		"RemoveExposedPin",
+		LOCTEXT("RemoveExposedPin", "Remove Pin"),
+		LOCTEXT("RemoveExposedPin_Tooltip", "Removes the exposed pin from the node"),
+		FSlateIcon(),
+		FUIAction(FExecuteAction::CreateLambda([Controller, ModelPin]() {
+			Controller->RemoveExposedPin(ModelPin->GetFName());
+		})
+	));
+}
+
 void FRigVMEditorModule::GetPinWorkflowContextMenuActions(IRigVMClientHost* RigVMClientHost, const UEdGraphPin* EdGraphPin, URigVMPin* ModelPin, UToolMenu* Menu) const
 {
 	const TArray<FRigVMUserWorkflow> Workflows = ModelPin->GetNode()->GetSupportedWorkflows(ERigVMUserWorkflowType::PinContext, ModelPin);
diff --git a/Engine/Plugins/Runtime/RigVM/Source/RigVMEditor/Private/Widgets/SRigVMGraphNode.cpp b/Engine/Plugins/Runtime/RigVM/Source/RigVMEditor/Private/Widgets/SRigVMGraphNode.cpp
--- a/Engine/Plugins/Runtime/RigVM/Source/RigVMEditor/Private/Widgets/SRigVMGraphNode.cpp
+++ b/Engine/Plugins/Runtime/RigVM/Source/RigVMEditor/Private/Widgets/SRigVMGraphNode.cpp
@@ -54,6 +54,242 @@
 const FSlateBrush* SRigVMGraphNode::CachedImg_CR_Pin_Connected = nullptr;
 const FSlateBrush* SRigVMGraphNode::CachedImg_CR_Pin_Disconnected = nullptr;
 
+class FRigVMGraphPinReorderDragDropOperation : public FDragDropOperation
+{
+public:
+	DRAG_DROP_OPERATOR_TYPE(FRigVMGraphPinReorderDragDropOperation, FDragDropOperation)
+
+	int32 PinIndex = INDEX_NONE;
+	TWeakObjectPtr<URigVMNode> ModelNode;
+
+	virtual TSharedPtr<SWidget> GetDefaultDecorator() const override
+	{
+		return SNullWidget::NullWidget;
+	}
+};
+
+void SRigVMGraphPinReorderHandle::Construct(const FArguments& InArgs)
+{
+	bIsInput = InArgs._bIsInput;
+	IndexAttribute = InArgs._Index;
+	LineExtension = InArgs._LineExtension;
+	ModelNode = InArgs._ModelNode;
+	CanReorderDelegate = InArgs._CanReorder;
+	ReorderDelegate = InArgs._OnReorder;
+	WeakClientHost = InArgs._WeakClientHost;
+
+	SetCursor(EMouseCursor::GrabHand);
+
+	ChildSlot
+	[
+		SNew(SBox)
+		.Padding(0.0f)
+		.HAlign(bIsInput ? HAlign_Left : HAlign_Right)
+		.VAlign(VAlign_Center)
+		.WidthOverride(InArgs._Width)
+		.HeightOverride(InArgs._Height)
+		[
+			SNew(SImage)
+			.Image(FCoreStyle::Get().GetBrush("VerticalBoxDragIndicatorShort"))
+		]
+	];
+}
+
+FReply SRigVMGraphPinReorderHandle::OnMouseButtonDown(const FGeometry& MyGeometry, const FPointerEvent& MouseEvent)
+{
+	return FReply::Handled().DetectDrag(SharedThis(this), EKeys::LeftMouseButton);
+}
+
+FReply SRigVMGraphPinReorderHandle::OnDragDetected(const FGeometry& MyGeometry, const FPointerEvent& MouseEvent)
+{
+	if (IsEnabled() && MouseEvent.IsMouseButtonDown(EKeys::LeftMouseButton))
+	{
+		const int32 PinIndex = IndexAttribute.Get();
+		if (ModelNode.IsValid())
+		{
+			if (ModelNode->GetPins()[PinIndex]->IsExecuteContext())
+			{
+				return FReply::Unhandled();
+			}
+		}
+		
+		TSharedPtr<FRigVMGraphPinReorderDragDropOperation> DragDropOp = MakeShared<FRigVMGraphPinReorderDragDropOperation>();
+		if (DragDropOp.IsValid())
+		{
+			DragDropOp->PinIndex = PinIndex;
+			DragDropOp->ModelNode = ModelNode;
+			return FReply::Handled().BeginDragDrop(DragDropOp.ToSharedRef());
+		}
+	}
+
+	return FReply::Unhandled();
+}
+
+FReply SRigVMGraphPinReorderHandle::OnDragOver(const FGeometry& MyGeometry, const FDragDropEvent& DragDropEvent)
+{
+	UE::Slate::FDeprecateVector2DResult LocalPosition = MyGeometry.AbsoluteToLocal(DragDropEvent.GetScreenSpacePosition());
+	const bool bIsAboveItem = LocalPosition.Y < MyGeometry.Size.Y * 0.5f;
+	const EItemDropZone DesiredDropZone = bIsAboveItem ? EItemDropZone::AboveItem : EItemDropZone::BelowItem;
+
+	if (TSharedPtr<FRigVMGraphPinReorderDragDropOperation> ReorderDragDropOp = DragDropEvent.GetOperationAs<FRigVMGraphPinReorderDragDropOperation>())
+	{
+		// only allow drag & drop of pins on the same node
+		if (ModelNode.IsValid() && ReorderDragDropOp->ModelNode.IsValid() && ModelNode.Get() == ReorderDragDropOp->ModelNode.Get())
+		{
+			DropZone = DesiredDropZone;
+
+			if (CanReorderDelegate.IsBound())
+			{
+				if (CanReorderDelegate.Execute(ReorderDragDropOp->PinIndex, IndexAttribute.Get(), DropZone.GetValue()) == INDEX_NONE)
+				{
+					DropZone.Reset();
+				}
+			}
+
+			
+			if (DropZone.IsSet())
+			{
+				return FReply::Handled();
+			}
+		}
+	}
+	else if (TSharedPtr<FGraphEditorDragDropAction> PinDragDropOp = DragDropEvent.GetOperationAs<FGraphEditorDragDropAction>())
+	{
+		if (PinDragDropOp->IsOfTypeImpl(TEXT("FDragConnection")))
+		{
+			DropZone = DesiredDropZone;
+		}
+	}
+	return SCompoundWidget::OnDragOver(MyGeometry, DragDropEvent);
+}
+
+void SRigVMGraphPinReorderHandle::OnDragLeave(const FDragDropEvent& DragDropEvent)
+{
+	SCompoundWidget::OnDragLeave(DragDropEvent);
+	DropZone.Reset();
+}
+
+FReply SRigVMGraphPinReorderHandle::OnDrop(const FGeometry& MyGeometry, const FDragDropEvent& DragDropEvent)
+{
+	if (DropZone.IsSet())
+	{
+		if (TSharedPtr<FRigVMGraphPinReorderDragDropOperation> ReorderDragDropOp = DragDropEvent.GetOperationAs<FRigVMGraphPinReorderDragDropOperation>())
+		{
+			// only allow drag & drop of pins on the same node
+			if (ModelNode.IsValid() && ReorderDragDropOp->ModelNode.IsValid() && ModelNode.Get() == ReorderDragDropOp->ModelNode.Get())
+			{
+				if (ReorderDelegate.IsBound())
+				{
+					ReorderDelegate.Execute(ReorderDragDropOp->PinIndex, IndexAttribute.Get(), DropZone.GetValue());
+				}
+			}
+		}
+		else if (TSharedPtr<FGraphEditorDragDropAction> PinDragDropOp = DragDropEvent.GetOperationAs<FGraphEditorDragDropAction>())
+		{
+			if (PinDragDropOp->IsOfTypeImpl(TEXT("FDragConnection")) && ModelNode.IsValid() && WeakClientHost.IsValid())
+			{
+				int32 NewIndex = IndexAttribute.Get();
+				if (DropZone.GetValue() == EItemDropZone::BelowItem)
+				{
+					NewIndex++;
+				}
+
+				URigVMController* Controller = WeakClientHost.Get()->GetOrCreateController(ModelNode.Get()->GetGraph());
+				if (Controller)
+				{
+					FName ExposedPinName = NAME_None;
+
+					// attach a listener here to determine the exposed pin being added.
+					// the code path to access the pins dragged etc is closed down within private
+					// code in GraphEditor - so relying on listening to changes coming from RigVM here instead.
+					const FDelegateHandle Handle = Controller->OnModified().AddLambda(
+						[Controller, &ExposedPinName](ERigVMGraphNotifType InNotification, URigVMGraph* InGraph, UObject* InSubject)
+						{
+							if (InGraph != Controller->GetGraph())
+							{
+								return;
+							}
+
+							if (InNotification == ERigVMGraphNotifType::PinAdded)
+							{
+								if (const URigVMPin* Pin = Cast<URigVMPin>(InSubject))
+								{
+									if (Pin->IsRootPin())
+									{
+										ExposedPinName = Pin->GetFName();
+									}
+								}
+							}
+						}
+					);
+					
+					FReply Reply = PinDragDropOp->DroppedOnNode(DragDropEvent.GetScreenSpacePosition(), FVector2f::ZeroVector);
+
+					// remove the handle immediately to clean up the listener
+					Controller->OnModified().Remove(Handle);
+					
+					if (Reply.IsEventHandled())
+					{
+						// now since we have a new exposed pin let's move it to the right spot
+						Controller->SetExposedPinIndex(ExposedPinName, NewIndex);
+						return Reply;
+					}
+				}
+			}
+		}
+
+		DropZone.Reset();
+		return FReply::Handled();
+	}
+	return SCompoundWidget::OnDrop(MyGeometry, DragDropEvent);
+}
+
+int32 SRigVMGraphPinReorderHandle::OnPaint(const FPaintArgs& Args, const FGeometry& AllottedGeometry, const FSlateRect& MyCullingRect,
+	FSlateWindowElementList& OutDrawElements, int32 LayerId, const FWidgetStyle& InWidgetStyle, bool bParentEnabled) const
+{
+	int32 LastLayerId = LayerId;
+
+	if (IsHovered())
+	{
+		LastLayerId = SCompoundWidget::OnPaint(Args, AllottedGeometry, MyCullingRect, OutDrawElements, LayerId, InWidgetStyle, bParentEnabled);
+	}
+
+	if (DropZone.IsSet())
+	{
+		TArray<FVector2D> ClampLine;
+		ClampLine.Emplace(0, DropZone.GetValue() == EItemDropZone::AboveItem ? 0 : AllottedGeometry.Size.Y - 1.f);
+		ClampLine.Emplace(AllottedGeometry.Size.X, ClampLine.Last().Y);
+
+		if (LineExtension.IsSet() || LineExtension.IsBound())
+		{
+			const float Extension = LineExtension.Get();
+			if (bIsInput)
+			{
+				ClampLine[0].X -= Extension;
+				ClampLine[1].X -= Extension;
+			}
+			else
+			{
+				ClampLine[0].X += Extension;
+				ClampLine[1].X += Extension;
+			}
+		}
+
+		FSlateDrawElement::MakeLines(
+			OutDrawElements,
+			++LayerId,
+			AllottedGeometry.ToPaintGeometry(),
+			ClampLine,
+			ESlateDrawEffect::None,
+			FSlateColor(EStyleColor::AccentBlue).GetSpecifiedColor(),
+			false,
+			3
+		);
+	}
+	
+	return LastLayerId;
+}
+
 void SRigVMGraphNode::Construct( const FArguments& InArgs )
 {
 	if (CachedImg_CR_Pin_Connected == nullptr)
@@ -273,6 +509,47 @@ void SRigVMGraphNode::UpdateCompactNode()
 	CreatePinWidgets();
 }
 
+TOptional<FSlateColor> SRigVMGraphNode::GetPinTextColor(const SGraphPin* InGraphPin) const
+{
+	TOptional<FSlateColor> Result = SGraphNode::GetPinTextColor(InGraphPin); 
+
+	const UEdGraphPin* PinObj = InGraphPin->GetPinObj();
+	if (PinObj == nullptr)
+	{
+		return Result;
+	}
+	
+	if (!ModelNode.IsValid())
+	{
+		return Result;
+	}
+
+	const URigVMFunctionInterfaceNode* FunctionInterfaceNode = Cast<URigVMFunctionInterfaceNode>(ModelNode);
+	if (!FunctionInterfaceNode)
+	{
+		return Result;
+	}
+
+	FString NodeName, PinPath;
+	if (!URigVMPin::SplitPinPathAtStart(PinObj->GetName(), NodeName, PinPath))
+	{
+		return Result;
+	}
+
+	if (const URigVMPin* ModelPin = ModelNode->FindPin(PinPath))
+	{
+		if (!ModelPin->IsLinked() &&
+			!FunctionInterfaceNode->IsInterfacePinUsed(ModelPin->GetRootPin()->GetFName()))
+		{
+			// we are other an output argument without any links to it
+			// or an input argument without links and without being bound within the graph
+			return FSlateColor::UseSubduedForeground();
+		}
+	}
+
+	return Result;
+}
+
 void SRigVMGraphNode::CreateAddPinButton()
 {
 	if (LeftNodeBox.IsValid() && ModelNode.IsValid())
@@ -2426,12 +2703,20 @@ void SRigVMGraphNode::UpdatePinTreeView()
 			}
 		}
 	}
-	
+
 	for(const FPinRowInfo& Row : Rows)
 	{
 		const FPinInfo& InputPinInfo = PinInfos[Row.InputPinInfo == INDEX_NONE ? Row.OutputPinInfo : Row.InputPinInfo];
 		const FPinInfo& OutputPinInfo = PinInfos[Row.OutputPinInfo == INDEX_NONE ? Row.InputPinInfo : Row.OutputPinInfo];
 
+		const bool bIsRootPin =
+			InputPinInfo.InputPinWidget.IsValid() ?
+			InputPinInfo.ParentIndex == INDEX_NONE :
+			OutputPinInfo.ParentIndex == INDEX_NONE;
+		const bool bUsesReorderWidget = bIsRootPin && ModelNode->IsA<URigVMFunctionInterfaceNode>();
+
+		TSharedPtr<SHorizontalBox> SlotLayout;
+		
 		if(InputPinInfo.InputPinWidget.IsValid())
 		{
 			if(InputPinInfo.bHideInputWidget)
@@ -2445,10 +2730,11 @@ void SRigVMGraphNode::UpdatePinTreeView()
 			// input pins
 			if(!OutputPinInfo.OutputPinWidget.IsValid())
 			{
-				TSharedPtr<SHorizontalBox> SlotLayout;
 				SHorizontalBox::FSlot* FirstSlot = nullptr;
 
-				const float MyEmptySidePadding = InputPinInfo.bHideInputWidget ? EmptySidePadding : 0.f; 
+				const float MyEmptySidePadding =
+					(InputPinInfo.bHideInputWidget && !bUsesReorderWidget)
+					? EmptySidePadding : 0.f; 
 				
 				LeftNodeBox->AddSlot()
                 .HAlign(HAlign_Fill)
@@ -2483,7 +2769,6 @@ void SRigVMGraphNode::UpdatePinTreeView()
 			// io pins
 			else
 			{
-				TSharedPtr<SHorizontalBox> SlotLayout;
 				SHorizontalBox::FSlot* FirstSlot = nullptr;
 
 				OutputPinInfo.OutputPinWidget->SetShowLabel(false);
@@ -2532,25 +2817,50 @@ void SRigVMGraphNode::UpdatePinTreeView()
 		// output pins
 		else if(OutputPinInfo.OutputPinWidget.IsValid())
 		{
+			const float MyEmptySidePadding = bUsesReorderWidget ? 0.f : EmptySidePadding;
+			
 			LeftNodeBox->AddSlot()
             .HAlign(HAlign_Fill)
 			.VAlign(VAlign_Center)
             .AutoHeight()
 			.MaxHeight(OutputPinInfo.bAutoHeight ? TAttribute<float>() : MaxHeight)
             [
-            	SNew(SHorizontalBox)
+				SAssignNew(SlotLayout, SHorizontalBox)
 	            .Visibility(this, &SRigVMGraphNode::GetPinVisibility, OutputPinInfo.Index, false)
 	            
 				+SHorizontalBox::Slot()
 				.FillWidth(1.f)
 				.HAlign(HAlign_Right)
 				.VAlign(VAlign_Center)
-	            .Padding(EmptySidePadding, TopPadding, PinWidgetSidePadding, 0.f)
+	            .Padding(MyEmptySidePadding, TopPadding, PinWidgetSidePadding, 0.f)
 				[
 	                OutputPinInfo.OutputPinWidget.ToSharedRef()
 				]
             ];
 		}
+
+		if (bUsesReorderWidget)
+		{
+			const bool bIsInput = ModelNode->IsA<URigVMFunctionReturnNode>();
+			IRigVMClientHost* Host = RigGraphNode->GetImplementingOuter<IRigVMClientHost>();
+			const URigVMPin* ModelPin = ModelNode->GetGraph()->FindPin(bIsInput ? InputPinInfo.Identifier : OutputPinInfo.Identifier);
+
+			SlotLayout->InsertSlot(bIsInput ? SlotLayout->NumSlots() : 0)
+			.AutoWidth()
+			.VAlign(VAlign_Center)
+			.Padding(0)
+			[
+				SNew(SRigVMGraphPinReorderHandle)
+				.bIsInput(bIsInput)
+				.Width(EmptySidePadding)
+				.Index(ModelPin->GetPinIndex())
+				.LineExtension(44.f)
+				.ModelNode(ModelNode)
+				.CanReorder(this, &SRigVMGraphNode::CanReorderExposedPin)
+				.OnReorder(this, &SRigVMGraphNode::OnReorderExposedPin)
+				.WeakClientHost(Host)
+			];
+		}
 	}
 
 	if(URigVMFunctionReferenceNode* FunctionReferenceNode = Cast<URigVMFunctionReferenceNode>(ModelNode))
@@ -2653,4 +2963,71 @@ EOverrideWidgetStatus::Type SRigVMGraphNode::GetOverrideStatus() const
 	return EOverrideWidgetStatus::ChangedHere;
 }
 
+int32 SRigVMGraphNode::CanReorderExposedPin(int32 InPinIndexToReorder, int32 InTargetPinIndex, EItemDropZone InDropZone)
+{
+	if (InPinIndexToReorder == InTargetPinIndex)
+	{
+		return INDEX_NONE;
+	}
+	
+	if (InDropZone == EItemDropZone::OntoItem)
+	{
+		InDropZone = EItemDropZone::BelowItem;
+	}
+
+	int32 NewPinIndex = InTargetPinIndex;
+	if (InDropZone == EItemDropZone::BelowItem)
+	{
+		if (NewPinIndex < InPinIndexToReorder)
+		{
+			NewPinIndex++;
+		}
+	}
+	else if (InDropZone == EItemDropZone::AboveItem)
+	{
+		if (NewPinIndex > InPinIndexToReorder)
+		{
+			NewPinIndex--;
+		}
+	}
+
+	if (InPinIndexToReorder == NewPinIndex)
+	{
+		return INDEX_NONE;
+	}
+
+	return NewPinIndex;
+}
+
+void SRigVMGraphNode::OnReorderExposedPin(int32 InPinIndexToReorder, int32 InTargetPinIndex, EItemDropZone InDropZone)
+{
+	const int32 NewPinIndex = CanReorderExposedPin(InPinIndexToReorder, InTargetPinIndex, InDropZone);
+	if (NewPinIndex == INDEX_NONE)
+	{
+		return;
+	}
+
+	const URigVMNode* InterfaceNode = ModelNode.Get();
+	if (!InterfaceNode)
+	{
+		return;
+	}
+
+	URigVMEdGraphNode* RigVMEdGraphNode = Cast<URigVMEdGraphNode>(GraphNode);
+	if (RigVMEdGraphNode == nullptr)
+	{
+		return;
+	}
+
+	URigVMController* Controller = RigVMEdGraphNode->GetController();
+	if (Controller == nullptr)
+	{
+		return;
+	}
+
+	const URigVMPin* Pin = InterfaceNode->GetPins()[InPinIndexToReorder];
+	check(Pin);
+	(void)Controller->SetExposedPinIndex(Pin->GetFName(), NewPinIndex);
+}
+
 #undef LOCTEXT_NAMESPACE
diff --git a/Engine/Plugins/Runtime/RigVM/Source/RigVMEditor/Public/Editor/RigVMGraphDetailCustomization.h b/Engine/Plugins/Runtime/RigVM/Source/RigVMEditor/Public/Editor/RigVMGraphDetailCustomization.h
--- a/Engine/Plugins/Runtime/RigVM/Source/RigVMEditor/Public/Editor/RigVMGraphDetailCustomization.h
+++ b/Engine/Plugins/Runtime/RigVM/Source/RigVMEditor/Public/Editor/RigVMGraphDetailCustomization.h
@@ -20,12 +20,62 @@
 #include "Widgets/SRigVMGraphPinNameListValueWidget.h"
 #include "Widgets/SRigVMLogWidget.h"
 #include "HAL/PlatformApplicationMisc.h"
+#include "IDetailDragDropHandler.h"
 
 #define UE_API RIGVMEDITOR_API
 
 class IDetailLayoutBuilder;
 class FRigVMGraphDetailCustomizationImpl;
 
+class FRigVMFunctionArgumentReorderDragDropOp : public FDecoratedDragDropOp
+{
+public:
+	DRAG_DROP_OPERATOR_TYPE(FRigVMFunctionArgumentReorderDragDropOp, FDecoratedDragDropOp)
+
+	FRigVMFunctionArgumentReorderDragDropOp(const TWeakObjectPtr<URigVMPin>& InPinPtr);
+
+	/** Inits the tooltip, needs to be called after constructing */
+	void Init();
+
+	/** Update the drag tool tip indicating whether the current drop target is valid */
+	void SetValidTarget(bool IsValidTarget);
+
+	const URigVMPin* GetPin() const;
+
+private:
+	TWeakObjectPtr<URigVMPin> PinPtr;
+};
+
+class FRigVMFunctionArgumentReorderDragDropHandler : public IDetailDragDropHandler
+{
+public:
+	FRigVMFunctionArgumentReorderDragDropHandler(
+		const TWeakObjectPtr<URigVMPin>& InPinPtr,
+		const TWeakObjectPtr<URigVMGraph>& InGraphPtr,
+		const TWeakInterfacePtr<IRigVMClientHost>& InWeakRigVMClientHost
+	)
+		: PinPtr(InPinPtr)
+		, GraphPtr(InGraphPtr)
+		, WeakRigVMClientHost(InWeakRigVMClientHost)
+	{
+	}
+
+	virtual ~FRigVMFunctionArgumentReorderDragDropHandler() override
+	{
+	}
+
+	virtual TSharedPtr<FDragDropOperation> CreateDragDropOperation() const override;
+	virtual bool AcceptDrop(const FDragDropEvent& DragDropEvent, EItemDropZone DropZone) const override;
+	virtual TOptional<EItemDropZone> CanAcceptDrop(const FDragDropEvent& DragDropEvent, EItemDropZone DropZone) const override;
+
+	const URigVMPin* GetPin() const;
+
+private:
+	TWeakObjectPtr<URigVMPin> PinPtr;
+	TWeakObjectPtr<URigVMGraph> GraphPtr;
+	TWeakInterfacePtr<IRigVMClientHost> WeakRigVMClientHost;
+};
+
 class FRigVMFunctionArgumentGroupLayout : public IDetailCustomNodeBuilder, public TSharedFromThis<FRigVMFunctionArgumentGroupLayout>
 {
 public:
@@ -91,14 +141,10 @@ private:
 	/** Determines if editing the pins on the node should be read only */
 	UE_API bool IsPinEditingReadOnly(bool bIsEditingPinType = false) const;
 
-	/** Determines if an argument can be moved up or down */
-	UE_API bool CanArgumentBeMoved(bool bMoveUp) const;
-
 	/** Callbacks for all the functionality for modifying arguments */
 	UE_API void OnRemoveClicked();
-	UE_API FReply OnArgMoveUp();
-	UE_API FReply OnArgMoveDown();
 
+	UE_API int32 OnGetArgumentIndex() const;
 	UE_API FText OnGetArgNameText() const;
 	UE_API FText OnGetArgToolTipText() const;
 	UE_API void OnArgNameTextCommitted(const FText& NewText, ETextCommit::Type InTextCommit);
diff --git a/Engine/Plugins/Runtime/RigVM/Source/RigVMEditor/Public/RigVMEditorModule.h b/Engine/Plugins/Runtime/RigVM/Source/RigVMEditor/Public/RigVMEditorModule.h
--- a/Engine/Plugins/Runtime/RigVM/Source/RigVMEditor/Public/RigVMEditorModule.h
+++ b/Engine/Plugins/Runtime/RigVM/Source/RigVMEditor/Public/RigVMEditorModule.h
@@ -121,6 +121,7 @@ protected:
 	UE_API virtual void GetNodeVariantContextMenuActions(IRigVMClientHost* RigVMClientHost, const URigVMEdGraphNode* EdGraphNode, URigVMNode* ModelNode, UToolMenu* Menu) const;
 	UE_API virtual void GetNodeTestContextMenuActions(IRigVMClientHost* RigVMClientHost, const URigVMEdGraphNode* EdGraphNode, URigVMNode* ModelNode, UToolMenu* Menu) const;
 	UE_API virtual void GetNodeDisplayContextMenuActions(IRigVMClientHost* RigVMClientHost, const URigVMEdGraphNode* EdGraphNode, URigVMNode* ModelNode, UToolMenu* Menu) const;
+	UE_API virtual void GetExposedPinContextMenuActions(IRigVMClientHost* RigVMClientHost, const UEdGraphPin* EdGraphPin, URigVMPin* ModelPin, UToolMenu* Menu) const;
 	UE_API virtual void GetPinWorkflowContextMenuActions(IRigVMClientHost* RigVMClientHost, const UEdGraphPin* EdGraphPin, URigVMPin* ModelPin, UToolMenu* Menu) const;
 	UE_API virtual void GetPinDebugContextMenuActions(IRigVMClientHost* RigVMClientHost, const UEdGraphPin* EdGraphPin, URigVMPin* ModelPin, UToolMenu* Menu) const;
 	UE_API virtual void GetPinArrayContextMenuActions(IRigVMClientHost* RigVMClientHost, const UEdGraphPin* EdGraphPin, URigVMPin* ModelPin, UToolMenu* Menu) const;
diff --git a/Engine/Plugins/Runtime/RigVM/Source/RigVMEditor/Public/Widgets/SRigVMGraphNode.h b/Engine/Plugins/Runtime/RigVM/Source/RigVMEditor/Public/Widgets/SRigVMGraphNode.h
--- a/Engine/Plugins/Runtime/RigVM/Source/RigVMEditor/Public/Widgets/SRigVMGraphNode.h
+++ b/Engine/Plugins/Runtime/RigVM/Source/RigVMEditor/Public/Widgets/SRigVMGraphNode.h
@@ -7,6 +7,7 @@
 #include "SGraphNode.h"
 #include "Widgets/Images/SImage.h"
 #include "Widgets/Text/STextBlock.h"
+#include "Widgets/Views/STableRow.h"
 #include "RigVMModel/RigVMPin.h"
 #include "RigVMBlueprint.h"
 #include "Overrides/SOverrideStatusWidget.h"
@@ -19,6 +20,55 @@ class SOverlay;
 class SGraphPin;
 class UEdGraphPin;
 
+class SRigVMGraphPinReorderHandle : public SCompoundWidget
+{
+public:
+
+	DECLARE_DELEGATE_RetVal_ThreeParams(int32, FCanReorderDelegate, int32 /* Source */, int32 /* Target */, EItemDropZone);
+	DECLARE_DELEGATE_ThreeParams(FReorderDelegate, int32 /* Source */, int32 /* Target */, EItemDropZone);
+	
+	SLATE_BEGIN_ARGS(SRigVMGraphPinReorderHandle)
+		: _bIsInput(true)
+		, _Width(50.f)
+		, _Height(24.f)
+		, _Index(INDEX_NONE)
+		, _LineExtension(0.f)
+	{}
+		SLATE_ARGUMENT(bool, bIsInput)
+		SLATE_ARGUMENT(float, Width)
+		SLATE_ARGUMENT(float, Height)
+		SLATE_ATTRIBUTE(int32, Index)
+		SLATE_ATTRIBUTE(float, LineExtension)
+		SLATE_ARGUMENT(TWeakObjectPtr<URigVMNode>, ModelNode)
+		SLATE_EVENT(FCanReorderDelegate, CanReorder)
+		SLATE_EVENT(FReorderDelegate, OnReorder)
+		SLATE_ARGUMENT(TWeakInterfacePtr<IRigVMClientHost>, WeakClientHost)
+	SLATE_END_ARGS()
+
+	void Construct(const FArguments& InArgs);
+
+	virtual FReply OnMouseButtonDown(const FGeometry& MyGeometry, const FPointerEvent& MouseEvent) override;
+	virtual FReply OnDragDetected(const FGeometry& MyGeometry, const FPointerEvent& MouseEvent) override;
+	virtual FReply OnDragOver(const FGeometry& MyGeometry, const FDragDropEvent& DragDropEvent) override;
+	virtual void OnDragLeave(const FDragDropEvent& DragDropEvent) override;
+	virtual FReply OnDrop(const FGeometry& MyGeometry, const FDragDropEvent& DragDropEvent) override;
+
+	virtual int32 OnPaint(const FPaintArgs& Args, const FGeometry& AllottedGeometry, const FSlateRect& MyCullingRect, FSlateWindowElementList& OutDrawElements, int32 LayerId, const FWidgetStyle& InWidgetStyle, bool bParentEnabled) const override;
+
+private:
+
+	bool bIsInput = true;
+	TAttribute<int32> IndexAttribute;
+	TAttribute<float> LineExtension;
+	TWeakObjectPtr<URigVMNode> ModelNode;
+	FCanReorderDelegate CanReorderDelegate;
+	FReorderDelegate ReorderDelegate;
+	TOptional<EItemDropZone> DropZone;
+	TWeakInterfacePtr<IRigVMClientHost> WeakClientHost;
+
+	friend class SRigVMGraphNode;
+};
+
 class SRigVMGraphNode : public SGraphNode
 {
 public:
@@ -60,6 +110,7 @@ public:
 	UE_API virtual void UpdateGraphNode() override;
 	UE_API void UpdateStandardNode();
 	UE_API void UpdateCompactNode();
+	UE_API virtual TOptional<FSlateColor> GetPinTextColor(const SGraphPin* InGraphPin) const override;
 
 	UE_API void CreateAddPinButton();
 
@@ -106,6 +157,9 @@ protected:
 	UE_API bool IsCategoryPin(const SGraphPin* InGraphPin) const;
 	UE_API EOverrideWidgetStatus::Type GetOverrideStatus() const;
 
+	int32 CanReorderExposedPin(int32 InPinIndexToReorder, int32 InTargetPinIndex, EItemDropZone InDropZone);
+	void OnReorderExposedPin(int32 InPinIndexToReorder, int32 InTargetPinIndex, EItemDropZone InDropZone);
+
 	/** Cached widget title area */
 	TSharedPtr<SOverlay> TitleAreaWidget;
 
