RigVM: Follow up work to Collapse node QOL in 44383861

#rnx

diff --git a/Engine/Plugins/Runtime/RigVM/Source/RigVMDeveloper/Public/RigVMModel/Nodes/RigVMFunctionInterfaceNode.h b/Engine/Plugins/Runtime/RigVM/Source/RigVMDeveloper/Public/RigVMModel/Nodes/RigVMFunctionInterfaceNode.h
--- a/Engine/Plugins/Runtime/RigVM/Source/RigVMDeveloper/Public/RigVMModel/Nodes/RigVMFunctionInterfaceNode.h
+++ b/Engine/Plugins/Runtime/RigVM/Source/RigVMDeveloper/Public/RigVMModel/Nodes/RigVMFunctionInterfaceNode.h
@@ -32,8 +32,6 @@ public:
 	// returns true if the interface pin is used for anything within the function / collapse graph
 	UE_API virtual bool IsInterfacePinUsed(const FName& InInterfacePinName) const;
 
-protected:
-
 	UE_API const URigVMPin* FindReferencedPin(const URigVMPin* InPin) const;
 	UE_API const URigVMPin* FindReferencedPin(const FString& InPinPath) const;
 
diff --git a/Engine/Plugins/Runtime/RigVM/Source/RigVMEditor/Private/Widgets/SRigVMGraphNode.cpp b/Engine/Plugins/Runtime/RigVM/Source/RigVMEditor/Private/Widgets/SRigVMGraphNode.cpp
--- a/Engine/Plugins/Runtime/RigVM/Source/RigVMEditor/Private/Widgets/SRigVMGraphNode.cpp
+++ b/Engine/Plugins/Runtime/RigVM/Source/RigVMEditor/Private/Widgets/SRigVMGraphNode.cpp
@@ -59,8 +59,7 @@ class FRigVMGraphPinReorderDragDropOperation : public FDragDropOperation
 public:
 	DRAG_DROP_OPERATOR_TYPE(FRigVMGraphPinReorderDragDropOperation, FDragDropOperation)
 
-	int32 PinIndex = INDEX_NONE;
-	TWeakObjectPtr<URigVMNode> ModelNode;
+	TWeakObjectPtr<URigVMPin> Pin;
 
 	virtual TSharedPtr<SWidget> GetDefaultDecorator() const override
 	{
@@ -71,9 +70,8 @@ public:
 void SRigVMGraphPinReorderHandle::Construct(const FArguments& InArgs)
 {
 	bIsInput = InArgs._bIsInput;
-	IndexAttribute = InArgs._Index;
+	Pin = InArgs._Pin;
 	LineExtension = InArgs._LineExtension;
-	ModelNode = InArgs._ModelNode;
 	CanReorderDelegate = InArgs._CanReorder;
 	ReorderDelegate = InArgs._OnReorder;
 	WeakClientHost = InArgs._WeakClientHost;
@@ -102,22 +100,17 @@ FReply SRigVMGraphPinReorderHandle::OnMouseButtonDown(const FGeometry& MyGeometr
 
 FReply SRigVMGraphPinReorderHandle::OnDragDetected(const FGeometry& MyGeometry, const FPointerEvent& MouseEvent)
 {
-	if (IsEnabled() && MouseEvent.IsMouseButtonDown(EKeys::LeftMouseButton))
+	if (IsEnabled() && MouseEvent.IsMouseButtonDown(EKeys::LeftMouseButton) && Pin.IsValid())
 	{
-		const int32 PinIndex = IndexAttribute.Get();
-		if (ModelNode.IsValid())
+		if (Pin->IsExecuteContext())
 		{
-			if (ModelNode->GetPins()[PinIndex]->IsExecuteContext())
-			{
-				return FReply::Unhandled();
-			}
+			return FReply::Unhandled();
 		}
 		
 		TSharedPtr<FRigVMGraphPinReorderDragDropOperation> DragDropOp = MakeShared<FRigVMGraphPinReorderDragDropOperation>();
 		if (DragDropOp.IsValid())
 		{
-			DragDropOp->PinIndex = PinIndex;
-			DragDropOp->ModelNode = ModelNode;
+			DragDropOp->Pin = Pin;
 			return FReply::Handled().BeginDragDrop(DragDropOp.ToSharedRef());
 		}
 	}
@@ -134,13 +127,13 @@ FReply SRigVMGraphPinReorderHandle::OnDragOver(const FGeometry& MyGeometry, cons
 	if (TSharedPtr<FRigVMGraphPinReorderDragDropOperation> ReorderDragDropOp = DragDropEvent.GetOperationAs<FRigVMGraphPinReorderDragDropOperation>())
 	{
 		// only allow drag & drop of pins on the same node
-		if (ModelNode.IsValid() && ReorderDragDropOp->ModelNode.IsValid() && ModelNode.Get() == ReorderDragDropOp->ModelNode.Get())
+		if (Pin.IsValid() && ReorderDragDropOp->Pin.IsValid() && Pin->GetNode() == ReorderDragDropOp->Pin->GetNode())
 		{
 			DropZone = DesiredDropZone;
 
 			if (CanReorderDelegate.IsBound())
 			{
-				if (CanReorderDelegate.Execute(ReorderDragDropOp->PinIndex, IndexAttribute.Get(), DropZone.GetValue()) == INDEX_NONE)
+				if (CanReorderDelegate.Execute(ReorderDragDropOp->Pin.Get(), Pin.Get(), DropZone.GetValue()) == INDEX_NONE)
 				{
 					DropZone.Reset();
 				}
@@ -176,26 +169,25 @@ FReply SRigVMGraphPinReorderHandle::OnDrop(const FGeometry& MyGeometry, const FD
 		if (TSharedPtr<FRigVMGraphPinReorderDragDropOperation> ReorderDragDropOp = DragDropEvent.GetOperationAs<FRigVMGraphPinReorderDragDropOperation>())
 		{
 			// only allow drag & drop of pins on the same node
-			if (ModelNode.IsValid() && ReorderDragDropOp->ModelNode.IsValid() && ModelNode.Get() == ReorderDragDropOp->ModelNode.Get())
+			if (Pin.IsValid() && ReorderDragDropOp->Pin.IsValid() && Pin->GetNode() == ReorderDragDropOp->Pin->GetNode())
 			{
 				if (ReorderDelegate.IsBound())
 				{
-					ReorderDelegate.Execute(ReorderDragDropOp->PinIndex, IndexAttribute.Get(), DropZone.GetValue());
+					ReorderDelegate.Execute(ReorderDragDropOp->Pin.Get(), Pin.Get(), DropZone.GetValue());
 				}
 			}
 		}
 		else if (TSharedPtr<FGraphEditorDragDropAction> PinDragDropOp = DragDropEvent.GetOperationAs<FGraphEditorDragDropAction>())
 		{
-			if (PinDragDropOp->IsOfTypeImpl(TEXT("FDragConnection")) && ModelNode.IsValid() && WeakClientHost.IsValid())
+			if (PinDragDropOp->IsOfTypeImpl(TEXT("FDragConnection")) && Pin.IsValid() && WeakClientHost.IsValid())
 			{
-				int32 NewIndex = IndexAttribute.Get();
+				int32 NewIndex = Pin->GetPinIndex();
 				if (DropZone.GetValue() == EItemDropZone::BelowItem)
 				{
 					NewIndex++;
 				}
 
-				URigVMController* Controller = WeakClientHost.Get()->GetOrCreateController(ModelNode.Get()->GetGraph());
-				if (Controller)
+				if (URigVMController* Controller = WeakClientHost.Get()->GetOrCreateController(Pin->GetGraph()))
 				{
 					FName ExposedPinName = NAME_None;
 
@@ -2843,7 +2835,7 @@ void SRigVMGraphNode::UpdatePinTreeView()
 		{
 			const bool bIsInput = ModelNode->IsA<URigVMFunctionReturnNode>();
 			IRigVMClientHost* Host = RigGraphNode->GetImplementingOuter<IRigVMClientHost>();
-			const URigVMPin* ModelPin = ModelNode->GetGraph()->FindPin(bIsInput ? InputPinInfo.Identifier : OutputPinInfo.Identifier);
+			URigVMPin* ModelPin = ModelNode->GetGraph()->FindPin(bIsInput ? InputPinInfo.Identifier : OutputPinInfo.Identifier);
 
 			SlotLayout->InsertSlot(bIsInput ? SlotLayout->NumSlots() : 0)
 			.AutoWidth()
@@ -2853,9 +2845,8 @@ void SRigVMGraphNode::UpdatePinTreeView()
 				SNew(SRigVMGraphPinReorderHandle)
 				.bIsInput(bIsInput)
 				.Width(EmptySidePadding)
-				.Index(ModelPin->GetPinIndex())
+				.Pin(ModelPin)
 				.LineExtension(44.f)
-				.ModelNode(ModelNode)
 				.CanReorder(this, &SRigVMGraphNode::CanReorderExposedPin)
 				.OnReorder(this, &SRigVMGraphNode::OnReorderExposedPin)
 				.WeakClientHost(Host)
@@ -2963,9 +2954,26 @@ EOverrideWidgetStatus::Type SRigVMGraphNode::GetOverrideStatus() const
 	return EOverrideWidgetStatus::ChangedHere;
 }
 
-int32 SRigVMGraphNode::CanReorderExposedPin(int32 InPinIndexToReorder, int32 InTargetPinIndex, EItemDropZone InDropZone)
+int32 SRigVMGraphNode::CanReorderExposedPin(const URigVMPin* InPinToReorder, const URigVMPin* InTargetPin, EItemDropZone InDropZone)
 {
-	if (InPinIndexToReorder == InTargetPinIndex)
+	if (!InPinToReorder || !InTargetPin || InPinToReorder == InTargetPin)
+	{
+		return INDEX_NONE;
+	}
+
+	check(InPinToReorder->IsRootPin());
+	check(InTargetPin->IsRootPin());
+	check(InPinToReorder->GetNode() == InTargetPin->GetNode());
+
+	const URigVMFunctionInterfaceNode* InterfaceNode = Cast<URigVMFunctionInterfaceNode>(InPinToReorder->GetNode());
+	if (InterfaceNode == nullptr)
+	{
+		return INDEX_NONE;
+	}
+
+	const URigVMPin* PinToReorder = InterfaceNode->FindReferencedPin(InPinToReorder);
+	const URigVMPin* TargetPin = InterfaceNode->FindReferencedPin(InTargetPin);
+	if (PinToReorder == nullptr || TargetPin == nullptr)
 	{
 		return INDEX_NONE;
 	}
@@ -2975,23 +2983,23 @@ int32 SRigVMGraphNode::CanReorderExposedPin(int32 InPinIndexToReorder, int32 InT
 		InDropZone = EItemDropZone::BelowItem;
 	}
 
-	int32 NewPinIndex = InTargetPinIndex;
+	int32 NewPinIndex = TargetPin->GetPinIndex();
 	if (InDropZone == EItemDropZone::BelowItem)
 	{
-		if (NewPinIndex < InPinIndexToReorder)
+		if (NewPinIndex < PinToReorder->GetPinIndex())
 		{
 			NewPinIndex++;
 		}
 	}
 	else if (InDropZone == EItemDropZone::AboveItem)
 	{
-		if (NewPinIndex > InPinIndexToReorder)
+		if (NewPinIndex > PinToReorder->GetPinIndex())
 		{
 			NewPinIndex--;
 		}
 	}
 
-	if (InPinIndexToReorder == NewPinIndex)
+	if (PinToReorder->GetPinIndex() == NewPinIndex)
 	{
 		return INDEX_NONE;
 	}
@@ -2999,9 +3007,9 @@ int32 SRigVMGraphNode::CanReorderExposedPin(int32 InPinIndexToReorder, int32 InT
 	return NewPinIndex;
 }
 
-void SRigVMGraphNode::OnReorderExposedPin(int32 InPinIndexToReorder, int32 InTargetPinIndex, EItemDropZone InDropZone)
+void SRigVMGraphNode::OnReorderExposedPin(const URigVMPin* InPinToReorder, const URigVMPin* InTargetPin, EItemDropZone InDropZone)
 {
-	const int32 NewPinIndex = CanReorderExposedPin(InPinIndexToReorder, InTargetPinIndex, InDropZone);
+	const int32 NewPinIndex = CanReorderExposedPin(InPinToReorder, InTargetPin, InDropZone);
 	if (NewPinIndex == INDEX_NONE)
 	{
 		return;
@@ -3025,9 +3033,7 @@ void SRigVMGraphNode::OnReorderExposedPin(int32 InPinIndexToReorder, int32 InTar
 		return;
 	}
 
-	const URigVMPin* Pin = InterfaceNode->GetPins()[InPinIndexToReorder];
-	check(Pin);
-	(void)Controller->SetExposedPinIndex(Pin->GetFName(), NewPinIndex);
+	(void)Controller->SetExposedPinIndex(InPinToReorder->GetFName(), NewPinIndex);
 }
 
 #undef LOCTEXT_NAMESPACE
diff --git a/Engine/Plugins/Runtime/RigVM/Source/RigVMEditor/Public/Widgets/SRigVMGraphNode.h b/Engine/Plugins/Runtime/RigVM/Source/RigVMEditor/Public/Widgets/SRigVMGraphNode.h
--- a/Engine/Plugins/Runtime/RigVM/Source/RigVMEditor/Public/Widgets/SRigVMGraphNode.h
+++ b/Engine/Plugins/Runtime/RigVM/Source/RigVMEditor/Public/Widgets/SRigVMGraphNode.h
@@ -24,22 +24,20 @@ class SRigVMGraphPinReorderHandle : public SCompoundWidget
 {
 public:
 
-	DECLARE_DELEGATE_RetVal_ThreeParams(int32, FCanReorderDelegate, int32 /* Source */, int32 /* Target */, EItemDropZone);
-	DECLARE_DELEGATE_ThreeParams(FReorderDelegate, int32 /* Source */, int32 /* Target */, EItemDropZone);
+	DECLARE_DELEGATE_RetVal_ThreeParams(int32, FCanReorderDelegate, const URigVMPin* /* Source */, const URigVMPin* /* Target */, EItemDropZone);
+	DECLARE_DELEGATE_ThreeParams(FReorderDelegate, const URigVMPin* /* Source */, const URigVMPin* /* Target */, EItemDropZone);
 	
 	SLATE_BEGIN_ARGS(SRigVMGraphPinReorderHandle)
 		: _bIsInput(true)
 		, _Width(50.f)
 		, _Height(24.f)
-		, _Index(INDEX_NONE)
 		, _LineExtension(0.f)
 	{}
 		SLATE_ARGUMENT(bool, bIsInput)
 		SLATE_ARGUMENT(float, Width)
 		SLATE_ARGUMENT(float, Height)
-		SLATE_ATTRIBUTE(int32, Index)
+		SLATE_ARGUMENT(TWeakObjectPtr<URigVMPin>, Pin)
 		SLATE_ATTRIBUTE(float, LineExtension)
-		SLATE_ARGUMENT(TWeakObjectPtr<URigVMNode>, ModelNode)
 		SLATE_EVENT(FCanReorderDelegate, CanReorder)
 		SLATE_EVENT(FReorderDelegate, OnReorder)
 		SLATE_ARGUMENT(TWeakInterfacePtr<IRigVMClientHost>, WeakClientHost)
@@ -58,9 +56,8 @@ public:
 private:
 
 	bool bIsInput = true;
-	TAttribute<int32> IndexAttribute;
+	TWeakObjectPtr<URigVMPin> Pin;
 	TAttribute<float> LineExtension;
-	TWeakObjectPtr<URigVMNode> ModelNode;
 	FCanReorderDelegate CanReorderDelegate;
 	FReorderDelegate ReorderDelegate;
 	TOptional<EItemDropZone> DropZone;
@@ -157,8 +154,8 @@ protected:
 	UE_API bool IsCategoryPin(const SGraphPin* InGraphPin) const;
 	UE_API EOverrideWidgetStatus::Type GetOverrideStatus() const;
 
-	int32 CanReorderExposedPin(int32 InPinIndexToReorder, int32 InTargetPinIndex, EItemDropZone InDropZone);
-	void OnReorderExposedPin(int32 InPinIndexToReorder, int32 InTargetPinIndex, EItemDropZone InDropZone);
+	int32 CanReorderExposedPin(const URigVMPin* InPinToReorder, const URigVMPin* InTargetPin, EItemDropZone InDropZone);
+	void OnReorderExposedPin(const URigVMPin* InPinToReorder, const URigVMPin* InTargetPin, EItemDropZone InDropZone);
 
 	/** Cached widget title area */
 	TSharedPtr<SOverlay> TitleAreaWidget;
