RigVM: Defensive coding to guard against invalid array access during EnumeratePropertiesOnHostAndStructInstance

#rnx

diff --git a/Engine/Plugins/Runtime/RigVM/Source/RigVMDeveloper/Private/RigVMModel/Nodes/RigVMUnitNode.cpp b/Engine/Plugins/Runtime/RigVM/Source/RigVMDeveloper/Private/RigVMModel/Nodes/RigVMUnitNode.cpp
--- a/Engine/Plugins/Runtime/RigVM/Source/RigVMDeveloper/Private/RigVMModel/Nodes/RigVMUnitNode.cpp
+++ b/Engine/Plugins/Runtime/RigVM/Source/RigVMDeveloper/Private/RigVMModel/Nodes/RigVMUnitNode.cpp
@@ -506,13 +506,18 @@ void URigVMUnitNode::EnumeratePropertiesOnHostAndStructInstance(URigVMHost* InHo
 	const FRigVMInstructionArray& Instructions = ByteCode.GetInstructions();
 	for(int32 Index = 0; Index < InstructionIndices.Num(); Index++)
 	{
-		const int32 InstructionIndex = InstructionIndices[Index]; 
-		if(Instructions[InstructionIndices[Index]].OpCode != ERigVMOpCode::Execute)
+		const int32& InstructionIndex = InstructionIndices[Index];
+		if (!Instructions.IsValidIndex(InstructionIndex))
 		{
 			continue;
 		}
-
+		
 		const FRigVMInstruction& Instruction = Instructions[InstructionIndex];
+		if(Instruction.OpCode != ERigVMOpCode::Execute)
+		{
+			continue;
+		}
+
 		const FRigVMOperandArray Operands = ByteCode.GetOperandsForOp(Instruction);
 		const UScriptStruct* ScriptStruct = Cast<UScriptStruct>(InInstance->GetStruct());
 
@@ -522,14 +527,20 @@ void URigVMUnitNode::EnumeratePropertiesOnHostAndStructInstance(URigVMHost* InHo
 			return;
 		}
 
-		for(int32 ArgumentIndex = 0; ArgumentIndex < Function->GetArguments().Num(); ArgumentIndex++)
+		const TArray<FRigVMFunctionArgument>& FunctionArguments = Function->GetArguments();
+		for(int32 ArgumentIndex = 0; ArgumentIndex < FunctionArguments.Num(); ArgumentIndex++)
 		{
-			const FRigVMFunctionArgument& Argument = Function->GetArguments()[ArgumentIndex];
+			const FRigVMFunctionArgument& Argument = FunctionArguments[ArgumentIndex];
 			const URigVMPin* Pin = FindPin(Argument.Name);
 			if(Pin == nullptr)
 			{
 				continue;
 			}
+
+			if (!Operands.IsValidIndex(ArgumentIndex))
+			{
+				continue;
+			}
 			
 			FRigVMOperand Operand = Operands[ArgumentIndex];
 
@@ -566,7 +577,7 @@ void URigVMUnitNode::EnumeratePropertiesOnHostAndStructInstance(URigVMHost* InHo
 			
 			switch(Operand.GetMemoryType())
 			{
-			case ERigVMMemoryType::External:
+				case ERigVMMemoryType::External:
 				{
 					TArray<FRigVMExternalVariable> ExternalVariables = InHost->GetExternalVariables();
 					if(ExternalVariables.IsValidIndex(Operand.GetRegisterIndex()))
@@ -577,7 +588,7 @@ void URigVMUnitNode::EnumeratePropertiesOnHostAndStructInstance(URigVMHost* InHo
 					}
 					break;
 				}
-			case ERigVMMemoryType::Literal:
+				case ERigVMMemoryType::Literal:
 				{
 					if(FRigVMMemoryStorageStruct* LiteralMemory = InHost->GetLiteralMemory())
 					{
@@ -589,7 +600,7 @@ void URigVMUnitNode::EnumeratePropertiesOnHostAndStructInstance(URigVMHost* InHo
 					}
 					break;
 				}
-			case ERigVMMemoryType::Work:
+				case ERigVMMemoryType::Work:
 				{
 					if(FRigVMMemoryStorageStruct* WorkMemory = InHost->GetWorkMemory())
 					{
@@ -601,7 +612,7 @@ void URigVMUnitNode::EnumeratePropertiesOnHostAndStructInstance(URigVMHost* InHo
 					}
 					break;
 				}
-			default:
+				default:
 				{
 					break;
 				}
