Control Rig: Only resolve the asset path for the node spawner when needed, not upon construction, since the asset registry might not be ready at that point.

#rnx

diff --git a/Engine/Plugins/Runtime/RigVM/Source/RigVMEditor/Private/EdGraph/NodeSpawners/RigVMEdGraphFunctionRefNodeSpawner.cpp b/Engine/Plugins/Runtime/RigVM/Source/RigVMEditor/Private/EdGraph/NodeSpawners/RigVMEdGraphFunctionRefNodeSpawner.cpp
--- a/Engine/Plugins/Runtime/RigVM/Source/RigVMEditor/Private/EdGraph/NodeSpawners/RigVMEdGraphFunctionRefNodeSpawner.cpp
+++ b/Engine/Plugins/Runtime/RigVM/Source/RigVMEditor/Private/EdGraph/NodeSpawners/RigVMEdGraphFunctionRefNodeSpawner.cpp
@@ -40,11 +40,7 @@ URigVMEdGraphFunctionRefNodeSpawner* URigVMEdGraphFunctionRefNodeSpawner::Create
 	NodeSpawner->ReferencedPublicFunctionHeader = InFunction->GetFunctionHeader();
 	NodeSpawner->NodeClass = URigVMEdGraphNode::StaticClass();
 	NodeSpawner->bIsLocalFunction = true;
-	if (!GEditor->PlayWorld && !GIsPlayInEditorWorld)
-	{
-		const FAssetData Asset = UE::RigVM::Editor::Tools::FindAssetFromAnyPath(InFunction->GetPathName(), true);
-		NodeSpawner->AssetPath = Asset.ToSoftObjectPath();
-	}
+	NodeSpawner->FunctionPath = InFunction->GetPathName();
 
 	FBlueprintActionUiSpec& MenuSignature = NodeSpawner->DefaultMenuSignature;
 
@@ -235,7 +231,7 @@ FBlueprintNodeSignature URigVMEdGraphFunctionRefNodeSpawner::GetSpawnerSignature
 	{
 		SignatureString = 
 			FString::Printf(TEXT("RigFunction=%s:%s"),
-			*AssetPath.ToString(),
+			*GetResolvedAssetPath().ToString(),
 			*ReferencedPublicFunctionHeader.Name.ToString());
 	}
 
@@ -263,9 +259,9 @@ UEdGraphNode* URigVMEdGraphFunctionRefNodeSpawner::Invoke(UEdGraph* ParentGraph,
 	// if we are trying to build the real function ref - but we haven't loaded the asset yet...
 	if(!FBlueprintNodeTemplateCache::IsTemplateOuter(ParentGraph))
 	{
-		if (AssetPath.IsValid())
+		if (GetResolvedAssetPath().IsValid())
 		{
-			if (URigVMBlueprint* Blueprint = Cast<URigVMBlueprint>(AssetPath.TryLoad()))
+			if (URigVMBlueprint* Blueprint = Cast<URigVMBlueprint>(GetResolvedAssetPath().TryLoad()))
 			{
 				ReferencedPublicFunctionHeader = Blueprint->GetLocalFunctionLibrary()->FindFunction(ReferencedPublicFunctionHeader.Name)->GetFunctionHeader();			
 			}
@@ -406,5 +402,17 @@ bool URigVMEdGraphFunctionRefNodeSpawner::IsTemplateNodeFilteredOut(FBlueprintAc
 	return false;
 }
 
+FSoftObjectPath URigVMEdGraphFunctionRefNodeSpawner::GetResolvedAssetPath() const
+{
+	if (!FunctionPath.IsEmpty())
+	{
+		const FAssetData Asset = UE::RigVM::Editor::Tools::FindAssetFromAnyPath(FunctionPath, true);
+		AssetPath = Asset.ToSoftObjectPath();
+		FunctionPath.Reset();
+	}
+
+	return AssetPath;
+}
+
 #undef LOCTEXT_NAMESPACE
 
diff --git a/Engine/Plugins/Runtime/RigVM/Source/RigVMEditor/Public/EdGraph/NodeSpawners/RigVMEdGraphFunctionRefNodeSpawner.h b/Engine/Plugins/Runtime/RigVM/Source/RigVMEditor/Public/EdGraph/NodeSpawners/RigVMEdGraphFunctionRefNodeSpawner.h
--- a/Engine/Plugins/Runtime/RigVM/Source/RigVMEditor/Public/EdGraph/NodeSpawners/RigVMEdGraphFunctionRefNodeSpawner.h
+++ b/Engine/Plugins/Runtime/RigVM/Source/RigVMEditor/Public/EdGraph/NodeSpawners/RigVMEdGraphFunctionRefNodeSpawner.h
@@ -51,7 +51,11 @@ public:
 	// End UBlueprintNodeSpawner interface
 
 private:
-
+	/** Returns a resolved asset path. This is directly from the asset data if CreateFromAssetData was used to create the spawner, or if
+	 *  CreateFromFunction, the resolved function path.
+	 */
+	FSoftObjectPath GetResolvedAssetPath() const;
+	
 	/** The public function definition we will spawn from [optional] */
 	UPROPERTY(Transient)
 	mutable FRigVMGraphFunctionHeader ReferencedPublicFunctionHeader;
@@ -61,7 +65,10 @@ private:
 	bool bIsLocalFunction;
 
 	UPROPERTY()
-	FSoftObjectPath AssetPath;
+	mutable FSoftObjectPath AssetPath;
+
+	UPROPERTY()
+	mutable FString FunctionPath;
 
 	static UE_API URigVMEdGraphNode* SpawnNode(UEdGraph* ParentGraph, UBlueprint* Blueprint, FRigVMGraphFunctionHeader& InFunction, FVector2D const Location);
 
