RigVM: Improve error reporting during localizeing of registries

#rnx

diff --git a/Engine/Plugins/Runtime/RigVM/Source/RigVM/Private/RigVMCore/RigVM.cpp b/Engine/Plugins/Runtime/RigVM/Source/RigVM/Private/RigVMCore/RigVM.cpp
--- a/Engine/Plugins/Runtime/RigVM/Source/RigVM/Private/RigVMCore/RigVM.cpp
+++ b/Engine/Plugins/Runtime/RigVM/Source/RigVM/Private/RigVMCore/RigVM.cpp
@@ -2651,9 +2651,18 @@ bool URigVM::CreateLocalizedRegistryIfRequired(bool bForce)
 	for(const TPair<const FRigVMDispatchFactory*, TArray<const FRigVMFunction*>>& Pair : RequiredFactories)
 	{
 		const FRigVMDispatchFactory* Factory = Pair.Key;
-		check(Factory);
+		if (Factory == nullptr)
+		{
+			UE_LOG(LogRigVM, Error, TEXT("Found an unexpected nullptr for required dispatch factories in VM '%s'."), *GetPathName());
+			return false;
+		}
+
 		UScriptStruct* FactoryStruct = Factory->GetScriptStruct();
-		check(FactoryStruct);
+		if (FactoryStruct == nullptr)
+		{
+			UE_LOG(LogRigVM, Error, TEXT("Dispatch factory '%s' has an unexpected nullptr for its script struct in VM '%s'."), *Factory->GetFactoryName().ToString(), *GetPathName());
+			return false;
+		}
 
 		// flatten the argument infos to the permutations used only
 		TArray<FRigVMTemplateArgumentInfo> TemporaryArgumentInfos;
