RigVM: Improve / fix unit tests around collapse nodes

#rnx

diff --git a/Engine/Plugins/Runtime/RigVM/Source/RigVMDeveloper/Private/RigVMModel/RigVMController.cpp b/Engine/Plugins/Runtime/RigVM/Source/RigVMDeveloper/Private/RigVMModel/RigVMController.cpp
--- a/Engine/Plugins/Runtime/RigVM/Source/RigVMDeveloper/Private/RigVMModel/RigVMController.cpp
+++ b/Engine/Plugins/Runtime/RigVM/Source/RigVMDeveloper/Private/RigVMModel/RigVMController.cpp
@@ -5513,25 +5513,21 @@ URigVMCollapseNode* URigVMController::CollapseNodes(const TArray<URigVMNode*>& I
 		const URigVMPin* PinToCollapse = CollapseInfo.Pin;
 		const URigVMPin* PinForConfiguration = bIsInput ? CollapseInfo.LinkedPins[0].Pin : CollapseInfo.Pin;
 		
-		if (CollapsedPins.Contains(PinToCollapse))
-		{
-			return;
-		}
-
-		if(PinToCollapse->IsExecuteContext() && PinToCollapse->GetDirection() == ERigVMPinDirection::IO)
+		if(PinToCollapse->IsExecuteContext())
 		{
 			for(const TPair<const URigVMPin*, const URigVMPin*>& Pair : CollapsedPins)
 			{
-				if(Pair.Key->IsExecuteContext() && Pair.Key->GetDirection() == ERigVMPinDirection::IO)
+				if(Pair.Key != PinToCollapse && Pair.Key->IsExecuteContext())
 				{
 					CollapsedPins.Add(PinToCollapse, Pair.Value);
 					break;
 				}
 			}
-			if (CollapsedPins.Contains(PinToCollapse))
-			{
-				return;
-			}
+		}
+
+		if (CollapsedPins.Contains(PinToCollapse))
+		{
+			return;
 		}
 
 		const FName DesiredPinName = PinForConfiguration->GetFName();
@@ -5542,7 +5538,12 @@ URigVMCollapseNode* URigVMController::CollapseNodes(const TArray<URigVMNode*>& I
 		URigVMPin* CollapsedPin = NewObject<URigVMPin>(CollapseNode, PinName);
 		ConfigurePinFromPin(CollapsedPin, PinForConfiguration, false);
 
-		if (CollapsedPin->GetDirection() == ERigVMPinDirection::IO)
+		if (PinToCollapse->IsExecuteContext())
+		{
+			CollapsedPin->SetDisplayName(FRigVMStruct::ExecuteName.Resolve());
+			CollapsedPin->Direction = ERigVMPinDirection::IO;
+		}
+		else if (CollapsedPin->GetDirection() == ERigVMPinDirection::IO)
 		{
 			if(!CollapsedPin->IsExecuteContext())
 			{
@@ -8890,24 +8891,27 @@ bool URigVMController::SetPinCategory(URigVMPin* InPin, const FString& InCategor
 	}
 
 	int32 IndexInCategory = INDEX_NONE;
-	if(InIndexInCategory == INDEX_NONE)
+	if (!NewCategory.IsEmpty())
 	{
-		IndexInCategory = InPin->GetNode()->GetPinsForCategory(NewCategory).Num();
-	}
-	else if (!NewCategory.IsEmpty())
-	{
-		const TArray<URigVMPin*> PinsInNewCategory = InPin->GetNode()->GetPinsForCategory(NewCategory);
-		IndexInCategory = FMath::Min(InIndexInCategory, PinsInNewCategory.Num());
-
-		// move all pins in the new category unless we are adding to the end
-		if (IndexInCategory != PinsInNewCategory.Num())
+		if(InIndexInCategory == INDEX_NONE)
 		{
-			for(URigVMPin* OtherPin : PinsInNewCategory)
+			IndexInCategory = InPin->GetNode()->GetPinsForCategory(NewCategory).Num();
+		}
+		else
+		{
+			const TArray<URigVMPin*> PinsInNewCategory = InPin->GetNode()->GetPinsForCategory(NewCategory);
+			IndexInCategory = FMath::Min(InIndexInCategory, PinsInNewCategory.Num());
+
+			// move all pins in the new category unless we are adding to the end
+			if (IndexInCategory != PinsInNewCategory.Num())
 			{
-				if(OtherPin != InPin && OtherPin->IndexInCategory != INDEX_NONE && OtherPin->IndexInCategory >= IndexInCategory)
+				for(URigVMPin* OtherPin : PinsInNewCategory)
 				{
-					OtherPin->IndexInCategory++;
-					Notify(ERigVMGraphNotifType::PinCategoryChanged, OtherPin);
+					if(OtherPin != InPin && OtherPin->IndexInCategory != INDEX_NONE && OtherPin->IndexInCategory >= IndexInCategory)
+					{
+						OtherPin->IndexInCategory++;
+						Notify(ERigVMGraphNotifType::PinCategoryChanged, OtherPin);
+					}
 				}
 			}
 		}
