[AutoRTFM] Enable various ControlRig/RigVM unit tests with AutoRTFM enabled.

This shook out a few places where we might encounter locks while inside a transaction and those locks are now fixed.

This CL selectively enables only some tests. Not all tests were good candidates for AutoRTFM for various reasons:
- They explicitly spin up multiple threads or tasks.
- They rely on compiling a Blueprint, which is currently not AutoRTFM-safe.
- They rely on serialization/compression and hit AutoRTFM-unaware code like Oodle.

diff --git a/Engine/Plugins/Runtime/RigVM/Source/RigVM/Private/RigVMCore/RigVM.cpp b/Engine/Plugins/Runtime/RigVM/Source/RigVM/Private/RigVMCore/RigVM.cpp
--- a/Engine/Plugins/Runtime/RigVM/Source/RigVM/Private/RigVMCore/RigVM.cpp
+++ b/Engine/Plugins/Runtime/RigVM/Source/RigVM/Private/RigVMCore/RigVM.cpp
@@ -410,7 +410,7 @@ void URigVM::RefreshArgumentNameCaches()
 {
 	// make sure that the functions cannot change at the same time as
 	// the argument name caches are being updated.
-	FScopeLock ResolveFunctionsScopeLock(&ResolveFunctionsMutex);
+	UE::TScopeLock ResolveFunctionsScopeLock(ResolveFunctionsMutex);
 
 	int32 TotalNumFactories;
 	if(LocalizedRegistry.IsValid())
@@ -1096,10 +1096,10 @@ bool URigVM::ResolveFunctionsIfRequired()
 	// only create the function resolval lock if we are running
 	// on the global registry.
 #if WITH_EDITOR
-	TUniquePtr<FScopeLock> ResolveFunctionsScopeLock;
+	TUniquePtr<UE::TScopeLock<FTransactionallySafeCriticalSection>> ResolveFunctionsScopeLock;
 	if(!LocalizedRegistry.IsValid())
 	{
-		ResolveFunctionsScopeLock = MakeUnique<FScopeLock>(&ResolveFunctionsMutex);
+		ResolveFunctionsScopeLock = MakeUnique<UE::TScopeLock<FTransactionallySafeCriticalSection>>(ResolveFunctionsMutex);
 	}
 #endif
 	
diff --git a/Engine/Plugins/Runtime/RigVM/Source/RigVM/Public/RigVMCore/RigVM.h b/Engine/Plugins/Runtime/RigVM/Source/RigVM/Public/RigVMCore/RigVM.h
--- a/Engine/Plugins/Runtime/RigVM/Source/RigVM/Public/RigVMCore/RigVM.h
+++ b/Engine/Plugins/Runtime/RigVM/Source/RigVM/Public/RigVMCore/RigVM.h
@@ -15,6 +15,7 @@
 #include "Math/Vector2D.h"
 #include "Misc/AssertionMacros.h"
 #include "Misc/MTTransactionallySafeAccessDetector.h"
+#include "Misc/TransactionallySafeCriticalSection.h"
 #include "RigVMByteCode.h"
 #include "RigVMCore/RigVMExternalVariable.h"
 #include "RigVMCore/RigVMFunction.h"
@@ -1186,7 +1187,7 @@ protected:
 	RIGVM_API void StartProfiling(FRigVMExtendedExecuteContext& Context);
 	RIGVM_API void StopProfiling(FRigVMExtendedExecuteContext& Context);
 
-	FCriticalSection ResolveFunctionsMutex;
+	FTransactionallySafeCriticalSection ResolveFunctionsMutex;
 
 	friend class URigVMCompiler;
 	friend struct FRigVMCompilerWorkData;
diff --git a/Engine/Plugins/Runtime/RigVM/Source/RigVMDeveloper/Private/RigVMModel/RigVMController.cpp b/Engine/Plugins/Runtime/RigVM/Source/RigVMDeveloper/Private/RigVMModel/RigVMController.cpp
--- a/Engine/Plugins/Runtime/RigVM/Source/RigVMDeveloper/Private/RigVMModel/RigVMController.cpp
+++ b/Engine/Plugins/Runtime/RigVM/Source/RigVMDeveloper/Private/RigVMModel/RigVMController.cpp
@@ -18557,7 +18557,7 @@ void URigVMController::GenerateRepopulatePinsNodeData(TArray<FRepopulatePinsNode
 	}
 
 	FRigVMControllerCompileBracketScope CompileBracketScope(this);
-	FScopeLock Lock(&PinPathCoreRedirectorsLock);
+	UE::TScopeLock Lock(PinPathCoreRedirectorsLock);
 
 	FRepopulatePinsNodeData NodeData;
 	NodeData.Node = InNode;
@@ -18610,7 +18610,7 @@ void URigVMController::GenerateRepopulatePinsNodeData(TArray<FRepopulatePinsNode
 void URigVMController::OrphanPins(const TArray<FRepopulatePinsNodeData>& NodesPinData)
 {
 	FRigVMControllerCompileBracketScope CompileBracketScope(this);
-	FScopeLock Lock(&PinPathCoreRedirectorsLock);
+	UE::TScopeLock Lock(PinPathCoreRedirectorsLock);
 
 	for (const FRepopulatePinsNodeData& NodeData : NodesPinData)
 	{
@@ -18662,7 +18662,7 @@ void URigVMController::OrphanPins(const TArray<FRepopulatePinsNodeData>& NodesPi
 void URigVMController::RepopulatePins(const TArray<FRepopulatePinsNodeData>& NodesPinData)
 {
 	FRigVMControllerCompileBracketScope CompileBracketScope(this);
-	FScopeLock Lock(&PinPathCoreRedirectorsLock);
+	UE::TScopeLock Lock(PinPathCoreRedirectorsLock);
 
 	const FRigVMRegistry& Registry = FRigVMRegistry::Get();
 	for (const FRepopulatePinsNodeData& NodeData : NodesPinData)
@@ -18680,7 +18680,7 @@ void URigVMController::RepopulatePinsOnNode(URigVMNode* InNode, bool bFollowCore
 	}
 
 	FRigVMControllerCompileBracketScope CompileBracketScope(this);
-	FScopeLock Lock(&PinPathCoreRedirectorsLock);
+	UE::TScopeLock Lock(PinPathCoreRedirectorsLock);
 
 	TArray<FRepopulatePinsNodeData> NodesPinData;
 	GenerateRepopulatePinsNodeData(NodesPinData, InNode, bFollowCoreRedirectors, bSetupOrphanedPins, bRecreateLinks);
diff --git a/Engine/Plugins/Runtime/RigVM/Source/RigVMDeveloper/Public/RigVMModel/RigVMController.h b/Engine/Plugins/Runtime/RigVM/Source/RigVMDeveloper/Public/RigVMModel/RigVMController.h
--- a/Engine/Plugins/Runtime/RigVM/Source/RigVMDeveloper/Public/RigVMModel/RigVMController.h
+++ b/Engine/Plugins/Runtime/RigVM/Source/RigVMDeveloper/Public/RigVMModel/RigVMController.h
@@ -22,6 +22,7 @@
 #include "RigVMCore/RigVMObjectArchive.h"
 #include "UObject/Interface.h"
 #include "Algo/Sort.h"
+#include "Misc/TransactionallySafeCriticalSection.h"
 #include "RigVMController.generated.h"
 
 #define UE_API RIGVMDEVELOPER_API
@@ -1667,7 +1668,7 @@ private:
 	};
 
 	static UE_API TMap<FRigVMStructPinRedirectorKey, FString> PinPathCoreRedirectors;
-	FCriticalSection PinPathCoreRedirectorsLock;
+	FTransactionallySafeCriticalSection PinPathCoreRedirectorsLock;
 
 	FRigVMUnitNodeCreatedContext UnitNodeCreatedContext;
 
diff --git a/Engine/Source/Runtime/CoreUObject/Public/UObject/FastReferenceCollector.h b/Engine/Source/Runtime/CoreUObject/Public/UObject/FastReferenceCollector.h
--- a/Engine/Source/Runtime/CoreUObject/Public/UObject/FastReferenceCollector.h
+++ b/Engine/Source/Runtime/CoreUObject/Public/UObject/FastReferenceCollector.h
@@ -608,6 +608,13 @@ COREUOBJECT_API bool& CalledSuperARO();
 template<class DispatcherType>
 FORCEINLINE_DEBUGGABLE void CallARO(DispatcherType& Dispatcher, UObject* Instance, FMemberWord Word)
 {
+	if (Instance->HasAnyInternalFlags(EInternalObjectFlags::AutoRTFMConstructionAborted))
+	{
+		// Object was constructed in a transaction that was aborted.
+		// The object is in a pre-construction state and won't have any valid subreferences.
+		return;
+	}
+
 #if WITH_EDITORONLY_DATA
 	bool& bCalledSuperARO = CalledSuperARO();
 	TGuardValue<bool> CalledSuperAROScope(bCalledSuperARO, false);
