RigVM: More soft defensive code within CreateLocalizedRegistry

#rnx

diff --git a/Engine/Plugins/Runtime/RigVM/Source/RigVM/Private/RigVMCore/RigVM.cpp b/Engine/Plugins/Runtime/RigVM/Source/RigVM/Private/RigVMCore/RigVM.cpp
--- a/Engine/Plugins/Runtime/RigVM/Source/RigVM/Private/RigVMCore/RigVM.cpp
+++ b/Engine/Plugins/Runtime/RigVM/Source/RigVM/Private/RigVMCore/RigVM.cpp
@@ -2676,11 +2676,25 @@ bool URigVM::CreateLocalizedRegistryIfRequired(bool bForce)
 				for(int32 TypeIndex = 0; TypeIndex < TypeIndices.Num(); TypeIndex++)
 				{
 					const FRigVMTemplateArgumentType Type = ReadLock.GetRegistry().GetType_NoLock(TypeIndices[TypeIndex]);
+					if (!Type.IsValid())
+					{
+						UE_LOG(LogRigVM, Error, TEXT("Cannot register dispatch dispatch '%s' for VM '%s'. Invalid types found.\nCannot localize registry."), *Factory->GetFactoryName().ToString(), *GetPathName());
+						return false;
+					}
+
 					TypeIndices[TypeIndex] = TemporaryRegistry.FindOrAddType_NoLock(Type);
-					check(TypeIndices[TypeIndex] != INDEX_NONE);
+					if (TypeIndices[TypeIndex] == INDEX_NONE)
+					{
+						UE_LOG(LogRigVM, Error, TEXT("Cannot register type '%s' in dispatch '%s' for VM '%s'\nCannot localize registry."), *Type.CPPType.ToString(), *Factory->GetFactoryName().ToString(), *GetPathName());
+						return false;
+					}
 
 					const FRigVMTemplateArgumentType TemporaryType = TemporaryRegistry.GetType_NoLock(TypeIndices[TypeIndex]);
-					check(Type == TemporaryType);
+					if(Type != TemporaryType)
+					{
+						UE_LOG(LogRigVM, Error, TEXT("Registered type '%s' in dispatch '%s' does not match ('%s') for VM '%s'\nCannot localize registry."), *Type.CPPType.ToString(), *Factory->GetFactoryName().ToString(), *TemporaryType.CPPType.ToString(), *GetPathName());
+						return false;
+					}
 				}
 				TemporaryArgumentInfos.Add(FRigVMTemplateArgumentInfo::MakeFlattenedFromTypeIndices(ArgumentInfo.Name, ArgumentInfo.Direction, TypeIndices));
 			}
