Control Rig: Fix nodes disappearing when asset is duplicated
#rnx

diff --git a/Engine/Plugins/Runtime/RigVM/Source/RigVMDeveloper/Private/RigVMBlueprint.cpp b/Engine/Plugins/Runtime/RigVM/Source/RigVMDeveloper/Private/RigVMBlueprint.cpp
--- a/Engine/Plugins/Runtime/RigVM/Source/RigVMDeveloper/Private/RigVMBlueprint.cpp
+++ b/Engine/Plugins/Runtime/RigVM/Source/RigVMDeveloper/Private/RigVMBlueprint.cpp
@@ -2652,6 +2652,16 @@ void URigVMBlueprint::PostDuplicate(bool bDuplicateForPIE)
 
 	SplitAssetVariant();
 
+	TArray<UEdGraph*> EdGraphs;
+	GetAllGraphs(EdGraphs);
+	for (UEdGraph* EdGraph : EdGraphs)
+	{
+		if (URigVMEdGraph* RigVMEdGraph = Cast<URigVMEdGraph>(EdGraph))
+		{
+			RigVMEdGraph->CachedModelGraph.Reset();
+		}
+	}
+
 	FBlueprintEditorUtils::MarkBlueprintAsStructurallyModified(this);
 
 	// now that the data in the new BP is correct, make sure there is a call to PostLoad()
