Control Rig: Fix outdated cache in RigVMEdGraph
#rnx

diff --git a/Engine/Plugins/Runtime/RigVM/Source/RigVMDeveloper/Private/RigVMBlueprint.cpp b/Engine/Plugins/Runtime/RigVM/Source/RigVMDeveloper/Private/RigVMBlueprint.cpp
--- a/Engine/Plugins/Runtime/RigVM/Source/RigVMDeveloper/Private/RigVMBlueprint.cpp
+++ b/Engine/Plugins/Runtime/RigVM/Source/RigVMDeveloper/Private/RigVMBlueprint.cpp
@@ -844,6 +844,10 @@ void URigVMBlueprint::Serialize(FArchive& Ar)
 		for (UEdGraph* EdGraph : EdGraphs)
 		{
 			EdGraph->Schema = GetRigVMEdGraphSchemaClass();
+			if (URigVMEdGraph* RigVMEdGraph = Cast<URigVMEdGraph>(EdGraph))
+			{
+				RigVMEdGraph->CachedModelGraph.Reset();
+			}
 		}
 
 		if (Ar.CustomVer(FRigVMObjectVersion::GUID) < FRigVMObjectVersion::AddVariantToRigVMAssets)
