Control Rig: Fix public functions not available in rigs
#rnx

diff --git a/Engine/Plugins/Runtime/RigVM/Source/RigVMEditor/Private/Editor/RigVMEditor.cpp b/Engine/Plugins/Runtime/RigVM/Source/RigVMEditor/Private/Editor/RigVMEditor.cpp
--- a/Engine/Plugins/Runtime/RigVM/Source/RigVMEditor/Private/Editor/RigVMEditor.cpp
+++ b/Engine/Plugins/Runtime/RigVM/Source/RigVMEditor/Private/Editor/RigVMEditor.cpp
@@ -214,16 +214,13 @@ void FRigVMEditorBase::InitRigVMEditor(const EToolkitMode::Type Mode, const TSha
 
 	CommonInitialization(Blueprints, false);
 
-	// If the class actions have not been populated, refresh them
+	// Refresh the class actions
 	{
 		UClass* ActionKey = InRigVMBlueprint->GetClass();
 		FBlueprintActionDatabase& ActionDatabase = FBlueprintActionDatabase::Get();
 		FBlueprintActionDatabase::FActionRegistry const& ActionRegistry = ActionDatabase.GetAllActions();
-		if (!ActionRegistry.Contains(ActionKey) ||
-			ActionRegistry.FindChecked(ActionKey).IsEmpty())
-		{
-			ActionDatabase.RefreshClassActions(ActionKey);
-		}
+		ActionDatabase.ClearAssetActions(ActionKey);
+		ActionDatabase.RefreshClassActions(ActionKey);
 	}
 	
 	// user-defined-struct can change even after load
