Only remove VariableNode class nodes (ignore derived)

diff --git a/Engine/Plugins/Runtime/RigVM/Source/RigVMDeveloper/Private/RigVMModel/RigVMController.cpp b/Engine/Plugins/Runtime/RigVM/Source/RigVMDeveloper/Private/RigVMModel/RigVMController.cpp
--- a/Engine/Plugins/Runtime/RigVM/Source/RigVMDeveloper/Private/RigVMModel/RigVMController.cpp
+++ b/Engine/Plugins/Runtime/RigVM/Source/RigVMDeveloper/Private/RigVMModel/RigVMController.cpp
@@ -2208,7 +2208,7 @@ void URigVMController::OnExternalVariableRemoved(const FName& InVarName, bool bS
 	TArray<URigVMNode*> Nodes = Graph->GetNodes();
 	for (URigVMNode* Node : Nodes)
 	{
-		if (URigVMVariableNode* VariableNode = Cast<URigVMVariableNode>(Node))
+		if (URigVMVariableNode* VariableNode = ExactCast<URigVMVariableNode>(Node))
 		{
 			if (URigVMPin* VariablePin = VariableNode->FindPin(URigVMVariableNode::VariableName))
 			{
