Control Rig: Fix duplicate array pin crash
#rnx

diff --git a/Engine/Plugins/Runtime/RigVM/Source/RigVMDeveloper/Private/RigVMModel/RigVMController.cpp b/Engine/Plugins/Runtime/RigVM/Source/RigVMDeveloper/Private/RigVMModel/RigVMController.cpp
--- a/Engine/Plugins/Runtime/RigVM/Source/RigVMDeveloper/Private/RigVMModel/RigVMController.cpp
+++ b/Engine/Plugins/Runtime/RigVM/Source/RigVMDeveloper/Private/RigVMModel/RigVMController.cpp
@@ -11715,6 +11715,8 @@ URigVMPin* URigVMController::InsertArrayPin(URigVMPin* ArrayPin, int32 InIndex,
 	{
 		URigVMPin* ExistingPin = ArrayPin->GetSubPins()[ExistingIndex];
 		RenameObject(ExistingPin, *FString::FormatAsNumber(ExistingIndex + 1));
+		ExistingPin->CachedPinPath.ResetCachedValue();
+		ExistingPin->CachedPinPathWithNodePath.ResetCachedValue();
 
 		// temporarily force a the index based display name
 		// so that the ed graph node can pick that up during handling the notify
@@ -11856,6 +11858,8 @@ bool URigVMController::RemoveArrayPin(const FString& InArrayElementPinPath, bool
 	{
 		URigVMPin* ExistingPin = ArrayPin->GetSubPins()[ExistingIndex];
 		ExistingPin->SetNameFromIndex();
+		ExistingPin->CachedPinPath.ResetCachedValue();
+		ExistingPin->CachedPinPathWithNodePath.ResetCachedValue();
 		Notify(ERigVMGraphNotifType::PinRenamed, ExistingPin);
 	}
 
@@ -13684,7 +13688,9 @@ bool URigVMController::RenameExposedPin(const FName& InOldPinName, const FName&
 			{
 				return false;
 			}
-			
+
+			InPin->CachedPinPath.ResetCachedValue();
+			InPin->CachedPinPathWithNodePath.ResetCachedValue();
 			InPin->SetDisplayName(InNewName);
 			InPin->IncrementVersion();
 			
diff --git a/Engine/Plugins/Runtime/RigVM/Source/RigVMEditor/Private/EdGraph/RigVMEdGraphConnectionDrawingPolicy.cpp b/Engine/Plugins/Runtime/RigVM/Source/RigVMEditor/Private/EdGraph/RigVMEdGraphConnectionDrawingPolicy.cpp
--- a/Engine/Plugins/Runtime/RigVM/Source/RigVMEditor/Private/EdGraph/RigVMEdGraphConnectionDrawingPolicy.cpp
+++ b/Engine/Plugins/Runtime/RigVM/Source/RigVMEditor/Private/EdGraph/RigVMEdGraphConnectionDrawingPolicy.cpp
@@ -57,18 +57,21 @@ void FRigVMEdGraphConnectionDrawingPolicy::BuildPinToPinWidgetMap(TMap<TSharedRe
 		{
 			static void AddSubPins_Recursive(UEdGraphPin* PinObj, TMap<UEdGraphPin*, TSharedPtr<SGraphPin>>& InPinToPinWidgetMap, TSharedPtr<SGraphPin>& InGraphPinWidget)
 			{
-				for(UEdGraphPin* SubPin : PinObj->SubPins)
+				if (PinObj)
 				{
-					// Only add to the pin-to-pin widget map if the sub-pin widget is not there already
-					TSharedPtr<SGraphPin>* SubPinWidgetPtr = InPinToPinWidgetMap.Find(SubPin);
-					if(SubPinWidgetPtr == nullptr)
+					for(UEdGraphPin* SubPin : PinObj->SubPins)
 					{
-						SubPinWidgetPtr = &InGraphPinWidget;
-					}
+						// Only add to the pin-to-pin widget map if the sub-pin widget is not there already
+						TSharedPtr<SGraphPin>* SubPinWidgetPtr = InPinToPinWidgetMap.Find(SubPin);
+						if(SubPinWidgetPtr == nullptr)
+						{
+							SubPinWidgetPtr = &InGraphPinWidget;
+						}
 
-					TSharedPtr<SGraphPin> PinWidgetPtr = *SubPinWidgetPtr;
-					InPinToPinWidgetMap.Add(SubPin, PinWidgetPtr);
-					AddSubPins_Recursive(SubPin, InPinToPinWidgetMap, PinWidgetPtr);
+						TSharedPtr<SGraphPin> PinWidgetPtr = *SubPinWidgetPtr;
+						InPinToPinWidgetMap.Add(SubPin, PinWidgetPtr);
+						AddSubPins_Recursive(SubPin, InPinToPinWidgetMap, PinWidgetPtr);
+					}
 				}
 			}
 		};
