RigVM: Change FConditionalScopeLock to allow an already held write lock to be held again as a read lock.

- We cannot rely on the underlying system locking primitive to support this, so we check if the current thread is already holding the write lock, and in that case, increment the lock count but don't actually call the system primitive to takea  read lock.

#rnx

diff --git a/Engine/Plugins/Runtime/RigVM/Source/RigVM/Public/RigVMCore/RigVMRegistry.h b/Engine/Plugins/Runtime/RigVM/Source/RigVM/Public/RigVMCore/RigVMRegistry.h
--- a/Engine/Plugins/Runtime/RigVM/Source/RigVM/Public/RigVMCore/RigVMRegistry.h
+++ b/Engine/Plugins/Runtime/RigVM/Source/RigVM/Public/RigVMCore/RigVMRegistry.h
@@ -538,11 +538,16 @@ protected:
 			// First, take the actual underlying lock.
 			if (DesiredLockType == LockType_Read)
 			{
-				Registry->Lock.ReadLock();
+				// Only attempt to lock if we're not already holding the write lock.
+				if (Registry->WriteLockOwnerThread != FPlatformTLS::GetCurrentThreadId())
+				{
+					Registry->Lock.ReadLock();
+				}
 			}
 			else if (DesiredLockType == LockType_Write)
 			{
 				Registry->Lock.WriteLock();
+				Registry->WriteLockOwnerThread = FPlatformTLS::GetCurrentThreadId();
 			}
 
 			// Increment the matching lock count within the RigVM registry.
@@ -609,11 +614,17 @@ protected:
 			if (DesiredLockType == LockType_Read)
 			{
 				ensure(CurrentLockCount >= 0);
-				Registry->Lock.ReadUnlock();
+				
+				// If we were already holding the write lock, we didn't apply a subsequent read lock.
+				if (Registry->WriteLockOwnerThread != FPlatformTLS::GetCurrentThreadId())
+				{
+					Registry->Lock.ReadUnlock();
+				}
 			}
 			else if (DesiredLockType == LockType_Write)
 			{
 				ensure(CurrentLockCount == 0);
+				Registry->WriteLockOwnerThread = static_cast<uint32>(-1);
 				Registry->Lock.WriteUnlock();
 			}
 		}
@@ -1484,6 +1495,8 @@ private:
 	mutable FTransactionallySafeRWLock Lock;
 	mutable std::atomic<ELockType> LockType = LockType_Invalid;
 	mutable std::atomic<int32> LockCount = 0;
+	mutable std::atomic<uint32> WriteLockOwnerThread = static_cast<uint32>(-1);
+	
 
 	// Notifies other system that types have been added/removed, and template permutations have been updated
 	FOnRigVMRegistryChanged OnRigVMRegistryChangedDelegate;
