RigVM: Adding IRigVMAssetInterface as parent of RigVMBlueprint. First step to separation of RigVM and Blueprint.

diff --git a/Engine/Plugins/Runtime/RigVM/Source/RigVMDeveloper/Private/EdGraph/RigVMEdGraphSchema.cpp b/Engine/Plugins/Runtime/RigVM/Source/RigVMDeveloper/Private/EdGraph/RigVMEdGraphSchema.cpp
--- a/Engine/Plugins/Runtime/RigVM/Source/RigVMDeveloper/Private/EdGraph/RigVMEdGraphSchema.cpp
+++ b/Engine/Plugins/Runtime/RigVM/Source/RigVMDeveloper/Private/EdGraph/RigVMEdGraphSchema.cpp
@@ -12,7 +12,7 @@
 #include "Framework/MultiBox/MultiBoxBuilder.h"
 #include "GraphEditorActions.h"
 #include "RigVMHost.h"
-#include "RigVMBlueprint.h"
+#include "RigVMBlueprintLegacy.h"
 #include "Framework/Application/SlateApplication.h"
 #include "Framework/Notifications/NotificationManager.h"
 #include "RigVMStringUtils.h"
@@ -28,12 +28,15 @@
 #include "RigVMFunctions/RigVMDispatch_Constant.h"
 #include "RigVMFunctions/RigVMDispatch_MakeStruct.h"
 #include "RigVMModel/Nodes/RigVMDispatchNode.h"
+#include "RigVMAsset.h"
 #if WITH_RIGVMLEGACYEDITOR
 #include "BlueprintEditor.h"
 #else
 #include "Editor/RigVMNewEditor.h"
 #endif
 
+#include "RigVMBlueprintUtils.h"
+
 #include UE_INLINE_GENERATED_CPP_BY_NAME(RigVMEdGraphSchema)
 
 #if WITH_EDITOR
@@ -56,10 +59,10 @@ const FText FRigVMSearchTags::FiB_Glyph = LOCTEXT("Glyph", "Glyph");
 const FText FRigVMSearchTags::FiB_GlyphStyleSet = LOCTEXT("GlyphStyleSet", "GlyphStyleSet");
 const FText FRigVMSearchTags::FiB_GlyphColor = LOCTEXT("GlyphColor", "GlyphColor");
 
-FRigVMLocalVariableNameValidator::FRigVMLocalVariableNameValidator(const UBlueprint* Blueprint, const URigVMGraph* Graph, FName InExistingName)
+FRigVMLocalVariableNameValidator::FRigVMLocalVariableNameValidator(const TRigVMAssetInterfacePtr InAsset, const URigVMGraph* Graph, FName InExistingName)
 	: FStringSetNameValidator(InExistingName.ToString())
 {
-	if (Blueprint)
+	if (const UBlueprint* Blueprint = Cast<UBlueprint>(InAsset.GetObject()))
 	{
 		TSet<FName> NamesTemp;
 		// We allow local variables with same name as blueprint variable
@@ -74,6 +77,10 @@ FRigVMLocalVariableNameValidator::FRigVMLocalVariableNameValidator(const UBluepr
 			Names.Add(Name.ToString());
 		}
 	}
+	else if (InAsset)
+	{
+		check(false);
+	}
 
 	if (Graph)
 	{
@@ -134,6 +141,38 @@ FRigVMNameValidator::FRigVMNameValidator(const UBlueprint* Blueprint, const UStr
 	}
 }
 
+FRigVMNameValidator::FRigVMNameValidator(const TRigVMAssetInterfacePtr Blueprint, const UStruct* ValidationScope, FName InExistingName)
+: FStringSetNameValidator(InExistingName.ToString())
+{
+	if (Blueprint)
+	{
+		TSet<FName> NamesTemp;
+		TArray<FRigVMGraphVariableDescription> Variables = Blueprint->GetAssetVariables();
+		TArray<URigVMGraph*> Graphs = Blueprint->GetRigVMClientHost()->GetRigVMClient()->GetAllModels(true, true);
+
+		for (const FRigVMGraphVariableDescription& Variable : Variables)
+		{
+			NamesTemp.Add(Variable.Name);
+		}
+		for (const URigVMGraph* Graph : Graphs)
+		{
+			if (URigVMCollapseNode* Node = Cast<URigVMCollapseNode>(Graph->GetOuter()))
+			{
+				NamesTemp.Add(Node->GetFName());
+			}
+			else
+			{
+				NamesTemp.Add(Graph->GetFName());
+			}
+		}
+
+		for (FName& Name : NamesTemp)
+		{
+			Names.Add(Name.ToString());
+		}
+	}
+}
+
 EValidatorResult FRigVMNameValidator::IsValid(const FString& Name, bool bOriginal)
 {
 	const EValidatorResult Result = FStringSetNameValidator::IsValid(Name, bOriginal);
@@ -211,7 +250,7 @@ bool FRigVMEdGraphSchemaAction_LocalVar::IsValidName(const FName& NewName, FText
 {
 	if (const URigVMEdGraph* EdGraphGraph = Cast<URigVMEdGraph>(GetVariableScope()))
 	{
-		FRigVMLocalVariableNameValidator NameValidator(EdGraphGraph->GetBlueprint(), EdGraphGraph->GetModel(), GetVariableName());
+		FRigVMLocalVariableNameValidator NameValidator(EdGraphGraph->GetAsset(), EdGraphGraph->GetModel(), GetVariableName());
 		const EValidatorResult Result = NameValidator.IsValid(NewName.ToString(), false);
 		if (Result != EValidatorResult::Ok && Result != EValidatorResult::ExistingName)
 		{
@@ -277,7 +316,7 @@ UEdGraphNode* FRigVMEdGraphSchemaAction_PromoteToVariable::PerformAction(UEdGrap
 		return nullptr;
 	}
 
-	URigVMBlueprint* Blueprint = RigGraph->GetBlueprint();
+	TRigVMAssetInterfacePtr Blueprint = RigGraph->GetAsset();
 	URigVMGraph* Model = RigGraph->GetModel();
 	URigVMController* Controller = RigGraph->GetController();
 	if((Blueprint == nullptr) ||
@@ -319,7 +358,7 @@ UEdGraphNode* FRigVMEdGraphSchemaAction_PromoteToVariable::PerformAction(UEdGrap
 	}
 	else
 	{
-		Blueprint->Modify();
+		Blueprint->GetObject()->Modify();
 
 		FString DefaultValue = ModelPin->GetDefaultValue();
 		if(!DefaultValue.IsEmpty())
@@ -493,9 +532,9 @@ FRigVMEdGraphSchemaAction_Event::FRigVMEdGraphSchemaAction_Event(const FName& In
 
 FReply FRigVMEdGraphSchemaAction_Event::OnDoubleClick(UBlueprint* InBlueprint)
 {
-	if (URigVMBlueprint* Blueprint = Cast<URigVMBlueprint>(InBlueprint))
+	if (TRigVMAssetInterfacePtr Blueprint = InBlueprint)
 	{
-		if(const URigVMNode* ModelNode = Blueprint->GetRigVMClient()->FindNode(NodePath))
+		if(const URigVMNode* ModelNode = Blueprint->GetRigVMClientHost()->GetRigVMClient()->FindNode(NodePath))
 		{
 			Blueprint->OnRequestJumpToHyperlink().Execute(ModelNode);
 			return FReply::Handled();
@@ -571,7 +610,7 @@ FReply FRigVMFunctionDragDropAction::DroppedOnPanel(const TSharedRef< class SWid
 	// For functions
 	else if (URigVMEdGraph* TargetRigGraph = Cast<URigVMEdGraph>(&Graph))
 	{
-		if (URigVMBlueprint* TargetRigBlueprint = Cast<URigVMBlueprint>(FBlueprintEditorUtils::FindBlueprintForGraph(TargetRigGraph)))
+		if (TRigVMAssetInterfacePtr TargetRigBlueprint = FRigVMBlueprintUtils::FindAssetForGraph(TargetRigGraph))
 		{
 			if (URigVMGraph* FunctionDefinitionGraph = SourceRigBlueprint->GetModel(SourceRigGraph))
 			{
@@ -581,12 +620,12 @@ FReply FRigVMFunctionDragDropAction::DroppedOnPanel(const TSharedRef< class SWid
 					{
 						if(URigVMFunctionLibrary* FunctionLibrary = Cast<URigVMFunctionLibrary>(FunctionDefinitionNode->GetOuter()))
 						{
-							if(URigVMBlueprint* FunctionRigBlueprint = Cast<URigVMBlueprint>(FunctionLibrary->GetOuter()))
+							if(TRigVMAssetInterfacePtr FunctionRigBlueprint = IRigVMAssetInterface::GetInterfaceOuter(FunctionLibrary))
 							{
 #if WITH_EDITOR
 								if(FunctionRigBlueprint != TargetRigBlueprint)
 								{
-									if(!FunctionRigBlueprint->IsFunctionPublic(FunctionDefinitionNode->GetFName()))
+									if(!FunctionRigBlueprint->GetRigVMClientHost()->IsFunctionPublic(FunctionDefinitionNode->GetFName()))
 									{
 										FRigVMGraphFunctionIdentifier FunctionIdentifier = FunctionDefinitionNode->GetFunctionIdentifier();
 										TargetRigBlueprint->BroadcastRequestLocalizeFunctionDialog(FunctionIdentifier);
@@ -649,7 +688,7 @@ FRigVMFunctionDragDropAction::FRigVMFunctionDragDropAction()
 {
 }
 
-TSharedRef<FRigVMFunctionDragDropAction> FRigVMFunctionDragDropAction::New(TSharedPtr<FEdGraphSchemaAction> InAction, URigVMBlueprint* InRigBlueprint, URigVMEdGraph* InRigGraph)
+TSharedRef<FRigVMFunctionDragDropAction> FRigVMFunctionDragDropAction::New(TSharedPtr<FEdGraphSchemaAction> InAction, TRigVMAssetInterfacePtr InRigBlueprint, URigVMEdGraph* InRigGraph)
 {
 	TSharedRef<FRigVMFunctionDragDropAction> Action = MakeShareable(new FRigVMFunctionDragDropAction);
 	Action->SourceAction = InAction;
@@ -989,6 +1028,11 @@ TSharedPtr<INameValidatorInterface> URigVMEdGraphSchema::GetNameValidator(const
 	return MakeShareable(new FRigVMNameValidator(BlueprintObj, ValidationScope, OriginalName));		
 }
 
+TSharedPtr<INameValidatorInterface> URigVMEdGraphSchema::GetNameValidator(const TRigVMAssetInterfacePtr Asset, const FName& OriginalName, const UStruct* ValidationScope) const
+{
+	return MakeShareable(new FRigVMNameValidator(Asset, ValidationScope, OriginalName));
+}
+
 bool URigVMEdGraphSchema::CanPromotePinToVariable(const UEdGraphPin& Pin, bool bCond) const
 {
 	return true;
@@ -1245,7 +1289,7 @@ FReply URigVMEdGraphSchema::BeginGraphDragAction(TSharedPtr<FEdGraphSchemaAction
 		FEdGraphSchemaAction_K2Graph* FuncAction = (FEdGraphSchemaAction_K2Graph*)InAction.Get();
 		if (URigVMEdGraph* RigGraph = Cast<URigVMEdGraph>(FuncAction->EdGraph))
 		{
-			if (URigVMBlueprint* RigBlueprint = Cast<URigVMBlueprint>(FBlueprintEditorUtils::FindBlueprintForGraph(RigGraph)))
+			if (TRigVMAssetInterfacePtr RigBlueprint = FRigVMBlueprintUtils::FindAssetForGraph(RigGraph))
 			{
 				TSharedRef<FRigVMFunctionDragDropAction> Action = FRigVMFunctionDragDropAction::New(InAction, RigBlueprint, RigGraph);
 				Action->SetAltDrag(MouseEvent.IsAltDown());
@@ -1259,7 +1303,7 @@ FReply URigVMEdGraphSchema::BeginGraphDragAction(TSharedPtr<FEdGraphSchemaAction
 		FRigVMEdGraphSchemaAction_LocalVar* VarAction = (FRigVMEdGraphSchemaAction_LocalVar*)InAction.Get();
 		if (URigVMEdGraph* RigGraph = Cast<URigVMEdGraph>(VarAction->GetVariableScope()))
 		{
-			if (URigVMBlueprint* RigBlueprint = Cast<URigVMBlueprint>(FBlueprintEditorUtils::FindBlueprintForGraph(RigGraph)))
+			if (TRigVMAssetInterfacePtr RigBlueprint = FRigVMBlueprintUtils::FindAssetForGraph(RigGraph))
 			{
 				TSharedRef<FRigVMFunctionDragDropAction> Action = FRigVMFunctionDragDropAction::New(InAction, RigBlueprint, RigGraph);
 				Action->SetAltDrag(MouseEvent.IsAltDown());
@@ -1274,7 +1318,7 @@ FReply URigVMEdGraphSchema::BeginGraphDragAction(TSharedPtr<FEdGraphSchemaAction
 FConnectionDrawingPolicy* URigVMEdGraphSchema::CreateConnectionDrawingPolicy(int32 InBackLayerID, int32 InFrontLayerID, float InZoomFactor, const FSlateRect& InClippingRect, class FSlateWindowElementList& InDrawElements, class UEdGraph* InGraphObj) const
 {
 #if WITH_EDITOR
-	if (const URigVMBlueprint* Blueprint = Cast<URigVMBlueprint>(FBlueprintEditorUtils::FindBlueprintForGraph(InGraphObj)))
+	if (const TRigVMAssetInterfacePtr Blueprint = FRigVMBlueprintUtils::FindAssetForGraph(InGraphObj))
 	{
 		return Blueprint->GetEditorModule()->CreateConnectionDrawingPolicy(InBackLayerID, InFrontLayerID, InZoomFactor, InClippingRect, InDrawElements, InGraphObj);
 	}	
@@ -1525,7 +1569,7 @@ EGraphType URigVMEdGraphSchema::GetGraphType(const UEdGraph* TestEdGraph) const
 {
 	if (const URigVMEdGraph* RigGraph = Cast<URigVMEdGraph>((UEdGraph*)TestEdGraph))
 	{
-		if (Cast<URigVMBlueprint>(FBlueprintEditorUtils::FindBlueprintForGraph(RigGraph)))
+		if (FRigVMBlueprintUtils::FindAssetForGraph(RigGraph))
 		{
 			if (const URigVMGraph* Model = RigGraph->GetModel())
 			{
@@ -1557,7 +1601,7 @@ FReply URigVMEdGraphSchema::TrySetGraphCategory(const UEdGraph* InGraph, const F
 {
 	if (const URigVMEdGraph* RigGraph = Cast<URigVMEdGraph>((UEdGraph*)InGraph))
 	{
-		if (URigVMBlueprint* RigBlueprint = Cast<URigVMBlueprint>(FBlueprintEditorUtils::FindBlueprintForGraph(RigGraph)))
+		if (const TRigVMAssetInterfacePtr RigBlueprint = FRigVMBlueprintUtils::FindAssetForGraph(RigGraph))
 		{
 			if (const URigVMGraph* Model = RigGraph->GetModel())
 			{
@@ -1581,7 +1625,7 @@ bool URigVMEdGraphSchema::TryDeleteGraph(UEdGraph* GraphToDelete) const
 {
 	if (const URigVMEdGraph* RigGraph = Cast<URigVMEdGraph>(GraphToDelete))
 	{
-		if (URigVMBlueprint* RigBlueprint = Cast<URigVMBlueprint>(FBlueprintEditorUtils::FindBlueprintForGraph(RigGraph)))
+		if (const TRigVMAssetInterfacePtr RigBlueprint = FRigVMBlueprintUtils::FindAssetForGraph(RigGraph))
 		{
 			if (const URigVMGraph* Model = RigBlueprint->GetModel(GraphToDelete))
 			{
@@ -1626,9 +1670,9 @@ bool URigVMEdGraphSchema::TryDeleteGraph(UEdGraph* GraphToDelete) const
 								{
 									if (Reference.IsValid())
 									{
-										const URigVMBlueprint* OtherBlueprint = Reference->GetTypedOuter<URigVMBlueprint>(); 
+										const TRigVMAssetInterfacePtr OtherBlueprint = IRigVMAssetInterface::GetInterfaceOuter(Reference.Get());
 										if (OtherBlueprint != RigBlueprint)
-										{											
+										{
 											if(RigBlueprint->OnRequestBulkEditDialog().IsBound())
 											{
 												URigVMController* FunctionController = RigBlueprint->GetController(LibraryNode->GetContainedGraph());
@@ -1743,7 +1787,7 @@ bool URigVMEdGraphSchema::TryToGetChildEvents(const UEdGraph* Graph, const int32
 
 UEdGraphPin* URigVMEdGraphSchema::DropPinOnNode(UEdGraphNode* InTargetNode, const FName& InSourcePinName, const FEdGraphPinType& InSourcePinType, EEdGraphPinDirection InSourcePinDirection) const
 {
-	if (URigVMBlueprint* RigBlueprint = Cast<URigVMBlueprint>(FBlueprintEditorUtils::FindBlueprintForNode(InTargetNode)))
+	if (TRigVMAssetInterfacePtr RigBlueprint = FRigVMBlueprintUtils::FindAssetForNode(InTargetNode))
 	{
 		if (URigVMEdGraphNode* RigNode = Cast<URigVMEdGraphNode>(InTargetNode))
 		{
@@ -2101,19 +2145,19 @@ bool URigVMEdGraphSchema::CanVariableBeDropped(UEdGraph* InGraph, FProperty* InV
 
 bool URigVMEdGraphSchema::RequestFunctionDropOnPanel(UEdGraph* InGraph, const FRigVMGraphFunctionIdentifier& InFunction, const FVector2D& InDropPosition, const FVector2D& InScreenPosition) const
 {
-	if (URigVMBlueprint* TargetRigBlueprint = Cast<URigVMBlueprint>(FBlueprintEditorUtils::FindBlueprintForGraph(InGraph)))
+	if (TRigVMAssetInterfacePtr TargetRigBlueprint = FRigVMBlueprintUtils::FindAssetForGraph(InGraph))
 	{
 		if(URigVMController* TargetController = TargetRigBlueprint->GetController(InGraph))
 		{
 			FRigVMGraphFunctionHeader Header = FRigVMGraphFunctionHeader::FindGraphFunctionHeader(InFunction);
 			if (IRigVMClientHost* SourceClientHost = Cast<IRigVMClientHost>(InFunction.HostObject.ResolveObject()))
 			{
-				if (IRigVMGraphFunctionHost* SourceFunctionHost = SourceClientHost->GetRigVMGraphFunctionHost())
+				if (TScriptInterface<IRigVMGraphFunctionHost> SourceFunctionHost = SourceClientHost->GetRigVMGraphFunctionHost())
 				{
 					if (FRigVMGraphFunctionStore* SourceStore = SourceFunctionHost->GetRigVMGraphFunctionStore())
 					{
 #if WITH_EDITOR
-						if(InFunction.HostObject != TargetRigBlueprint->GetRigVMBlueprintGeneratedClass())
+						if(InFunction.HostObject != TargetRigBlueprint->GetRigVMClientHost()->GetRigVMGraphFunctionHost().GetObject())
 						{
 							if (!SourceStore->IsFunctionPublic(InFunction))
 							{
@@ -2138,8 +2182,7 @@ bool URigVMEdGraphSchema::RequestVariableDropOnPanel(UEdGraph* InGraph, FPropert
 #if WITH_EDITOR
 	if (CanVariableBeDropped(InGraph, InVariableToDrop))
 	{
-		UBlueprint* Blueprint = FBlueprintEditorUtils::FindBlueprintForGraph(InGraph);
-		URigVMBlueprint* RigBlueprint = Cast<URigVMBlueprint>(Blueprint);
+		TRigVMAssetInterfacePtr RigBlueprint = FRigVMBlueprintUtils::FindAssetForGraph(InGraph);
 		if (RigBlueprint != nullptr)
 		{
 			RigBlueprint->OnVariableDropped().Broadcast(InGraph, InVariableToDrop, FDeprecateSlateVector2D(InDropPosition), FDeprecateSlateVector2D(InScreenPosition));
@@ -2258,7 +2301,7 @@ void URigVMEdGraphSchema::EndGraphNodeInteraction(UEdGraphNode* InNode) const
 	
 	bool bMovedSomething = false;
 
-	FGuardSkipDirtyBlueprintStatus GuardDirtyBlueprintStatus(Graph->GetBlueprint(), true);
+	FGuardSkipDirtyBlueprintStatus GuardDirtyBlueprintStatus(Graph->GetAsset(), true);
 
 	Graph->GetController()->OpenUndoBracket(TEXT("Move Nodes"));
 
@@ -2506,12 +2549,9 @@ void URigVMEdGraphSchema::HandleModifiedEvent(ERigVMGraphNotifType InNotifType,
 
 TSubclassOf<URigVMEdGraphNode> URigVMEdGraphSchema::GetGraphNodeClass(const URigVMEdGraph* InGraph) const
 {
-	if (const UBlueprint* Blueprint = FBlueprintEditorUtils::FindBlueprintForGraph(InGraph))
+	if (TRigVMAssetInterfacePtr RigBlueprint = FRigVMBlueprintUtils::FindAssetForGraph(InGraph))
 	{
-		if (const URigVMBlueprint* RigBlueprint = CastChecked<URigVMBlueprint>(Blueprint))
-		{
-			return RigBlueprint->GetRigVMEdGraphNodeClass();
-		}
+		return RigBlueprint->GetRigVMClientHost()->GetRigVMEdGraphNodeClass();
 	}
 	return nullptr;
 }
