RigVM: Adding IRigVMAssetInterface as parent of RigVMBlueprint. First step to separation of RigVM and Blueprint.

diff --git a/Engine/Plugins/Runtime/RigVM/Source/RigVMDeveloper/Private/RigVMAsset.cpp b/Engine/Plugins/Runtime/RigVM/Source/RigVMDeveloper/Private/RigVMAsset.cpp
new file mode 100644
--- /dev/null
+++ b/Engine/Plugins/Runtime/RigVM/Source/RigVMDeveloper/Private/RigVMAsset.cpp
@@ -0,0 +1,4140 @@
+// Copyright Epic Games, Inc. All Rights Reserved.
+
+#include "RigVMAsset.h"
+
+#include "RigVMBlueprintLegacy.h"
+
+#include "RigVMBlueprintGeneratedClass.h"
+#include "EdGraph/EdGraph.h"
+#include "EdGraphNode_Comment.h"
+#include "BlueprintActionDatabaseRegistrar.h"
+#include "EdGraph/RigVMEdGraph.h"
+#include "EdGraph/RigVMEdGraphNode.h"
+#include "EdGraph/RigVMEdGraphSchema.h"
+#include "UObject/AssetRegistryTagsContext.h"
+#include "UObject/ObjectSaveContext.h"
+#include "UObject/UObjectGlobals.h"
+#include "RigVMObjectVersion.h"
+#include "BlueprintCompilationManager.h"
+#include "RigVMCompiler/RigVMCompiler.h"
+#include "RigVMCore/RigVMRegistry.h"
+#include "AssetRegistry/AssetRegistryModule.h"
+#include "RigVMPythonUtils.h"
+#include "RigVMTypeUtils.h"
+#include "Algo/Count.h"
+#include "RigVMModel/RigVMControllerActions.h"
+#include "RigVMModel/Nodes/RigVMAggregateNode.h"
+#include "RigVMModel/Nodes/RigVMDispatchNode.h"
+#include "UObject/UE5MainStreamObjectVersion.h"
+#include "Stats/StatsHierarchical.h"
+#include "UObject/UObjectIterator.h"
+
+#include UE_INLINE_GENERATED_CPP_BY_NAME(RigVMAsset)
+
+#if WITH_EDITOR
+#include "Kismet2/KismetDebugUtilities.h"
+#include "Kismet2/WatchedPin.h"
+#include "Kismet2/BlueprintEditorUtils.h"
+#include "RigVMBlueprintUtils.h"
+#include "Editor/UnrealEdEngine.h"
+#include "Editor/Transactor.h"
+#include "CookOnTheSide/CookOnTheFlyServer.h"
+#include "RigVMEditorModule.h"
+#include "ScopedTransaction.h"
+#if !WITH_RIGVMLEGACYEDITOR
+#include "RigVMEditor/Private/Editor/Kismet/RigVMBlueprintCompilationManager.h"
+#endif
+#endif//WITH_EDITOR
+
+#define LOCTEXT_NAMESPACE "RigVMAsset"
+
+TAutoConsoleVariable<bool> CVarRigVMEnablePreLoadFiltering(
+	TEXT("RigVM.EnablePreLoadFiltering"),
+	true,
+	TEXT("When true the RigVMGraphs will be skipped during preload to speed up load times."));
+
+TAutoConsoleVariable<bool> CVarRigVMEnablePostLoadHashing(
+	TEXT("RigVM.EnablePostLoadHashing"),
+	true,
+	TEXT("When true refreshing the RigVMGraphs will be skipped if the hash matches the serialized hash."));
+
+static TArray<UClass*> GetClassObjectsInPackage(UPackage* InPackage)
+{
+	TArray<UObject*> Objects;
+	GetObjectsWithOuter(InPackage, Objects, false);
+
+	TArray<UClass*> ClassObjects;
+	for (UObject* Object : Objects)
+	{
+		if (UClass* Class = Cast<UClass>(Object))
+		{
+			ClassObjects.Add(Class);
+		}
+	}
+
+	return ClassObjects;
+}
+
+FRigVMEdGraphDisplaySettings::FRigVMEdGraphDisplaySettings(): bShowNodeInstructionIndex(false)
+                                                            , bShowNodeRunCounts(false)
+                                                            , NodeRunLowerBound(1)
+                                                            , NodeRunLimit(256)
+                                                            , MinMicroSeconds(0.0)
+                                                            , MaxMicroSeconds(1.0)
+                                                            , TotalMicroSeconds(0.0)
+                                                            , AverageFrames(64)
+                                                            , bAutoDetermineRange(true)
+                                                            , LastMinMicroSeconds(0.0)
+                                                            , LastMaxMicroSeconds(1.0)
+                                                            , MinDurationColor(FLinearColor::Green)
+                                                            , MaxDurationColor(FLinearColor::Red)
+                                                            , TagDisplayMode(ERigVMTagDisplayMode::All)
+{
+}
+
+void FRigVMEdGraphDisplaySettings::SetTotalMicroSeconds(double InTotalMicroSeconds)
+{
+	TotalMicroSeconds = AggregateAverage(TotalMicroSecondsFrames, TotalMicroSeconds, InTotalMicroSeconds);
+}
+
+void FRigVMEdGraphDisplaySettings::SetLastMinMicroSeconds(double InMinMicroSeconds)
+{
+	LastMinMicroSeconds = AggregateAverage(MinMicroSecondsFrames, LastMinMicroSeconds, InMinMicroSeconds);
+}
+
+void FRigVMEdGraphDisplaySettings::SetLastMaxMicroSeconds(double InMaxMicroSeconds)
+{
+	LastMaxMicroSeconds = AggregateAverage(MaxMicroSecondsFrames, LastMaxMicroSeconds, InMaxMicroSeconds);
+}
+
+double FRigVMEdGraphDisplaySettings::AggregateAverage(TArray<double>& InFrames, double InPrevious, double InNext) const
+{
+	const int32 NbFrames = FMath::Min(AverageFrames, 256);
+	if(NbFrames < 2)
+	{
+		InFrames.Reset();
+		return InNext;
+	}
+	
+	InFrames.Add(InNext);
+	if(InFrames.Num() >= NbFrames)
+	{
+		double Average = 0;
+		for(const double Value : InFrames)
+		{
+			Average += Value;
+		}
+		Average /= double(NbFrames);
+		InFrames.Reset();
+		return Average;
+	}
+
+	if(InPrevious == DBL_MAX || InPrevious < -SMALL_NUMBER)
+	{
+		return InNext;
+	}
+	return InPrevious;
+}
+
+TRigVMAssetInterfacePtr IRigVMAssetInterface::GetInterfaceOuter(const UObject* InObject)
+{
+	if (!InObject)
+	{
+		return nullptr;
+	}
+	 
+	UObject* Outer = InObject->GetOuter();
+	while (Outer)
+	{
+		if (Outer->Implements<URigVMAssetInterface>())
+		{
+			return TRigVMAssetInterfacePtr(Outer);
+		}
+		Outer = Outer->GetOuter();
+	}
+	return nullptr;
+}
+
+FSoftObjectPath IRigVMAssetInterface::PreDuplicateAssetPath;
+FSoftObjectPath IRigVMAssetInterface::PreDuplicateHostPath;
+TArray<TRigVMAssetInterfacePtr> IRigVMAssetInterface::sCurrentlyOpenedRigVMBlueprints;
+#if WITH_EDITOR
+FCriticalSection IRigVMAssetInterface::QueuedCompilerMessageDelegatesMutex;
+TArray<FOnRigVMReportCompilerMessage::FDelegate> IRigVMAssetInterface::QueuedCompilerMessageDelegates;
+#endif
+
+IRigVMAssetInterface::IRigVMAssetInterface(const FObjectInitializer& ObjectInitializer)
+{
+	bSuspendModelNotificationsForSelf = false;
+	bSuspendAllNotifications = false;
+	bSuspendPythonMessagesForRigVMClient = true;
+	bMarkBlueprintAsStructurallyModifiedPending = false;
+
+	bAutoRecompileVM = true;
+	bVMRecompilationRequired = false;
+	bIsCompiling = false;
+	VMRecompilationBracket = 0;
+	bSkipDirtyBlueprintStatus = false;
+
+	bDirtyDuringLoad = false;
+	bErrorsDuringCompilation = false;
+
+#if WITH_EDITOR
+	TArray<FOnRigVMReportCompilerMessage::FDelegate> DelegatesForReportFromCompiler;
+	{
+		FScopeLock Lock(&QueuedCompilerMessageDelegatesMutex);
+		Swap(QueuedCompilerMessageDelegates, DelegatesForReportFromCompiler);
+	}
+
+	for(const FOnRigVMReportCompilerMessage::FDelegate& Delegate : DelegatesForReportFromCompiler)
+	{
+		ReportCompilerMessageEvent.Add(Delegate);
+	}
+#endif
+
+}
+
+void IRigVMAssetInterface::Clear()
+{
+	FCoreUObjectDelegates::OnEndLoadPackage.RemoveAll(this);
+	FRigVMRegistry::Get().OnRigVMRegistryChanged().RemoveAll(this);
+}
+
+TScriptInterface<IRigVMClientHost> IRigVMAssetInterface::GetRigVMClientHost()
+{
+	return GetObject();
+}
+
+void IRigVMAssetInterface::RequestClientHostRigVMInit()
+{
+	if (TScriptInterface<IRigVMClientHost> ClientHost = GetRigVMClientHost())
+	{
+		return ClientHost->RequestRigVMInit();
+	}
+}
+
+void IRigVMAssetInterface::CollectExternalDependencies(TArray<FRigVMExternalDependency>& OutDependencies, const FName& InCategory, const FRigVMClient* InClient) const
+{
+	check(InClient);
+
+	const TArray<URigVMGraph*> Graphs = InClient->GetAllModels(true, true);
+	for(const URigVMGraph* Graph : Graphs)
+	{
+		CollectExternalDependencies(OutDependencies, InCategory, Graph);
+	}
+}
+
+void IRigVMAssetInterface::CollectExternalDependencies(TArray<FRigVMExternalDependency>& OutDependencies, const FName& InCategory, const FRigVMGraphFunctionStore* InFunctionStore) const
+{
+	check(InFunctionStore);
+	for(const FRigVMGraphFunctionData& Function : InFunctionStore->PublicFunctions)
+	{
+		CollectExternalDependencies(OutDependencies, InCategory, &Function);
+	}
+	for(const FRigVMGraphFunctionData& Function : InFunctionStore->PrivateFunctions)
+	{
+		CollectExternalDependencies(OutDependencies, InCategory, &Function);
+	}
+}
+
+void IRigVMAssetInterface::CollectExternalDependencies(TArray<FRigVMExternalDependency>& OutDependencies, const FName& InCategory, const FRigVMGraphFunctionData* InFunction) const
+{
+	check(InFunction);
+	CollectExternalDependencies(OutDependencies, InCategory, &InFunction->Header);
+	CollectExternalDependencies(OutDependencies, InCategory, &InFunction->CompilationData);
+}
+
+void IRigVMAssetInterface::CollectExternalDependencies(TArray<FRigVMExternalDependency>& OutDependencies, const FName& InCategory, const FRigVMGraphFunctionHeader* InHeader) const
+{
+	check(InHeader);
+	if(InCategory == IRigVMExternalDependencyManager::RigVMGraphFunctionCategory)
+	{
+		OutDependencies.AddUnique({InHeader->LibraryPointer.GetLibraryNodePath(), InCategory});
+	}
+	for(const FRigVMGraphFunctionArgument& Argument : InHeader->Arguments)
+	{
+		CollectExternalDependenciesForCPPTypeObject(OutDependencies, InCategory, Argument.CPPTypeObject.Get());
+	}
+}
+
+void IRigVMAssetInterface::CollectExternalDependencies(TArray<FRigVMExternalDependency>& OutDependencies, const FName& InCategory, const FRigVMFunctionCompilationData* InFunction) const
+{
+	check(InFunction);
+
+	for(const FRigVMFunctionCompilationPropertyDescription& Property : InFunction->LiteralPropertyDescriptions)
+	{
+		CollectExternalDependenciesForCPPTypeObject(OutDependencies, InCategory, Property.CPPTypeObject.Get());
+	}
+	for(const FRigVMFunctionCompilationPropertyDescription& Property : InFunction->WorkPropertyDescriptions)
+	{
+		CollectExternalDependenciesForCPPTypeObject(OutDependencies, InCategory, Property.CPPTypeObject.Get());
+	}
+
+	const FRigVMRegistry& Registry = FRigVMRegistry::Get();
+	for(const FName& FunctionName : InFunction->FunctionNames)
+	{
+		if(const FRigVMFunction* Function = Registry.FindFunction(*FunctionName.ToString()))
+		{
+			for(const FRigVMFunctionArgument& Argument :  Function->Arguments)
+			{
+				const FRigVMTemplateArgumentType& ArgumentType = Registry.FindTypeFromCPPType(Argument.Type);
+				if(ArgumentType.IsValid())
+				{
+					CollectExternalDependenciesForCPPTypeObject(OutDependencies, InCategory, ArgumentType.CPPTypeObject.Get());
+				}
+			}
+		}
+	}
+}
+
+void IRigVMAssetInterface::CollectExternalDependencies(TArray<FRigVMExternalDependency>& OutDependencies, const FName& InCategory, const URigVMGraph* InGraph) const
+{
+	for(const URigVMNode* Node : InGraph->GetNodes())
+	{
+		CollectExternalDependencies(OutDependencies, InCategory, Node);
+	}
+
+	const TArray<FRigVMGraphVariableDescription> LocalVariables = InGraph->GetLocalVariables();
+	for(const FRigVMGraphVariableDescription& LocalVariable : LocalVariables)
+	{
+		CollectExternalDependenciesForCPPTypeObject(OutDependencies, InCategory, LocalVariable.CPPTypeObject.Get());
+	}
+}
+
+void IRigVMAssetInterface::CollectExternalDependencies(TArray<FRigVMExternalDependency>& OutDependencies, const FName& InCategory, const URigVMNode* InNode) const
+{
+	for(const URigVMPin* Pin : InNode->GetPins())
+	{
+		CollectExternalDependencies(OutDependencies, InCategory, Pin);
+	}
+	if(const URigVMFunctionReferenceNode* FunctionReferenceNode = Cast<URigVMFunctionReferenceNode>(InNode))
+	{
+		if(const FRigVMGraphFunctionData* Function = FunctionReferenceNode->GetReferencedFunctionData(true))
+		{
+			CollectExternalDependencies(OutDependencies, InCategory, Function);
+		}
+	}
+}
+
+void IRigVMAssetInterface::CollectExternalDependencies(TArray<FRigVMExternalDependency>& OutDependencies, const FName& InCategory, const URigVMPin* InPin) const
+{
+	CollectExternalDependenciesForCPPTypeObject(OutDependencies, InCategory, InPin->GetCPPTypeObject());
+	for(const URigVMPin* SubPin : InPin->GetSubPins())
+	{
+		CollectExternalDependencies(OutDependencies, InCategory, SubPin);
+	}
+}
+
+void IRigVMAssetInterface::CollectExternalDependencies(TArray<FRigVMExternalDependency>& OutDependencies, const FName& InCategory, const UStruct* InStruct) const
+{
+	check(InStruct);
+	if(InCategory == IRigVMExternalDependencyManager::UserDefinedStructCategory)
+	{
+		if(const UUserDefinedStruct* UserDefinedStruct = Cast<UUserDefinedStruct>(InStruct))
+		{
+			OutDependencies.AddUnique({UserDefinedStruct->GetPathName(), InCategory});
+		}
+	}
+	for (TFieldIterator<FProperty> PropertyIt(InStruct); PropertyIt; ++PropertyIt)
+	{
+		const FProperty* Property = *PropertyIt;
+		while(const FArrayProperty* ArrayProperty = CastField<FArrayProperty>(Property))
+		{
+			Property = ArrayProperty->Inner;
+		}
+		if(const FStructProperty* StructProperty = CastField<FStructProperty>(Property))
+		{
+			CollectExternalDependencies(OutDependencies, InCategory, StructProperty->Struct);
+		}
+		else if(const FEnumProperty* EnumProperty = CastField<FEnumProperty>(Property))
+		{
+			if(const UEnum* Enum = EnumProperty->GetEnum())
+			{
+				CollectExternalDependencies(OutDependencies, InCategory, Enum);
+			}
+		}
+		else if(const FByteProperty* ByteProperty = CastField<FByteProperty>(Property))
+		{
+			if(const UEnum* Enum = ByteProperty->Enum)
+			{
+				CollectExternalDependencies(OutDependencies, InCategory, Enum);
+			}
+		}
+	}
+}
+
+void IRigVMAssetInterface::CollectExternalDependencies(TArray<FRigVMExternalDependency>& OutDependencies, const FName& InCategory, const UEnum* InEnum) const
+{
+	check(InEnum);
+	if(InCategory == IRigVMExternalDependencyManager::UserDefinedEnumCategory)
+	{
+		if(const UUserDefinedEnum* UserDefinedEnum = Cast<UUserDefinedEnum>(InEnum))
+		{
+			OutDependencies.AddUnique({UserDefinedEnum->GetPathName(), InCategory});
+		}
+	}
+}
+
+void IRigVMAssetInterface::CollectExternalDependenciesForCPPTypeObject(TArray<FRigVMExternalDependency>& OutDependencies, const FName& InCategory, const UObject* InObject) const
+{
+	if(InObject)
+	{
+		if(const UEnum* Enum = Cast<UEnum>(InObject))
+		{
+			CollectExternalDependencies(OutDependencies, InCategory, Enum);
+		}
+		else if(const UStruct* Struct = Cast<UStruct>(InObject))
+		{
+			CollectExternalDependencies(OutDependencies, InCategory, Struct);
+		}
+	}
+}
+
+void IRigVMAssetInterface::CommonInitialization(const FObjectInitializer& ObjectInitializer)
+{
+	// guard against this running multiple times
+	TScriptInterface<IRigVMClientHost> Host(ObjectInitializer.GetObj());
+	FRigVMClient* RigVMClient = Host->GetRigVMClient();
+	check(RigVMClient->GetDefaultSchemaClass() == nullptr);
+
+	RigVMClient->SetDefaultSchemaClass(Host->GetRigVMSchemaClass());
+	RigVMClient->SetDefaultExecuteContextStruct(Host->GetRigVMExecuteContextStruct());
+
+	for(UEdGraph* UberGraph : GetUberGraphs())
+	{
+		if(URigVMEdGraph* EdGraph = Cast<URigVMEdGraph>(UberGraph))
+		{
+			EdGraph->Schema = Host->GetRigVMEdGraphSchemaClass();
+		}
+	}
+
+	RigVMClient->SetOuterClientHost(GetObject(), TEXT("RigVMClient"));
+	{
+		TGuardValue<bool> DisableClientNotifs(RigVMClient->bSuspendNotifications, true);
+		RigVMClient->GetOrCreateFunctionLibrary(false, &ObjectInitializer, false);
+		RigVMClient->AddModel(FRigVMClient::RigVMModelPrefix, false, &ObjectInitializer, false);
+	}
+
+	TObjectPtr<URigVMEdGraph>& FunctionLibraryEdGraph = GetFunctionLibraryEdGraph();
+ 	FunctionLibraryEdGraph = Cast<URigVMEdGraph>(ObjectInitializer.CreateDefaultSubobject(ObjectInitializer.GetObj(), TEXT("RigVMFunctionLibraryEdGraph"), Host->GetRigVMEdGraphClass(), Host->GetRigVMEdGraphClass(), true, true));
+	FunctionLibraryEdGraph->Schema = Host->GetRigVMEdGraphSchemaClass();
+	FunctionLibraryEdGraph->bAllowRenaming = 0;
+	FunctionLibraryEdGraph->bEditable = 0;
+	FunctionLibraryEdGraph->bAllowDeletion = 0;
+	FunctionLibraryEdGraph->bIsFunctionDefinition = false;
+	FunctionLibraryEdGraph->ModelNodePath = RigVMClient->GetFunctionLibrary()->GetNodePath();
+	FunctionLibraryEdGraph->InitializeFromAsset(ObjectInitializer.GetObj());
+}
+
+void IRigVMAssetInterface::InitializeModelIfRequired(bool bRecompileVM)
+{
+	DECLARE_SCOPE_HIERARCHICAL_COUNTER_FUNC()
+
+	FRigVMClient* RigVMClient = GetRigVMClientHost()->GetRigVMClient();
+	if (RigVMClient->GetController(0) == nullptr)
+	{
+		const TArray<URigVMGraph*> Models = GetRigVMClientHost()->GetRigVMClient()->GetAllModels(true, false);
+		for(const URigVMGraph* Model : Models)
+		{
+			RigVMClient->GetOrCreateController(Model);
+		}
+
+		bool bRecompileRequired = false;
+		TArray<TObjectPtr<UEdGraph>> UbergraphPages = GetUberGraphs();
+		for (int32 i = 0; i < UbergraphPages.Num(); ++i)
+		{
+			if (URigVMEdGraph* Graph = Cast<URigVMEdGraph>(UbergraphPages[i]))
+			{
+				if (bRecompileVM)
+				{
+					bRecompileRequired = true;
+				}
+
+				Graph->InitializeFromAsset(GetObject());
+			}
+		}
+
+		if(bRecompileRequired)
+		{
+			RecompileVM();
+		}
+
+		GetFunctionLibraryEdGraph()->InitializeFromAsset(GetObject());
+	}
+}
+
+bool IRigVMAssetInterface::ExportGraphToText(UEdGraph* InEdGraph, FString& OutText)
+{
+	OutText.Empty();
+
+	if (URigVMGraph* RigGraph = GetModel(InEdGraph))
+	{
+		if (URigVMCollapseNode* CollapseNode = Cast<URigVMCollapseNode>(RigGraph->GetOuter()))
+		{
+			if (URigVMController* Controller = GetOrCreateController(CollapseNode->GetGraph()))
+			{
+				TArray<FName> NodeNamesToExport;
+				NodeNamesToExport.Add(CollapseNode->GetFName());
+				OutText = Controller->ExportNodesToText(NodeNamesToExport);
+			}
+		}
+	}
+
+	// always return true so that the default mechanism doesn't take over
+	return true;
+}
+
+bool IRigVMAssetInterface::CanImportGraphFromText(const FString& InClipboardText)
+{
+	return GetTemplateController(true)->CanImportNodesFromText(InClipboardText);
+}
+
+bool IRigVMAssetInterface::RequiresForceLoadMembers(UObject* InObject) const
+{
+	// only filter if the console variable is enabled
+	if(!CVarRigVMEnablePreLoadFiltering->GetBool())
+	{
+		return RequiresForceLoadMembersSuper(InObject);
+	}
+
+	// we can stop traversing when hitting a URigVMNode
+	// except for collapse nodes - since they contain a graphs again
+	// and variable  nodes - since they are needed during preload by the BP compiler
+	if(InObject->IsA<URigVMNode>())
+	{
+		if(!InObject->IsA<URigVMCollapseNode>() &&
+			!InObject->IsA<URigVMVariableNode>())
+		{
+			return false;
+		}
+	}
+	return RequiresForceLoadMembersSuper(InObject);
+}
+
+void IRigVMAssetInterface::PostEditChangeChainProperty(FPropertyChangedChainEvent& PropertyChangedEvent)
+{
+	PostEditChangeChainPropertySuper(PropertyChangedEvent);
+	PostEditChangeChainPropertyEvent.Broadcast(PropertyChangedEvent);
+}
+
+void IRigVMAssetInterface::PostRename(UObject* OldOuter, const FName OldName)
+{
+	PostRenameSuper(OldOuter, OldName);
+
+	// Whenever the asset is renamed/moved, generated classes parented to the old package
+	// are not moved to the new package automatically (see FAssetRenameManager), so we
+	// have to manually perform the move/rename, to avoid invalid reference to the old package
+	
+	// Note: while asset duplication doesn't duplicate the classes either, it is not a problem there
+	// because we always recompile in post duplicate.
+	TArray<UClass*> ClassObjects = GetClassObjectsInPackage(OldOuter->GetPackage());
+
+	for (UClass* ClassObject : ClassObjects)
+	{
+		if (URigVMMemoryStorageGeneratorClass* MemoryClass = Cast<URigVMMemoryStorageGeneratorClass>(ClassObject))
+		{
+			MemoryClass->Rename(nullptr, GetPackage(), REN_DoNotDirty | REN_DontCreateRedirectors | REN_NonTransactional);
+		}
+	}
+
+	const FString OldAssetPath = FString::Printf(TEXT("%s.%s"), *OldOuter->GetPathName(), *OldName.ToString());
+	ReplaceFunctionIdentifiers(OldAssetPath, GetObject()->GetPathName());
+}
+
+void IRigVMAssetInterface::GetPreloadDependencies(TArray<UObject*>& OutDeps)
+{
+	GetPreloadDependenciesSuper(OutDeps);
+
+	TArray<UClass*> ClassObjects = GetClassObjectsInPackage(GetPackage());
+
+	for (UClass* ClassObject : ClassObjects)
+	{
+		if (URigVMMemoryStorageGeneratorClass* MemoryClass = Cast<URigVMMemoryStorageGeneratorClass>(ClassObject))
+		{
+			OutDeps.Add(MemoryClass);
+		}
+	}
+}
+
+UObject* IRigVMAssetInterface::GetEditorObjectForRigVMGraph(const URigVMGraph* InVMGraph) const
+{
+	if(InVMGraph)
+	{
+		if(InVMGraph->GetOutermost() != GetObject()->GetOutermost())
+		{
+			return nullptr;
+		}
+
+		if(InVMGraph->IsA<URigVMFunctionLibrary>())
+		{
+			return GetFunctionLibraryEdGraph();
+		}
+
+		TArray<UEdGraph*> EdGraphs;
+		GetAllEdGraphs(EdGraphs);
+
+		bool bIsFunctionDefinition = false;
+		if (URigVMLibraryNode* LibraryNode = Cast<URigVMLibraryNode>(InVMGraph->GetOuter()))
+		{
+			bIsFunctionDefinition = LibraryNode->GetGraph()->IsA<URigVMFunctionLibrary>();
+		}
+
+		for (UEdGraph* EdGraph : EdGraphs)
+		{
+			if (URigVMEdGraph* RigGraph = Cast<URigVMEdGraph>(EdGraph))
+			{
+				if (RigGraph->bIsFunctionDefinition != bIsFunctionDefinition)
+				{
+					continue;
+				}
+
+				if ((RigGraph->ModelNodePath == InVMGraph->GetNodePath()) ||
+					(RigGraph->ModelNodePath.IsEmpty() && (GetRigVMClientHost()->GetRigVMClient()->GetDefaultModel() == InVMGraph)))
+				{
+					return RigGraph;
+				}
+			}
+		}
+	}
+	
+	return nullptr;
+}
+
+URigVMGraph* IRigVMAssetInterface::GetRigVMGraphForEditorObject(UObject* InObject) const
+{
+	if(URigVMEdGraph* RigVMEdGraph = Cast<URigVMEdGraph>(InObject))
+	{
+		const FRigVMClient* RigVMClient = GetRigVMClientHost()->GetRigVMClient();
+		if (RigVMEdGraph->bIsFunctionDefinition)
+		{
+			if (URigVMLibraryNode* LibraryNode = RigVMClient->GetFunctionLibrary()->FindFunction(*RigVMEdGraph->ModelNodePath))
+			{
+				return LibraryNode->GetContainedGraph();
+			}
+		}
+		else
+		{
+			return RigVMClient->GetModel(RigVMEdGraph->ModelNodePath);
+		}
+	}
+
+	return nullptr;
+}
+
+void IRigVMAssetInterface::HandleRigVMGraphAdded(const FRigVMClient* InClient, const FString& InNodePath)
+{
+	if(URigVMGraph* Model = InClient->GetModel(InNodePath))
+	{
+		if(!GetObject()->HasAnyFlags(RF_ClassDefaultObject | RF_NeedInitialization | RF_NeedLoad | RF_NeedPostLoad) &&
+			GetObject()->GetOuter() != GetTransientPackage() &&
+			!GIsTransacting)
+		{
+			CreateEdGraph(Model, true);
+			RecompileVM();
+		}
+
+#if WITH_EDITOR
+		if(!bSuspendPythonMessagesForRigVMClient)
+		{
+			const FString BlueprintName = InClient->GetDefaultSchema()->GetSanitizedName(GetObject()->GetName(), true, false);
+			RigVMPythonUtils::Print(BlueprintName, 
+				FString::Printf(TEXT("blueprint.add_model('%s')"),
+					*Model->GetName()));
+		}
+#endif
+	}
+}
+
+void IRigVMAssetInterface::HandleRigVMGraphRemoved(const FRigVMClient* InClient, const FString& InNodePath)
+{
+	if(URigVMGraph* Model = InClient->GetModel(InNodePath))
+	{
+		RemoveEdGraph(Model);
+		RecompileVM();
+
+#if WITH_EDITOR
+		if(!bSuspendPythonMessagesForRigVMClient)
+		{
+			const FString BlueprintName = InClient->GetDefaultSchema()->GetSanitizedName(GetObject()->GetName(), true, false);
+			RigVMPythonUtils::Print(BlueprintName, 
+				FString::Printf(TEXT("blueprint.remove_model('%s')"),
+					*Model->GetName()));
+		}
+#endif
+	}
+}
+
+void IRigVMAssetInterface::HandleRigVMGraphRenamed(const FRigVMClient* InClient, const FString& InOldNodePath, const FString& InNewNodePath)
+{
+	if(InClient->GetModel(InNewNodePath))
+	{
+		TArray<UEdGraph*> EdGraphs;
+		GetAllEdGraphs(EdGraphs);
+
+		for (UEdGraph* EdGraph : EdGraphs)
+		{
+			if (URigVMEdGraph* RigGraph = Cast<URigVMEdGraph>(EdGraph))
+			{
+				RigGraph->HandleRigVMGraphRenamed(InOldNodePath, InNewNodePath);
+			}
+		}
+
+		MarkAssetAsModified();
+	}
+}
+
+void IRigVMAssetInterface::HandleConfigureRigVMController(const FRigVMClient* InClient, URigVMController* InControllerToConfigure)
+{
+	InControllerToConfigure->OnModified().AddRaw(this, &IRigVMAssetInterface::HandleModifiedEvent);
+
+	TWeakObjectPtr<UObject> WeakThis(GetObject());
+
+	// this delegate is used by the controller to determine variable validity
+	// during a bind process. the controller itself doesn't own the variables,
+	// so we need a delegate to request them from the owning blueprint
+	InControllerToConfigure->GetExternalVariablesDelegate.BindLambda([](URigVMGraph* InGraph) -> TArray<FRigVMExternalVariable> {
+
+		if (InGraph)
+		{
+			if(IRigVMAssetInterface* Asset = InGraph->GetImplementingOuter<IRigVMAssetInterface>())
+			{
+				return Asset->GetExternalVariables(true /* rely on variables within blueprint */);
+			}
+		}
+		return TArray<FRigVMExternalVariable>();
+	});
+
+
+	// this delegate is used by the controller to retrieve the current bytecode of the VM
+	InControllerToConfigure->GetCurrentByteCodeDelegate.BindLambda([WeakThis]() -> const FRigVMByteCode* {
+
+		if (TStrongObjectPtr<UObject> Object = WeakThis.Pin())
+		{
+			if (TRigVMAssetInterfacePtr Asset = Object.Get())
+			{
+				if (URigVM* VM = Asset->GetVM(false))
+				{
+					return &VM->GetByteCode();
+				}
+			}
+		}
+		return nullptr;
+	});
+
+#if WITH_EDITOR
+
+	// this sets up three delegates:
+	// a) get external variables (mapped to Controller->GetExternalVariables)
+	// b) bind pin to variable (mapped to Controller->BindPinToVariable)
+	// c) create external variable (mapped to the passed in tfunction)
+	// the last one is defined within the blueprint since the controller
+	// doesn't own the variables and can't create one itself.
+	InControllerToConfigure->SetupDefaultUnitNodeDelegates(TDelegate<FName(FRigVMExternalVariable, FString)>::CreateLambda(
+		[WeakThis](FRigVMExternalVariable InVariableToCreate, FString InDefaultValue) -> FName {
+			if (TStrongObjectPtr<UObject> Object = WeakThis.Pin())
+			{
+				if (Object->Implements<URigVMAssetInterface>())
+				{
+					TRigVMAssetInterfacePtr Asset = Object.Get();
+					return Asset->AddHostMemberVariableFromExternal(InVariableToCreate, InDefaultValue);
+				}
+			}
+			return NAME_None;
+		}
+	));
+
+	TWeakObjectPtr<URigVMController> WeakController = InControllerToConfigure;
+	InControllerToConfigure->RequestBulkEditDialogDelegate.BindLambda([WeakThis, WeakController](URigVMLibraryNode* InFunction, ERigVMControllerBulkEditType InEditType) -> FRigVMController_BulkEditResult 
+	{
+		if (TStrongObjectPtr<UObject> Object = WeakThis.Pin())
+		{
+			if (TRigVMAssetInterfacePtr Asset = Object.Get())
+			{
+				if (TStrongObjectPtr<URigVMController> Controller = WeakController.Pin())
+				{
+					if(Asset->OnRequestBulkEditDialog().IsBound())
+					{
+						return Asset->OnRequestBulkEditDialog().Execute(Asset, Controller.Get(), InFunction, InEditType);
+					}
+				}
+			}
+		}
+		return FRigVMController_BulkEditResult();
+	});
+
+	InControllerToConfigure->RequestBreakLinksDialogDelegate.BindLambda([WeakThis, WeakController](TArray<URigVMLink*> InLinks) -> bool 
+	{
+		if (TStrongObjectPtr<UObject> Object = WeakThis.Pin())
+		{
+			if (Object->Implements<URigVMAssetInterface>())
+			{
+				TRigVMAssetInterfacePtr BaseAsset = Object.Get();
+				if (TStrongObjectPtr<URigVMController> Controller = WeakController.Pin())
+				{
+					if(BaseAsset->OnRequestBreakLinksDialog().IsBound())
+					{
+						return BaseAsset->OnRequestBreakLinksDialog().Execute(InLinks);
+					}
+				}
+			}
+		}
+		return false;
+	});
+
+	InControllerToConfigure->RequestPinTypeSelectionDelegate.BindLambda([WeakThis](const TArray<TRigVMTypeIndex>& InTypes) -> TRigVMTypeIndex 
+	{
+		if (TStrongObjectPtr<UObject> Object = WeakThis.Pin())
+		{
+			if (Object->Implements<URigVMAssetInterface>())
+			{
+				TRigVMAssetInterfacePtr BaseAsset = Object.Get();
+				if(BaseAsset->OnRequestPinTypeSelectionDialog().IsBound())
+				{
+					return BaseAsset->OnRequestPinTypeSelectionDialog().Execute(InTypes);
+				}
+			}
+		}
+		return INDEX_NONE;
+	});
+
+	InControllerToConfigure->RequestNewExternalVariableDelegate.BindLambda([WeakThis](FRigVMGraphVariableDescription InVariable, bool bInIsPublic, bool bInIsReadOnly) -> FName
+	{
+		if (TStrongObjectPtr<UObject> Object = WeakThis.Pin())
+		{
+			if (Object->Implements<URigVMAssetInterface>())
+			{
+				TRigVMAssetInterfacePtr BaseAsset = Object.Get();
+				TArray<FRigVMExternalVariable> ExternalVariables = BaseAsset->GetExternalVariables(true);
+				for (FRigVMExternalVariable& ExistingVariable : ExternalVariables)
+				{
+					if (ExistingVariable.Name == InVariable.Name)
+					{
+						return FName();
+					}
+				}
+
+				FRigVMExternalVariable ExternalVariable = InVariable.ToExternalVariable();
+				return BaseAsset->AddMemberVariable(InVariable.Name,
+					ExternalVariable.TypeObject ? ExternalVariable.TypeObject->GetPathName() : ExternalVariable.TypeName.ToString(),
+					bInIsPublic,
+					bInIsReadOnly,
+					InVariable.DefaultValue);
+			}
+		}
+		
+		return FName();
+	});
+
+
+	InControllerToConfigure->RequestJumpToHyperlinkDelegate.BindLambda([WeakThis](const UObject* InSubject)
+	{
+		if (TStrongObjectPtr<UObject> Object = WeakThis.Pin())
+		{
+			if (Object->Implements<URigVMAssetInterface>())
+			{
+				TRigVMAssetInterfacePtr BaseAsset = Object.Get();
+				if(BaseAsset->OnRequestJumpToHyperlink().IsBound())
+				{
+					BaseAsset->OnRequestJumpToHyperlink().Execute(InSubject);
+				}
+			}
+		}
+	});
+
+#endif
+}
+
+UObject* IRigVMAssetInterface::ResolveUserDefinedTypeById(const FString& InTypeName) const
+{
+	const TMap<FString, FSoftObjectPath>& Info = GetUserDefinedStructGuidToPathName(false);
+	const FSoftObjectPath* ResultPathPtr = Info.Find(InTypeName);
+
+	if (ResultPathPtr == nullptr)
+	{
+		return nullptr;
+	}
+
+	if (UObject* TypeObject = ResultPathPtr->TryLoad())
+	{
+		// Ensure we have a hold on this type so it doesn't get nixed on the next GC.
+		const_cast<IRigVMAssetInterface*>(this)->GetUserDefinedTypesInUse(false).Add(TypeObject);
+		return TypeObject;
+	}
+
+	return nullptr;
+}
+
+bool IRigVMAssetInterface::TryImportGraphFromText(const FString& InClipboardText, UEdGraph** OutGraphPtr)
+{
+	if (OutGraphPtr)
+	{
+		*OutGraphPtr = nullptr;
+	}
+
+	if (URigVMController* FunctionLibraryController = GetOrCreateController(GetLocalFunctionLibrary()))
+	{
+		TGuardValue<FRigVMController_RequestLocalizeFunctionDelegate> RequestLocalizeDelegateGuard(
+            FunctionLibraryController->RequestLocalizeFunctionDelegate,
+            FRigVMController_RequestLocalizeFunctionDelegate::CreateLambda([this](FRigVMGraphFunctionIdentifier& InFunctionToLocalize)
+            {
+            	BroadcastRequestLocalizeFunctionDialog(InFunctionToLocalize);
+				const URigVMLibraryNode* LocalizedFunctionNode = GetLocalFunctionLibrary()->FindPreviouslyLocalizedFunction(InFunctionToLocalize);
+				return LocalizedFunctionNode != nullptr;
+            })
+        );
+		
+		TArray<FName> ImportedNodeNames = FunctionLibraryController->ImportNodesFromText(InClipboardText, true, true);
+		if (ImportedNodeNames.Num() == 0)
+		{
+			return false;
+		}
+
+		URigVMCollapseNode* CollapseNode = Cast<URigVMCollapseNode>(GetLocalFunctionLibrary()->FindFunction(ImportedNodeNames[0]));
+		if (ImportedNodeNames.Num() > 1 || CollapseNode == nullptr || CollapseNode->GetContainedGraph() == nullptr)
+		{
+			FunctionLibraryController->Undo();
+			return false;
+		}
+
+		UEdGraph* EdGraph = GetEdGraph(CollapseNode->GetContainedGraph());
+		if (OutGraphPtr)
+		{
+			*OutGraphPtr = EdGraph;
+		}
+
+		BroadcastGraphImported(EdGraph);
+	}
+
+	// always return true so that the default mechanism doesn't take over
+	return true;
+}
+
+URigVMEditorSettings* IRigVMAssetInterface::GetRigVMEditorSettings() const
+{
+	return GetMutableDefault<URigVMEditorSettings>(GetRigVMClientHost()->GetRigVMEditorSettingsClass());
+}
+
+#if WITH_EDITOR
+const FLazyName& IRigVMAssetInterface::GetPanelNodeFactoryName() const
+{
+	return RigVMPanelNodeFactoryName;
+}
+
+const FLazyName& IRigVMAssetInterface::GetPanelPinFactoryName() const
+{
+	return RigVMPanelPinFactoryName;
+}
+
+IRigVMEditorModule* IRigVMAssetInterface::GetEditorModule() const
+{
+	return &IRigVMEditorModule::Get();
+}
+#endif
+
+void IRigVMAssetInterface::SerializeImpl(FArchive& Ar)
+{
+	UE_RIGVM_ARCHIVETRACE_SCOPE(Ar, FString::Printf(TEXT("IRigVMAssetInterface(%s)"), *GetObject()->GetName()));
+	
+	Ar.UsingCustomVersion(FRigVMObjectVersion::GUID);
+	
+	if(IsValidChecked(GetObject()))
+	{
+		TScriptInterface<IRigVMClientHost> Host = GetRigVMClientHost();
+		Host->GetRigVMClient()->SetOuterClientHost(GetObject(), TEXT("RigVMClient"));
+	}
+	
+	SerializeSuper(Ar);
+	UE_RIGVM_ARCHIVETRACE_ENTRY(Ar, TEXT("Super::Serialize"));
+
+	if(Ar.IsObjectReferenceCollector())
+	{
+		Ar.UsingCustomVersion(FRigVMObjectVersion::GUID);
+
+#if WITH_EDITORONLY_DATA
+		if (Ar.IsCooking() && IsReferencedObjectPathsStored())
+		{
+			for (FSoftObjectPath ObjectPath : GetReferencedObjectPaths())
+			{
+				ObjectPath.Serialize(Ar);
+			}
+		}
+		else
+#endif
+		{
+			TArray<IRigVMGraphFunctionHost*> ReferencedFunctionHosts = GetReferencedFunctionHosts(false);
+
+			for(IRigVMGraphFunctionHost* ReferencedFunctionHost : ReferencedFunctionHosts)
+			{
+				if (URigVMBlueprintGeneratedClass* BPGeneratedClass = Cast<URigVMBlueprintGeneratedClass>(ReferencedFunctionHost))
+				{
+					Ar << BPGeneratedClass;
+				}
+			}
+		}
+	}
+
+	if(Ar.IsLoading())
+	{
+		TArray<UEdGraph*> EdGraphs;
+		GetAllEdGraphs(EdGraphs);
+		for (UEdGraph* EdGraph : EdGraphs)
+		{
+			EdGraph->Schema = GetRigVMClientHost()->GetRigVMEdGraphSchemaClass();
+			if (URigVMEdGraph* RigVMEdGraph = Cast<URigVMEdGraph>(EdGraph))
+			{
+				RigVMEdGraph->CachedModelGraph.Reset();
+			}
+		}
+
+		if (Ar.CustomVer(FRigVMObjectVersion::GUID) < FRigVMObjectVersion::AddVariantToRigVMAssets)
+		{
+			GetAssetVariant().Guid = FRigVMVariant::GenerateGUID(GetPackage()->GetPathName());
+		}
+	}
+
+	PostSerialize(Ar);
+}
+
+void IRigVMAssetInterface::PreSave(FObjectPreSaveContext ObjectSaveContext)
+{
+	PreSaveSuper(ObjectSaveContext);
+
+	GetRigVMClientHost()->GetRigVMClient()->PreSave(ObjectSaveContext);
+
+	GetSupportedEventNames().Reset();
+	if (URigVM* VM = GetVM(true))
+	{
+		GetSupportedEventNames() = VM->GetEntryNames();
+	}
+	
+	if (TScriptInterface<IRigVMGraphFunctionHost> FunctionHost = GetRigVMClientHost()->GetRigVMGraphFunctionHost())
+	{
+		FRigVMGraphFunctionStore* Store = FunctionHost->GetRigVMGraphFunctionStore();
+
+		GetPublicGraphFunctions().Reset();
+		GetPublicGraphFunctions().SetNum(Store->PublicFunctions.Num());
+		for (int32 i=0; i<GetPublicGraphFunctions().Num(); ++i)
+		{
+			GetPublicGraphFunctions()[i] = Store->PublicFunctions[i].Header;
+		}
+
+		URigVMFunctionLibrary* FunctionLibrary = GetLocalFunctionLibrary();
+		FunctionLibrary->FunctionToVariant.Reset();
+		for (int32 Pass=0; Pass<2; ++Pass)
+		{
+			const TArray<FRigVMGraphFunctionData>& Functions = (Pass == 0) ?
+				Store->PrivateFunctions
+				: Store->PublicFunctions;
+			for (const FRigVMGraphFunctionData& Function : Functions)
+			{
+				FunctionLibrary->FunctionToVariant.Add(Function.Header.Name, Function.Header.Variant);
+			}
+		}
+	}
+
+#if WITH_EDITORONLY_DATA
+	GetReferencedObjectPaths().Reset();
+
+	TArray<IRigVMGraphFunctionHost*> ReferencedFunctionHosts = GetReferencedFunctionHosts(false);
+	for(IRigVMGraphFunctionHost* ReferencedFunctionHost : ReferencedFunctionHosts)
+	{
+		if (URigVMBlueprintGeneratedClass* BPGeneratedClass = Cast<URigVMBlueprintGeneratedClass>(ReferencedFunctionHost))
+		{
+			GetReferencedObjectPaths().AddUnique(BPGeneratedClass);
+		}
+	}
+
+	IsReferencedObjectPathsStored() = true;
+#endif
+
+	GetFunctionReferenceNodeData() = GetReferenceNodeData();
+	IAssetRegistry::GetChecked().AssetTagsFinalized(*GetObject());
+
+	CachedAssetTags.Reset();
+
+	// also store the user defined struct guid to path name on the blueprint itself
+	// to aid the controller when recovering from user defined struct name changes or
+	// guid changes.
+	GetUserDefinedStructGuidToPathName(false).Reset();
+	GetUserDefinedEnumToPathName(false).Reset();
+	GetUserDefinedTypesInUse(false).Reset();
+	TArray<URigVMGraph*> AllModels = GetAllModels();
+	for(const URigVMGraph* Graph : AllModels)
+	{
+		for(const URigVMNode* Node : Graph->GetNodes())
+		{
+			const TArray<URigVMPin*> AllPins = Node->GetAllPinsRecursively();
+			for(const URigVMPin* Pin : AllPins)
+			{
+				if(const UUserDefinedStruct* UserDefinedStruct = Cast<UUserDefinedStruct>(Pin->GetCPPTypeObject()))
+				{
+					const FString GuidBasedName = RigVMTypeUtils::GetUniqueStructTypeName(UserDefinedStruct);
+					GetUserDefinedStructGuidToPathName(false).FindOrAdd(GuidBasedName) = FSoftObjectPath(UserDefinedStruct);
+				}
+				else if (const UUserDefinedEnum* UserDefinedEnum = Cast<UUserDefinedEnum>(Pin->GetCPPTypeObject()))
+				{
+					const FString EnumName = RigVMTypeUtils::CPPTypeFromEnum(UserDefinedEnum);
+					GetUserDefinedEnumToPathName(false).FindOrAdd(EnumName, UserDefinedEnum);
+				}
+			}
+		}
+	}
+
+#if WITH_EDITORONLY_DATA
+	OldMemoryStorageGeneratorClasses.Reset();
+#endif
+}
+
+void IRigVMAssetInterface::PostLoad()
+{
+	PostLoadSuper();
+
+	FRigVMRegistry::Get().RefreshEngineTypesIfRequired();
+
+	bVMRecompilationRequired = true;
+	{
+		TGuardValue<bool> IsCompilingGuard(bIsCompiling, true);
+		
+		TArray<IRigVMGraphFunctionHost*> ReferencedFunctionHosts = GetReferencedFunctionHosts(true);
+
+		// PostLoad all referenced function hosts so that their function data are fully loaded 
+		// and ready to be inlined into this BP during compilation
+		for (IRigVMGraphFunctionHost* FunctionHost : ReferencedFunctionHosts)
+		{
+			if (URigVMBlueprintGeneratedClass* BPGeneratedClass = Cast<URigVMBlueprintGeneratedClass>(FunctionHost))
+			{
+				if (BPGeneratedClass->HasAllFlags(RF_NeedPostLoad))
+				{
+					BPGeneratedClass->ConditionalPostLoad();
+				}
+			}
+		}
+		
+		// temporarily disable default value validation during load time, serialized values should always be accepted
+		TGuardValue<bool> DisablePinDefaultValueValidation(GetOrCreateController()->bValidatePinDefaults, false);
+
+		// remove all non-controlrig-graphs
+		TArray<TObjectPtr<UEdGraph>> NewUberGraphPages;
+		for (UEdGraph* Graph : GetUberGraphs())
+		{
+			URigVMEdGraph* RigGraph = Cast<URigVMEdGraph>(Graph);
+			if (RigGraph && RigGraph->GetClass() == GetRigVMClientHost()->GetRigVMEdGraphClass())
+			{
+				NewUberGraphPages.Add(RigGraph);
+			}
+			else
+			{
+                // We are renaming an object to a new outer while we may still be loading. Since we
+                // are destroying the object, pass REN_AllowPackageLinkerMismatch to avoid forcing
+                // the load to complete since that is wasteful.
+				Graph->MarkAsGarbage();
+				Graph->Rename(nullptr, GetTransientPackage(), REN_AllowPackageLinkerMismatch);
+			}
+		}
+		SetUberGraphs(NewUberGraphPages);
+		
+		TArray<TGuardValue<bool>> EditableGraphGuards;
+		{
+			for (URigVMGraph* Graph : GetAllModels())
+			{
+				EditableGraphGuards.Emplace(Graph->bEditable, true);
+			}
+		}
+		
+		InitializeModelIfRequired(false /* recompile vm */);
+		{
+			TGuardValue<bool> GuardNotifications(bSuspendModelNotificationsForSelf, true);
+			
+			const FRigVMClientPatchResult PatchResult = GetRigVMClientHost()->GetRigVMClient()->PatchModelsOnLoad();
+			if(PatchResult.RequiresToMarkPackageDirty())
+			{
+				(void)MarkPackageDirty();
+				bDirtyDuringLoad = true;
+			}
+			
+			GetRigVMClientHost()->GetRigVMClient()->PatchFunctionReferencesOnLoad();
+			GetFunctionReferenceNodeData() = GetReferenceNodeData();
+
+			PatchVariableNodesOnLoad();
+			PatchVariableNodesWithIncorrectType();
+			PathDomainSpecificContentOnLoad();
+			PatchBoundVariables();
+			PatchParameterNodesOnLoad();
+			PatchLinksWithCast();
+			
+			TMap<URigVMLibraryNode*, FRigVMGraphFunctionHeader> OldHeaders;
+			// Backwards compatibility. Store public access in the model
+			TArray<FName> BackwardsCompatiblePublicFunctions;
+			GetBackwardsCompatibilityPublicFunctions(BackwardsCompatiblePublicFunctions, OldHeaders);
+
+			GetRigVMClientHost()->GetRigVMClient()->PatchFunctionsOnLoad(GetRigVMClientHost()->GetRigVMGraphFunctionHost().GetInterface(), BackwardsCompatiblePublicFunctions, OldHeaders);
+
+			const FRigVMClientPatchResult PinDefaultValuePatchResult = GetRigVMClientHost()->GetRigVMClient()->PatchPinDefaultValues();
+			if(PinDefaultValuePatchResult.RequiresToMarkPackageDirty())
+			{
+				(void)MarkPackageDirty();
+				bDirtyDuringLoad = true;
+			}
+		}
+
+#if WITH_EDITOR
+
+		{
+			TGuardValue<bool> GuardNotifications(bSuspendModelNotificationsForSelf, true);
+
+			// refresh the graph such that the pin hierarchies matches their CPPTypeObject
+			// this step is needed everytime we open a BP in the editor, b/c even after load
+			// model data can change while the Control Rig BP is not opened
+			// for example, if a user defined struct changed after BP load,
+			// any pin that references the struct needs to be regenerated
+			RefreshAllModels();
+		}
+		
+		// at this point we may still have links which are detached. we may or may not be able to 
+		// reattach them.
+		GetRigVMClientHost()->GetRigVMClient()->ProcessDetachedLinks();
+
+		GetRigVMClientHost()->GetRigVMGraphFunctionHost()->GetRigVMGraphFunctionStore()->RemoveAllCompilationData();
+
+		// perform backwards compat value upgrades
+		TArray<URigVMGraph*> GraphsToValidate = GetAllModels();
+		for (int32 GraphIndex = 0; GraphIndex < GraphsToValidate.Num(); GraphIndex++)
+		{
+			URigVMGraph* GraphToValidate = GraphsToValidate[GraphIndex];
+			if(GraphToValidate == nullptr)
+			{
+				continue;
+			}
+
+			for(URigVMNode* Node : GraphToValidate->GetNodes())
+			{
+				URigVMController* Controller = GetOrCreateController(GraphToValidate);
+				FRigVMControllerNotifGuard NotifGuard(Controller, true);
+				Controller->RemoveUnusedOrphanedPins(Node, true);
+			}
+				
+			for(URigVMNode* Node : GraphToValidate->GetNodes())
+			{
+				// avoid function reference related validation for temp assets, a temp asset may get generated during
+				// certain content validation process. It is usually just a simple file-level copy of the source asset
+				// so these references are usually not fixed-up properly. Thus, it is meaningless to validate them.
+				// They should not be allowed to dirty the source asset either.
+				if (!GetPackage()->GetName().StartsWith("/Temp/"))
+				{
+					if(URigVMFunctionReferenceNode* FunctionReferenceNode = Cast<URigVMFunctionReferenceNode>(Node))
+					{
+						if(URigVMBuildData* BuildData = URigVMBuildData::Get())
+						{
+							BuildData->RegisterFunctionReference(FunctionReferenceNode);
+						}
+					}
+				}
+			}
+		}
+
+		CompileLog.Messages.Reset();
+		CompileLog.NumErrors = CompileLog.NumWarnings = 0;
+#endif
+	}
+
+	// remove invalid class objects that were parented to the rigvmbp object
+	RemoveDeprecatedVMMemoryClass();
+	
+#if WITH_EDITOR
+	if(GIsEditor)
+	{
+		// delay compilation until the package has been loaded
+		FCoreUObjectDelegates::OnEndLoadPackage.AddRaw(this, &IRigVMAssetInterface::HandlePackageDone);
+	}
+#else
+	RecompileVMIfRequired();
+#endif
+	GetRigVMClientHost()->RequestRigVMInit();
+
+	
+	if (!GetAssetVariant().Guid.IsValid())
+	{
+		GetAssetVariant().Guid = FRigVMVariant::GenerateGUID();
+	}
+
+	if (UPackage* Package = GetPackage())
+	{
+		Package->SetDirtyFlag(bDirtyDuringLoad);
+	}
+
+#if WITH_EDITOR
+	// if we are running with -game we are in editor code,
+	// but GIsEditor is turned off
+	if(!GIsEditor)
+	{
+		HandlePackageDone();
+	}
+#endif
+
+	// RigVMRegistry changes can be triggered when new user defined types(structs/enums) are added/removed
+	// in which case we have to refresh the model
+	FRigVMRegistry::Get().OnRigVMRegistryChanged().RemoveAll(this);
+	FRigVMRegistry::Get().OnRigVMRegistryChanged().AddRaw(this, &IRigVMAssetInterface::OnRigVMRegistryChanged);
+}
+
+#if WITH_EDITORONLY_DATA
+void IRigVMAssetInterface::DeclareConstructClasses(TArray<FTopLevelAssetPath>& OutConstructClasses, const UClass* SpecificSubclass)
+{
+	OutConstructClasses.Add(FTopLevelAssetPath(URigVMController::StaticClass()));
+}
+#endif
+
+#if WITH_EDITOR
+void IRigVMAssetInterface::HandlePackageDone(const FEndLoadPackageContext& Context)
+{
+	if (!Context.LoadedPackages.Contains(GetPackage()))
+	{
+		return;
+	}
+	HandlePackageDone();
+}
+
+void IRigVMAssetInterface::HandlePackageDone()
+{
+	FCoreUObjectDelegates::OnEndLoadPackage.RemoveAll(this);
+
+	if(URigVMBuildData* BuildData = URigVMBuildData::Get())
+	{
+		if(URigVMFunctionLibrary* FunctionLibrary = GetRigVMClientHost()->GetRigVMClient()->GetFunctionLibrary())
+		{
+			// for backwards compatibility load the function references from the
+			// model's storage over to the centralized build data
+			if(!FunctionLibrary->FunctionReferences_DEPRECATED.IsEmpty())
+			{
+				// let's also update the asset data of the dependents
+				IAssetRegistry& AssetRegistry = IAssetRegistry::GetChecked();
+				
+				for(const TTuple< TObjectPtr<URigVMLibraryNode>, FRigVMFunctionReferenceArray >& Pair :
+					FunctionLibrary->FunctionReferences_DEPRECATED)
+				{
+					TSoftObjectPtr<URigVMLibraryNode> FunctionKey(Pair.Key);
+						
+					for(int32 ReferenceIndex = 0; ReferenceIndex < Pair.Value.Num(); ReferenceIndex++)
+					{
+						// update the build data
+						BuildData->RegisterFunctionReference(FunctionKey->GetFunctionIdentifier(), Pair.Value[ReferenceIndex]);
+
+						// find all control rigs matching the reference node
+						FAssetData AssetData = AssetRegistry.GetAssetByObjectPath(
+							Pair.Value[ReferenceIndex].ToSoftObjectPath().GetWithoutSubPath());
+
+						// if the asset has never been loaded - make sure to load it once and mark as dirty
+						if(AssetData.IsValid() && !AssetData.IsAssetLoaded())
+						{
+							if(TRigVMAssetInterfacePtr Dependent = AssetData.GetAsset())
+							{
+								if(Dependent != this)
+								{
+									(void)Dependent->MarkPackageDirty();
+								}
+							}
+						}
+					}
+				}
+				
+				FunctionLibrary->FunctionReferences_DEPRECATED.Reset();
+				(void)MarkPackageDirty();
+			}
+		}
+
+		// update the build data from the current function references
+		const TArray<FRigVMReferenceNodeData> ReferenceNodeDatas = GetReferenceNodeData();
+		for(const FRigVMReferenceNodeData& ReferenceNodeData : ReferenceNodeDatas)
+		{
+			BuildData->RegisterFunctionReference(ReferenceNodeData);
+		}
+
+		BuildData->ClearInvalidReferences();
+	}
+	
+	{
+		const FRigVMCompileSettingsDuringLoadGuard Guard(GetVMCompileSettings());
+		RecompileVM();
+	}
+	GetRigVMClientHost()->RequestRigVMInit();
+	BroadcastRigVMPackageDone();
+}
+
+void IRigVMAssetInterface::BroadcastRigVMPackageDone()
+{
+	TArray<UObject*> ArchetypeInstances = GetArchetypeInstances(true, false);
+	for (UObject* Instance : ArchetypeInstances)
+	{
+		if (URigVMHost* InstanceHost = Cast<URigVMHost>(Instance))
+		{
+			InstanceHost->BroadCastEndLoadPackage();
+		}
+	}
+}
+
+void IRigVMAssetInterface::RemoveDeprecatedVMMemoryClass() 
+{
+	TArray<UObject*> Objects;
+	GetObjectsWithOuter(GetObject(), Objects, false);
+
+#if WITH_EDITORONLY_DATA
+	OldMemoryStorageGeneratorClasses.Reserve(Objects.Num());
+	for (UObject* Object : Objects)
+	{
+		if (URigVMMemoryStorageGeneratorClass* DeprecatedClass = Cast<URigVMMemoryStorageGeneratorClass>(Object))
+		{
+			// Making sure it is fully loaded before removing it to avoid ambiguity regarding load order
+			DeprecatedClass->ConditionalPostLoad();
+			
+			DeprecatedClass->Rename(nullptr, GetTransientPackage(), REN_DoNotDirty | REN_DontCreateRedirectors | REN_NonTransactional);
+			OldMemoryStorageGeneratorClasses.Add(DeprecatedClass);
+		}
+	}
+#endif
+}
+#endif
+
+void IRigVMAssetInterface::GenerateUserDefinedDependenciesData(FRigVMExtendedExecuteContext& Context)
+{
+	if (URigVM* VM = GetVM(true))
+	{
+		const TArray<const UObject*> UserDefinedDependencies = VM->GetUserDefinedDependencies({ VM->GetDefaultMemoryByType(ERigVMMemoryType::Literal), VM->GetDefaultMemoryByType(ERigVMMemoryType::Work) });
+		GetUserDefinedStructGuidToPathName(true).Reset();
+		GetUserDefinedEnumToPathName(true).Reset();
+		GetUserDefinedTypesInUse(true).Reset();
+
+		for (const UObject* UserDefinedDependency : UserDefinedDependencies)
+		{
+			if (const UUserDefinedStruct* UserDefinedStruct = Cast<UUserDefinedStruct>(UserDefinedDependency))
+			{
+				const FString GuidBasedName = RigVMTypeUtils::GetUniqueStructTypeName(UserDefinedStruct);
+				GetUserDefinedStructGuidToPathName(true).Add(GuidBasedName, UserDefinedStruct);
+			}
+			else if (const UUserDefinedEnum* UserDefinedEnum = Cast<UUserDefinedEnum>(UserDefinedDependency))
+			{
+				const FString EnumName = RigVMTypeUtils::CPPTypeFromEnum(UserDefinedEnum);
+				GetUserDefinedEnumToPathName(true).Add(EnumName, UserDefinedEnum);
+			}
+		}
+	}
+}
+
+void IRigVMAssetInterface::RecompileVM()
+{
+	if(bIsCompiling)
+	{
+		return;
+	}
+
+	TGuardValue<bool> CompilingGuard(bIsCompiling, true);
+	
+	bErrorsDuringCompilation = false;
+
+	FRigVMEdGraphDisplaySettings& RigGraphDisplaySettings = GetRigGraphDisplaySettings();
+	if(RigGraphDisplaySettings.bAutoDetermineRange)
+	{
+		RigGraphDisplaySettings.MinMicroSeconds = RigGraphDisplaySettings.LastMinMicroSeconds = DBL_MAX;
+		RigGraphDisplaySettings.MaxMicroSeconds = RigGraphDisplaySettings.LastMaxMicroSeconds = (double)INDEX_NONE;
+	}
+	else if(RigGraphDisplaySettings.MaxMicroSeconds < RigGraphDisplaySettings.MinMicroSeconds)
+	{
+		RigGraphDisplaySettings.MinMicroSeconds = 0;
+		RigGraphDisplaySettings.MaxMicroSeconds = 5;
+	}
+	
+	RigGraphDisplaySettings.TotalMicroSeconds = 0.0;
+	RigGraphDisplaySettings.MinMicroSecondsFrames.Reset();
+	RigGraphDisplaySettings.MaxMicroSecondsFrames.Reset();
+	RigGraphDisplaySettings.TotalMicroSecondsFrames.Reset();
+
+	URigVM* VM = GetVM(true);
+	FRigVMExtendedExecuteContext* Context = GetRigVMExtendedExecuteContext();
+	if (VM && Context)
+	{
+		FRigVMClient* RigVMClient = GetRigVMClientHost()->GetRigVMClient();
+		TGuardValue<bool> ReentrantGuardSelf(bSuspendModelNotificationsForSelf, true);
+		TGuardValue<bool> ReentrantGuardOthers(RigVMClient->bSuspendModelNotificationsForOthers, true);
+
+		PreCompile();
+
+		GetObject()->Modify(false);
+		VM->Reset(*Context);
+
+		// Clear all Errors
+		CompileLog.Messages.Reset();
+		CompileLog.NumErrors = CompileLog.NumWarnings = 0;
+		
+		TArray<UEdGraph*> EdGraphs;
+		GetAllEdGraphs(EdGraphs);
+
+		for (UEdGraph* Graph : EdGraphs)
+		{
+			URigVMEdGraph* RigGraph = Cast<URigVMEdGraph>(Graph);
+			if (RigGraph == nullptr)
+			{
+				continue;
+			}
+
+			for (UEdGraphNode* GraphNode : Graph->Nodes)
+			{
+				if (URigVMEdGraphNode* RigVMEdGraphNode = Cast<URigVMEdGraphNode>(GraphNode))
+				{
+					RigVMEdGraphNode->ClearErrorInfo();
+				}
+			}
+		}
+
+		URigVMCompiler* Compiler = URigVMCompiler::StaticClass()->GetDefaultObject<URigVMCompiler>();
+		GetVMCompileSettings().SetExecuteContextStruct(RigVMClient->GetDefaultExecuteContextStruct());
+
+		const FRigVMCompileSettings Settings = (bCompileInDebugMode) ? FRigVMCompileSettings::Fast(GetVMCompileSettings().GetExecuteContextStruct()) : GetVMCompileSettings();
+		Compiler->Compile(Settings, RigVMClient->GetAllModels(false, false), GetOrCreateController(), VM, *Context, GetExternalVariables(false), &PinToOperandMap);
+		
+
+		VM->Initialize(*Context);
+		GenerateUserDefinedDependenciesData(*Context);
+
+		if (bErrorsDuringCompilation)
+		{
+			if(Settings.SurpressErrors)
+			{
+				Settings.Reportf(EMessageSeverity::Info, GetObject(),
+					TEXT("Compilation Errors may be suppressed for ControlRigBlueprint: %s. See VM Compile Setting in Class Settings for more Details"), *GetObject()->GetName());
+			}
+			bVMRecompilationRequired = false;
+			if(VM)
+			{
+				VMCompiledEvent.Broadcast(GetObject(), VM, *Context);
+			}
+			return;
+		}
+
+		InitializeArchetypeInstances();
+
+		bVMRecompilationRequired = false;
+		VMCompiledEvent.Broadcast(GetObject(), VM, *Context);
+
+#if WITH_EDITOR
+		RefreshBreakpoints();
+#endif
+	}
+}
+
+void IRigVMAssetInterface::RecompileVMIfRequired()
+{
+	if (bVMRecompilationRequired)
+	{
+		RecompileVM();
+	}
+}
+
+void IRigVMAssetInterface::RequestAutoVMRecompilation()
+{
+	bVMRecompilationRequired = true;
+	if (bAutoRecompileVM && VMRecompilationBracket == 0)
+	{
+		RecompileVMIfRequired();
+	}
+}
+
+void IRigVMAssetInterface::IncrementVMRecompileBracket()
+{
+	VMRecompilationBracket++;
+}
+
+void IRigVMAssetInterface::DecrementVMRecompileBracket()
+{
+	if (VMRecompilationBracket == 1)
+	{
+		if (bAutoRecompileVM)
+		{
+			RecompileVMIfRequired();
+		}
+		VMRecompilationBracket = 0;
+	}
+	else if (VMRecompilationBracket > 0)
+	{
+		VMRecompilationBracket--;
+	}
+}
+
+void IRigVMAssetInterface::RefreshAllModels(ERigVMLoadType InLoadType)
+{
+	const bool bEnablePostLoadHashing = CVarRigVMEnablePostLoadHashing->GetBool();
+
+	GetRigVMClientHost()->GetRigVMClient()->RefreshAllModels(InLoadType, bEnablePostLoadHashing, bIsCompiling);
+}
+
+void IRigVMAssetInterface::OnRigVMRegistryChanged()
+{
+	RefreshAllModels();
+	RebuildGraphFromModel();
+	// avoids slate crash
+	RefreshAllNodes();
+}
+
+void IRigVMAssetInterface::HandleReportFromCompiler(EMessageSeverity::Type InSeverity, UObject* InSubject, const FString& InMessage)
+{
+	UObject* SubjectForMessage = InSubject;
+	if(URigVMNode* ModelNode = Cast<URigVMNode>(SubjectForMessage))
+	{
+		if(IRigVMAssetInterface* RigBlueprint = ModelNode->GetImplementingOuter<IRigVMAssetInterface>())
+		{
+			if(URigVMEdGraph* EdGraph = Cast<URigVMEdGraph>(RigBlueprint->GetEdGraph(ModelNode->GetGraph())))
+			{
+				if(UEdGraphNode* EdNode = EdGraph->FindNodeForModelNodeName(ModelNode->GetFName()))
+				{
+					SubjectForMessage = EdNode;
+				}
+			}
+		}
+	}
+
+	FCompilerResultsLog* CurrentMessageLog = GetCurrentMessageLog();
+	FCompilerResultsLog* Log = CurrentMessageLog ? CurrentMessageLog : &CompileLog;
+	if (InSeverity == EMessageSeverity::Error)
+	{
+		SetAssetStatus(ERigVMAssetStatus::RVMA_Error);
+		(void)MarkPackageDirty();
+
+		// see UnitTest "ControlRig.Basics.OrphanedPins" to learn why errors are suppressed this way
+		if (GetVMCompileSettings().SurpressErrors)
+		{
+			Log->bSilentMode = true;
+		}
+
+		if(InMessage.Contains(TEXT("@@")))
+		{
+			Log->Error(*InMessage, SubjectForMessage);
+		}
+		else
+		{
+			Log->Error(*InMessage);
+		}
+
+		BroadCastReportCompilerMessage(InSeverity, InSubject, InMessage);
+
+		// see UnitTest "ControlRig.Basics.OrphanedPins" to learn why errors are suppressed this way
+		if (!GetVMCompileSettings().SurpressErrors)
+		{ 
+			FScriptExceptionHandler::Get().HandleException(ELogVerbosity::Error, *InMessage, *FString());
+		}
+		
+		bErrorsDuringCompilation = true;
+	}
+	else if (InSeverity == EMessageSeverity::Warning)
+	{
+		if(InMessage.Contains(TEXT("@@")))
+		{
+			Log->Warning(*InMessage, SubjectForMessage);
+		}
+		else
+		{
+			Log->Warning(*InMessage);
+		}
+
+		BroadCastReportCompilerMessage(InSeverity, InSubject, InMessage);
+		FScriptExceptionHandler::Get().HandleException(ELogVerbosity::Warning, *InMessage, *FString());
+	}
+	else
+	{
+		if(InMessage.Contains(TEXT("@@")))
+		{
+			Log->Note(*InMessage, SubjectForMessage);
+		}
+		else
+		{
+			Log->Note(*InMessage);
+		}
+
+		static const FString Error = TEXT("Error");
+		static const FString Warning = TEXT("Warning");
+		if(InMessage.Contains(Error, ESearchCase::IgnoreCase) ||
+			InMessage.Contains(Warning, ESearchCase::IgnoreCase))
+		{
+			BroadCastReportCompilerMessage(InSeverity, InSubject, InMessage);
+		}
+		UE_LOG(LogRigVMDeveloper, Display, TEXT("%s"), *InMessage);
+	}
+
+	if (URigVMEdGraphNode* EdGraphNode = Cast<URigVMEdGraphNode>(SubjectForMessage))
+	{
+		EdGraphNode->SetErrorInfo(InSeverity, InMessage);
+		EdGraphNode->bHasCompilerMessage = EdGraphNode->ErrorType <= int32(EMessageSeverity::Info);
+	}
+}
+
+TArray<IRigVMGraphFunctionHost*> IRigVMAssetInterface::GetReferencedFunctionHosts(bool bForceLoad) const
+{
+	TArray<IRigVMGraphFunctionHost*> ReferencedBlueprints;
+	
+	TArray<UEdGraph*> EdGraphs;
+	GetAllEdGraphs(EdGraphs);
+	for (UEdGraph* EdGraph : EdGraphs)
+	{
+		for(UEdGraphNode* Node : EdGraph->Nodes)
+		{
+			if(URigVMEdGraphNode* RigNode = Cast<URigVMEdGraphNode>(Node))
+			{
+				if(URigVMFunctionReferenceNode* FunctionRefNode = Cast<URigVMFunctionReferenceNode>(RigNode->GetModelNode()))
+				{
+					IRigVMGraphFunctionHost* Host = nullptr;
+					if (bForceLoad || FunctionRefNode->IsReferencedFunctionHostLoaded())
+					{
+						// Load the function host
+						Host = FunctionRefNode->GetReferencedFunctionHeader().GetFunctionHost();
+					}
+					else if (bForceLoad || FunctionRefNode->IsReferencedNodeLoaded())
+					{
+						// Load the reference library node
+						if(const URigVMLibraryNode* ReferencedNode = FunctionRefNode->LoadReferencedNode())
+						{
+							if(URigVMFunctionLibrary* ReferencedFunctionLibrary = ReferencedNode->GetLibrary())
+							{
+								FSoftObjectPath FunctionHostPath = ReferencedFunctionLibrary->GetFunctionHostObjectPath();
+								if (UObject* FunctionHostObj = FunctionHostPath.TryLoad())
+								{
+									Host = Cast<IRigVMGraphFunctionHost>(FunctionHostObj);
+								}
+							}
+						}
+					}
+
+					if (Host != nullptr && Host != GetRigVMClientHost()->GetRigVMGraphFunctionHost().GetInterface())
+					{
+						ReferencedBlueprints.Add(Host);
+					}
+				}
+			}
+		}
+	}
+	
+	return ReferencedBlueprints;
+}
+
+#if WITH_EDITOR
+
+void IRigVMAssetInterface::SetDebugMode(const bool bValue)
+{
+	bCompileInDebugMode = bValue;
+}
+
+void IRigVMAssetInterface::ClearBreakpoints()
+{
+	for(URigVMNode* Node : GetRigVMBreakpointNodes())
+	{
+		Node->SetHasBreakpoint(false);		
+	}
+	
+	GetRigVMBreakpointNodes().Empty();
+	RefreshBreakpoints();
+}
+
+bool IRigVMAssetInterface::AddBreakpoint(const FString& InBreakpointNodePath)
+{
+	URigVMLibraryNode* FunctionNode = nullptr;
+	
+	// Find the node in the graph
+	FRigVMClient* RigVMClient = GetRigVMClientHost()->GetRigVMClient();
+	for(const URigVMGraph* Model : *RigVMClient)
+	{
+		URigVMNode* BreakpointNode = Model->FindNode(InBreakpointNodePath);
+		if (BreakpointNode == nullptr)
+		{
+			// If we cannot find the node, it might be because it is inside a function
+			FString FunctionName = InBreakpointNodePath, Right;
+			URigVMNode::SplitNodePathAtStart(InBreakpointNodePath, FunctionName, Right);
+
+			// Look inside the local function library
+			if (URigVMLibraryNode* LibraryNode = GetLocalFunctionLibrary()->FindFunction(FName(FunctionName)))
+			{
+				BreakpointNode = LibraryNode->GetContainedGraph()->FindNode(Right);
+				FunctionNode = LibraryNode;
+			}
+		}
+
+		if(BreakpointNode)
+		{
+			return AddBreakpoint(BreakpointNode, FunctionNode);
+		}
+	}
+	return false;
+}
+
+bool IRigVMAssetInterface::AddBreakpoint(URigVMNode* InBreakpointNode, URigVMLibraryNode* LibraryNode)
+{
+	if (InBreakpointNode == nullptr)
+	{
+		return false;
+	}
+
+	bool bSuccess = true;
+	if (LibraryNode)
+	{
+		// If the breakpoint node is inside a library node, find all references to the library node
+		TArray<TSoftObjectPtr<URigVMFunctionReferenceNode>> References = LibraryNode->GetLibrary()->GetReferencesForFunction(LibraryNode->GetFName());
+		for (TSoftObjectPtr<URigVMFunctionReferenceNode> Reference : References)
+		{
+			if (!Reference.IsValid())
+			{
+				continue;
+			}
+
+			IRigVMAssetInterface* ReferenceBlueprint = Reference->GetImplementingOuter<IRigVMAssetInterface>();
+
+			// If the reference is not inside another function, add a breakpoint in the blueprint containing the
+			// reference, without a function specified
+			bool bIsInsideFunction = Reference->GetRootGraph()->IsA<URigVMFunctionLibrary>();
+			if(!bIsInsideFunction)
+			{
+				bSuccess &= ReferenceBlueprint->AddBreakpoint(InBreakpointNode);
+			}
+			else
+			{
+				// Otherwise, we need to add breakpoints to all the blueprints that reference this
+				// function (when the blueprint graph is flattened)
+				
+				// Get all the functions containing this reference
+				URigVMNode* Node = Reference.Get();
+				while (Node->GetGraph() != ReferenceBlueprint->GetLocalFunctionLibrary())
+				{
+					if (URigVMLibraryNode* ParentLibraryNode = Cast<URigVMLibraryNode>(Node->GetGraph()->GetOuter()))
+					{
+						// Recursively add breakpoints to the reference blueprint, specifying the parent function
+						bSuccess &= ReferenceBlueprint->AddBreakpoint(InBreakpointNode, ParentLibraryNode);
+					}
+
+					Node = Cast<URigVMNode>(Node->GetGraph()->GetOuter());
+				}
+			}
+		}
+	}
+	else
+	{
+		if (!GetRigVMBreakpointNodes().Contains(InBreakpointNode))
+		{
+			// Add the breakpoint to the VM
+			bSuccess = AddBreakpointToHost(InBreakpointNode);
+			BreakpointAddedEvent.Broadcast();
+		}
+	}
+
+	return bSuccess;
+}
+
+bool IRigVMAssetInterface::AddBreakpointToHost(URigVMNode* InBreakpointNode)
+{
+	const FRigVMByteCode* ByteCode = GetController()->GetCurrentByteCode();
+	FRigVMClient* RigVMClient = GetRigVMClientHost()->GetRigVMClient();
+	TSet<FString> AddedCallpaths;
+
+	if (ByteCode)
+	{
+		FRigVMInstructionArray Instructions = ByteCode->GetInstructions();	
+		
+		// For each instruction, see if the node is in the callpath
+		// Only add one breakpoint for each callpath related to this node (i.e. if a node produces multiple
+		// instructions, only add a breakpoint to the first instruction)
+		for (int32 i = 0; i< Instructions.Num(); ++i)
+		{
+			for(URigVMGraph* Model : *RigVMClient)
+			{
+				const FRigVMASTProxy Proxy = FRigVMASTProxy::MakeFromCallPath(ByteCode->GetCallPathForInstruction(i), Model);
+				if (Proxy.GetCallstack().Contains(InBreakpointNode))
+				{
+					// Find the callpath related to the breakpoint node
+					FRigVMASTProxy BreakpointProxy = Proxy;
+					while(BreakpointProxy.GetSubject() != InBreakpointNode)
+					{
+						BreakpointProxy = BreakpointProxy.GetParent();
+					}
+					const FString& BreakpointCallPath = BreakpointProxy.GetCallstack().GetCallPath();
+
+					// Only add this callpath breakpoint once
+					if (!AddedCallpaths.Contains(BreakpointCallPath))
+					{
+						AddedCallpaths.Add(BreakpointCallPath);
+						GetDebugInfo().AddBreakpoint(i, InBreakpointNode, static_cast<uint16>(BreakpointProxy.GetCallstack().Num()));
+					}
+				}
+			}
+		}
+	}
+
+	if (AddedCallpaths.Num() > 0)
+	{
+		GetRigVMBreakpointNodes().AddUnique(InBreakpointNode);
+		return true;
+	}
+	
+	return false;
+}
+
+bool IRigVMAssetInterface::RemoveBreakpoint(const FString& InBreakpointNodePath)
+{
+	// Find the node in the graph
+	URigVMNode* BreakpointNode = nullptr;
+	FRigVMClient* RigVMClient = GetRigVMClientHost()->GetRigVMClient();
+	
+	for(URigVMGraph* Model : *RigVMClient)
+	{
+		BreakpointNode = Model->FindNode(InBreakpointNodePath);
+		if (BreakpointNode == nullptr)
+		{
+			// If we cannot find the node, it might be because it is inside a function
+			FString FunctionName = InBreakpointNodePath, Right;
+			URigVMNode::SplitNodePathAtStart(InBreakpointNodePath, FunctionName, Right);
+
+			// Look inside the local function library
+			if (URigVMLibraryNode* LibraryNode = GetLocalFunctionLibrary()->FindFunction(FName(FunctionName)))
+			{
+				BreakpointNode = LibraryNode->GetContainedGraph()->FindNode(Right);
+			}
+		}
+		if(BreakpointNode)
+		{
+			break;
+		}
+	}
+
+	if(BreakpointNode)
+	{
+		bool bSuccess = RemoveBreakpoint(BreakpointNode);
+
+		// Remove the breakpoint from all the loaded dependent blueprints
+		TArray<IRigVMAssetInterface*> DependentBlueprints = GetDependentResolvedAssets(true, true);
+		DependentBlueprints.Remove(this);
+		for (IRigVMAssetInterface* Dependent : DependentBlueprints)
+		{
+			bSuccess &= Dependent->RemoveBreakpoint(BreakpointNode);
+		}
+		return bSuccess;
+	}
+
+	return false;
+}
+
+bool IRigVMAssetInterface::RemoveBreakpoint(URigVMNode* InBreakpointNode)
+{
+	if (GetRigVMBreakpointNodes().Contains(InBreakpointNode))
+	{
+		GetRigVMBreakpointNodes().Remove(InBreakpointNode);
+
+		// Multiple breakpoint nodes might set a breakpoint to the same instruction. When we remove
+		// one of the breakpoint nodes, we do not want to remove the instruction breakpoint if there
+		// is another breakpoint node addressing it. For that reason, we just recompute all the
+		// breakpoint instructions.
+		// Refreshing breakpoints in the control rig will keep the state it had before.
+		RefreshBreakpoints();
+		return true;
+	}
+
+	return false;
+}
+
+void IRigVMAssetInterface::RefreshBreakpoints()
+{
+	GetDebugInfo().Reset();
+	for (URigVMNode* Node : GetRigVMBreakpointNodes())
+	{
+		AddBreakpointToHost(Node);
+	}
+}
+
+TArray<FRigVMReferenceNodeData> IRigVMAssetInterface::GetReferenceNodeData() const
+{
+	TArray<FRigVMReferenceNodeData> Data;
+	
+	const TArray<URigVMGraph*> AllModels = GetAllModels();
+	for (URigVMGraph* ModelToVisit : AllModels)
+	{
+		for(URigVMNode* Node : ModelToVisit->GetNodes())
+		{
+			if(URigVMFunctionReferenceNode* ReferenceNode = Cast<URigVMFunctionReferenceNode>(Node))
+			{
+				Data.Add(FRigVMReferenceNodeData(ReferenceNode));
+			}
+		}
+	}
+	return Data;
+}
+
+#endif
+
+void IRigVMAssetInterface::SetupDefaultObjectDuringCompilation(URigVMHost* InCDO)
+{
+	InCDO->PostInitInstanceIfRequired();
+	InCDO->VMRuntimeSettings = GetVMRuntimeSettings();
+}
+
+void IRigVMAssetInterface::MarkDirtyDuringLoad()
+{
+	bDirtyDuringLoad = true;
+}
+
+bool IRigVMAssetInterface::IsMarkedDirtyDuringLoad() const
+{
+	return bDirtyDuringLoad;
+}
+
+void IRigVMAssetInterface::RequestRigVMInit() const
+{
+	TArray<UObject*> Instances = GetArchetypeInstances(true, true);
+	for (UObject* Instance : Instances)
+	{
+		if (URigVMHost* RigVMHost = Cast<URigVMHost>(Instance))
+		{
+			RigVMHost->RequestInit();
+		}
+	}
+}
+
+URigVMGraph* IRigVMAssetInterface::GetModel(const UEdGraph* InEdGraph) const
+{
+	const FRigVMClient* RigVMClient = GetRigVMClientHost()->GetRigVMClient();
+#if WITH_EDITORONLY_DATA
+	if (InEdGraph != nullptr && InEdGraph == GetFunctionLibraryEdGraph())
+	{
+		return RigVMClient->GetFunctionLibrary();
+	}
+#endif
+
+	return RigVMClient->GetModel(InEdGraph);
+}
+
+URigVMGraph* IRigVMAssetInterface::GetModel(const FString& InNodePath) const
+{
+	return GetRigVMClientHost()->GetRigVMClient()->GetModel(InNodePath);
+}
+
+URigVMGraph* IRigVMAssetInterface::GetDefaultModel() const
+{
+	return GetRigVMClientHost()->GetRigVMClient()->GetDefaultModel();
+}
+
+TArray<URigVMGraph*> IRigVMAssetInterface::GetAllModels() const
+{
+	return GetRigVMClientHost()->GetRigVMClient()->GetAllModels(true, true);
+}
+
+URigVMFunctionLibrary* IRigVMAssetInterface::GetLocalFunctionLibrary() const
+{
+	return GetRigVMClientHost()->GetRigVMClient()->GetFunctionLibrary();
+}
+
+URigVMFunctionLibrary* IRigVMAssetInterface::GetOrCreateLocalFunctionLibrary(bool bSetupUndoRedo)
+{
+	return GetRigVMClientHost()->GetRigVMClient()->GetOrCreateFunctionLibrary(bSetupUndoRedo);
+}
+
+URigVMGraph* IRigVMAssetInterface::AddModel(FString InName, bool bSetupUndoRedo, bool bPrintPythonCommand)
+{
+	TGuardValue<bool> EnablePythonPrint(bSuspendPythonMessagesForRigVMClient, !bPrintPythonCommand);
+	return GetRigVMClientHost()->GetRigVMClient()->AddModel(InName, bSetupUndoRedo, bPrintPythonCommand);
+}
+
+bool IRigVMAssetInterface::RemoveModel(FString InName, bool bSetupUndoRedo, bool bPrintPythonCommand)
+{
+	TGuardValue<bool> EnablePythonPrint(bSuspendPythonMessagesForRigVMClient, !bPrintPythonCommand);
+	return GetRigVMClientHost()->GetRigVMClient()->RemoveModel(InName, bSetupUndoRedo, bPrintPythonCommand);
+}
+
+FRigVMGetFocusedGraph& IRigVMAssetInterface::OnGetFocusedGraph()
+{
+	return GetRigVMClientHost()->GetRigVMClient()->OnGetFocusedGraph();
+}
+
+const FRigVMGetFocusedGraph& IRigVMAssetInterface::OnGetFocusedGraph() const
+{
+	return GetRigVMClientHost()->GetRigVMClient()->OnGetFocusedGraph();
+}
+
+URigVMHost* IRigVMAssetInterface::GetDebuggedRigVMHost()
+{
+	return Cast<URigVMHost>(GetObjectBeingDebugged());
+}
+
+URigVMGraph* IRigVMAssetInterface::GetFocusedModel() const
+{
+	return GetRigVMClientHost()->GetRigVMClient()->GetFocusedModel();
+}
+
+URigVMController* IRigVMAssetInterface::GetController(const URigVMGraph* InGraph) const
+{
+	return GetRigVMClientHost()->GetRigVMClient()->GetController(InGraph);
+}
+
+URigVMController* IRigVMAssetInterface::GetControllerByName(const FString InGraphName) const
+{
+	return GetRigVMClientHost()->GetRigVMClient()->GetControllerByName(InGraphName);
+}
+
+URigVMController* IRigVMAssetInterface::GetOrCreateController(URigVMGraph* InGraph)
+{
+	return GetRigVMClientHost()->GetRigVMClient()->GetOrCreateController(InGraph);
+}
+
+URigVMController* IRigVMAssetInterface::GetController(const UEdGraph* InEdGraph) const
+{
+	return GetRigVMClientHost()->GetRigVMClient()->GetController(InEdGraph);
+}
+
+URigVMController* IRigVMAssetInterface::GetOrCreateController(const UEdGraph* InEdGraph)
+{
+	return GetRigVMClientHost()->GetRigVMClient()->GetOrCreateController(InEdGraph);
+}
+
+TArray<FString> IRigVMAssetInterface::GeneratePythonCommands(const FString InNewBlueprintName)
+{
+	TArray<FString> InternalCommands;
+
+	if(GetObject()->GetClass() == GetObject()->StaticClass())
+	{
+		InternalCommands.Add(TEXT("import unreal"));
+		InternalCommands.Add(TEXT("unreal.load_module('RigVMDeveloper')"));
+		InternalCommands.Add(TEXT("blueprint = unreal.RigVMBlueprint()"));
+		InternalCommands.Add(TEXT("hierarchy = blueprint.hierarchy"));
+		InternalCommands.Add(TEXT("hierarchy_controller = hierarchy.get_controller()"));
+	}
+
+	InternalCommands.Add(TEXT("library = blueprint.get_local_function_library()"));
+	InternalCommands.Add(TEXT("library_controller = blueprint.get_controller(library)"));
+	InternalCommands.Add(TEXT("blueprint.set_auto_vm_recompile(False)"));
+	
+	// Add variables
+	for (const FRigVMExternalVariable& ExternalVariable : GetExternalVariables(true))
+	{
+		FString CPPType;
+		UObject* CPPTypeObject = nullptr;
+		RigVMTypeUtils::CPPTypeFromExternalVariable(ExternalVariable, CPPType, &CPPTypeObject);
+		if (CPPTypeObject)
+		{
+			if (ExternalVariable.bIsArray)
+			{
+				CPPType = RigVMTypeUtils::ArrayTypeFromBaseType(CPPTypeObject->GetPathName());
+			}
+			else
+			{
+				CPPType = CPPTypeObject->GetPathName();
+			}
+		}
+		// FName AddMemberVariable(const FName& InName, const FString& InCPPType, bool bIsPublic = false, bool bIsReadOnly = false, FString InDefaultValue = TEXT(""));
+		InternalCommands.Add(FString::Printf(TEXT("blueprint.add_member_variable('%s', '%s', %s, %s)"),
+					*ExternalVariable.Name.ToString(),
+					*CPPType,
+					ExternalVariable.bIsPublic ? TEXT("True") : TEXT("False"),
+					ExternalVariable.bIsReadOnly ? TEXT("True") : TEXT("False")));	
+	}
+	
+	// Create graphs
+	{
+		TArray<URigVMGraph*> AllModels = GetAllModels();
+		AllModels.RemoveAll([](const URigVMGraph* GraphToRemove) -> bool
+		{
+			return GraphToRemove->GetTypedOuter<URigVMAggregateNode>() != nullptr;
+		});
+		
+		// Find all graphs to process and sort them by dependencies
+		TArray<URigVMGraph*> ProcessedGraphs;
+		while (ProcessedGraphs.Num() < AllModels.Num())
+		{
+			for (URigVMGraph* Graph : AllModels)
+			{
+				if (ProcessedGraphs.Contains(Graph))
+				{
+					continue;
+				}
+
+				bool bFoundUnprocessedReference = false;
+				for (auto Node : Graph->GetNodes())
+				{
+					if (URigVMFunctionReferenceNode* Reference = Cast<URigVMFunctionReferenceNode>(Node))
+					{
+						if (Reference->GetReferencedFunctionHeader().LibraryPointer.HostObject != GetRigVMClientHost()->GetRigVMGraphFunctionHost().GetObject())
+						{
+							continue;
+						}
+
+						URigVMLibraryNode* LibraryNode = Cast<URigVMLibraryNode>(Reference->GetReferencedFunctionHeader().LibraryPointer.GetNodeSoftPath().ResolveObject());
+						if (!ProcessedGraphs.Contains(LibraryNode->GetContainedGraph()))
+						{
+							bFoundUnprocessedReference = true;
+							break;
+						}
+					}
+					else if (URigVMCollapseNode* CollapseNode = Cast<URigVMCollapseNode>(Node))
+					{
+						if(!CollapseNode->IsA<URigVMAggregateNode>())
+						{
+							if (!ProcessedGraphs.Contains(CollapseNode->GetContainedGraph()))
+							{
+								bFoundUnprocessedReference = true;
+								break;
+							}
+						}
+					}
+				}
+
+				if (!bFoundUnprocessedReference)
+				{
+					ProcessedGraphs.Add(Graph);
+				}
+			}
+		}	
+
+		// Dump python commands for each graph
+		for (URigVMGraph* Graph : ProcessedGraphs)
+		{
+			if (Graph->IsA<URigVMFunctionLibrary>())
+			{
+				continue;
+			}
+
+			URigVMController* Controller = GetController(Graph);
+			if (Graph->GetParentGraph()) 
+			{
+				// Add them all as functions (even collapsed graphs)
+				// The controller will deal with deleting collapsed graph function when it creates the collapse node
+				{						
+					// Add Function
+					InternalCommands.Add(FString::Printf(TEXT("function_%s = library_controller.add_function_to_library('%s', mutable=%s)\ngraph = function_%s.get_contained_graph()"),
+							*RigVMPythonUtils::PythonizeName(Graph->GetGraphName()),
+							*Graph->GetGraphName(),
+							Graph->GetEntryNode()->IsMutable() ? TEXT("True") : TEXT("False"),
+							*RigVMPythonUtils::PythonizeName(Graph->GetGraphName())));
+
+					URigVMLibraryNode* LibraryNode = Cast<URigVMLibraryNode>(Graph->GetOuter());
+
+					InternalCommands.Add(FString::Printf(TEXT("library_controller.set_node_category_by_name('%s', '%s')"),
+						*Graph->GetGraphName(),
+						*LibraryNode->GetNodeCategory()));
+
+					InternalCommands.Add(FString::Printf(TEXT("library_controller.set_node_keywords_by_name('%s', '%s')"),
+						*Graph->GetGraphName(),
+						*LibraryNode->GetNodeKeywords() ));
+
+					InternalCommands.Add(FString::Printf(TEXT("library_controller.set_node_description_by_name('%s', '%s')"),
+						*Graph->GetGraphName(),
+						*LibraryNode->GetNodeDescription()));
+
+					InternalCommands.Add(FString::Printf(TEXT("library_controller.set_node_color_by_name('%s', %s)"),
+						*Graph->GetGraphName(),
+						*RigVMPythonUtils::LinearColorToPythonString(LibraryNode->GetNodeColor()) ));
+					
+					URigVMFunctionEntryNode* EntryNode = Graph->GetEntryNode();
+					URigVMFunctionReturnNode* ReturnNode = Graph->GetReturnNode();
+					
+					// Set Entry and Return nodes in the correct position
+					{
+						//bool SetNodePositionByName(const FName& InNodeName, const FVector2D& InPosition, bool bSetupUndoRedo = true, bool bMergeUndoAction = false);
+						InternalCommands.Add(FString::Printf(TEXT("blueprint.get_controller_by_name('%s').set_node_position_by_name('Entry', unreal.Vector2D(%f, %f))"),
+								*Graph->GetGraphName(),
+								EntryNode->GetPosition().X, 
+								EntryNode->GetPosition().Y));
+
+						InternalCommands.Add(FString::Printf(TEXT("blueprint.get_controller_by_name('%s').set_node_position_by_name('Return', unreal.Vector2D(%f, %f))"),
+								*Graph->GetGraphName(),
+								ReturnNode->GetPosition().X, 
+								ReturnNode->GetPosition().Y));
+					}
+					
+					// Add Exposed Pins
+					{
+
+						bool bHitFirstExecute = false;
+						bool bRenamedExecute = false;
+						for (auto Pin : EntryNode->GetPins())
+						{
+							if (Pin->GetDirection() != ERigVMPinDirection::Output)
+							{
+								continue;
+							}
+
+							if(Pin->IsExecuteContext())
+							{
+								if(!bHitFirstExecute)
+								{
+									bHitFirstExecute = true;
+									if (Pin->GetName() != FRigVMStruct::ExecuteContextName.ToString())
+									{
+										bRenamedExecute = true;
+										InternalCommands.Add(FString::Printf(TEXT("blueprint.get_controller_by_name('%s').rename_exposed_pin('%s', '%s')"),
+											*Graph->GetGraphName(),
+											*FRigVMStruct::ExecuteContextName.ToString(),
+											*Pin->GetName()));
+									}
+									continue;
+								}
+							}
+
+							// FName AddExposedPin(const FName& InPinName, ERigVMPinDirection InDirection, const FString& InCPPType, const FName& InCPPTypeObjectPath, const FString& InDefaultValue, bool bSetupUndoRedo = true);
+							InternalCommands.Add(FString::Printf(TEXT("blueprint.get_controller_by_name('%s').add_exposed_pin('%s', %s, '%s', '%s', '%s')"),
+									*Graph->GetGraphName(),
+									*Pin->GetName(),
+									*RigVMPythonUtils::EnumValueToPythonString<ERigVMPinDirection>((int64)ERigVMPinDirection::Input),
+									*Pin->GetCPPType(),
+									Pin->GetCPPTypeObject() ? *Pin->GetCPPTypeObject()->GetPathName() : TEXT(""),
+									*Pin->GetDefaultValue()));
+						}
+
+						bHitFirstExecute = false;
+						for (auto Pin : ReturnNode->GetPins())
+						{
+							if (Pin->GetDirection() != ERigVMPinDirection::Input)
+							{
+								continue;
+							}
+
+							if(Pin->IsExecuteContext())
+							{
+								if(!bHitFirstExecute)
+								{
+									bHitFirstExecute = true;
+									if (!bRenamedExecute && Pin->GetName() != FRigVMStruct::ExecuteContextName.ToString())
+									{
+										bRenamedExecute = true;
+										InternalCommands.Add(FString::Printf(TEXT("blueprint.get_controller_by_name('%s').rename_exposed_pin('%s', '%s')"),
+											*Graph->GetGraphName(),
+											*FRigVMStruct::ExecuteContextName.ToString(),
+											*Pin->GetName()));
+									}
+									continue;
+								}
+							}
+
+							// FName AddExposedPin(const FName& InPinName, ERigVMPinDirection InDirection, const FString& InCPPType, const FName& InCPPTypeObjectPath, const FString& InDefaultValue, bool bSetupUndoRedo = true);
+							InternalCommands.Add(FString::Printf(TEXT("blueprint.get_controller_by_name('%s').add_exposed_pin('%s', %s, '%s', '%s', '%s')"),
+									*Graph->GetGraphName(),
+									*Pin->GetName(),
+									*RigVMPythonUtils::EnumValueToPythonString<ERigVMPinDirection>((int64)ERigVMPinDirection::Output),
+									*Pin->GetCPPType(),
+									Pin->GetCPPTypeObject() ? *Pin->GetCPPTypeObject()->GetPathName() : TEXT(""),
+									*Pin->GetDefaultValue()));
+						}
+					}
+				}
+			}
+			else if(Graph != GetDefaultModel())
+			{
+				InternalCommands.Add(FString::Printf(TEXT("blueprint.add_model('%s')"),
+						*Graph->GetName()));
+			}
+			
+			InternalCommands.Append(Controller->GeneratePythonCommands());
+		}
+	}
+
+	InternalCommands.Add(TEXT("blueprint.set_auto_vm_recompile(True)"));
+
+	// Split multiple commands into different array elements
+	TArray<FString> InnerFunctionCmds;
+	for (FString Cmd : InternalCommands)
+	{
+		FString Left, Right=Cmd;
+		while (Right.Split(TEXT("\n"), &Left, &Right))
+		{
+			InnerFunctionCmds.Add(Left);
+		}
+		InnerFunctionCmds.Add(Right);
+	}
+
+	// Define a function and insert all the commands
+	// We do not want to pollute the global state with our definitions
+	TArray<FString> Commands;
+	Commands.Add(FString::Printf(TEXT("import unreal\n"
+		"def create_asset():\n")));
+	for (const FString& InnerCmd : InnerFunctionCmds)
+	{
+		Commands.Add(FString::Printf(TEXT("\t%s"), *InnerCmd));
+	}
+
+	Commands.Add(TEXT("create_asset()\n"));
+	return Commands;
+}
+
+TArray<FRigVMExternalDependency> IRigVMAssetInterface::GetExternalDependenciesForCategory(const FName& InCategory) const
+{
+	TArray<FRigVMExternalDependency> Dependencies;
+	if(const FRigVMClient* Client = GetRigVMClientHost()->GetRigVMClient())
+	{
+		CollectExternalDependencies(Dependencies, InCategory, Client);
+	}
+	if(const TScriptInterface<IRigVMGraphFunctionHost> FunctionHost = GetRigVMClientHost()->GetRigVMGraphFunctionHost())
+	{
+		if(const FRigVMGraphFunctionStore* FunctionStore = FunctionHost->GetRigVMGraphFunctionStore())
+		{
+			CollectExternalDependencies(Dependencies, InCategory, FunctionStore);
+		}
+	}
+
+#if WITH_EDITOR
+	const TArray<FRigVMGraphVariableDescription> MemberVariables = GetAssetVariables();
+	for(const FRigVMGraphVariableDescription& MemberVariable : MemberVariables)
+	{
+		CollectExternalDependenciesForCPPTypeObject(Dependencies, InCategory, MemberVariable.CPPTypeObject.Get());
+	}
+#endif
+	return Dependencies;
+}
+
+#if WITH_EDITOR
+void IRigVMAssetInterface::AddVariableSearchMetaDataInfo(const FName InVariableName, TArray<UBlueprintExtension::FSearchTagDataPair>& OutTaggedMetaData) const
+{
+	for (const FRigVMExternalVariable& Variable : GetExternalVariables(true))
+	{
+		if (Variable.Name.IsEqual(InVariableName))
+		{
+			OutTaggedMetaData.Emplace(FRigVMSearchTags::FiB_Name, FText::FromName(InVariableName));
+
+			const FEdGraphPinType& Type = RigVMTypeUtils::PinTypeFromExternalVariable(Variable);
+			OutTaggedMetaData.Emplace(FRigVMSearchTags::FiB_PinCategory, FText::FromName(Type.PinCategory));
+			OutTaggedMetaData.Emplace(FRigVMSearchTags::FiB_PinSubCategory, FText::FromName(Type.PinSubCategory));
+			if (Type.PinSubCategoryObject.IsValid())
+			{
+				OutTaggedMetaData.Emplace(FRigVMSearchTags::FiB_ObjectClass, FText::FromString(Type.PinSubCategoryObject->GetPathName()));
+			}
+			OutTaggedMetaData.Emplace(FRigVMSearchTags::FiB_IsArray, FText::Format(LOCTEXT("RigVMNodePinIsArray", "{0}"), Type.IsArray() ? 1 : 0));
+		}
+	}
+}
+#endif
+
+URigVMGraph* IRigVMAssetInterface::GetTemplateModel(bool bIsFunctionLibrary)
+{
+#if WITH_EDITORONLY_DATA
+	if (TemplateModel == nullptr)
+	{
+		if (bIsFunctionLibrary)
+		{
+			TemplateModel = NewObject<URigVMFunctionLibrary>(GetObject(), TEXT("TemplateFunctionLibrary"));
+		}
+		else
+		{
+			TemplateModel = NewObject<URigVMGraph>(GetObject(), TEXT("TemplateModel"));
+		}
+		TemplateModel->SetFlags(RF_Transient);
+		TemplateModel->SetExecuteContextStruct(GetRigVMClientHost()->GetRigVMClient()->GetDefaultExecuteContextStruct());
+	}
+	return TemplateModel;
+#else
+	return nullptr;
+#endif
+}
+
+URigVMController* IRigVMAssetInterface::GetTemplateController(bool bIsFunctionLibrary)
+{
+#if WITH_EDITORONLY_DATA
+	if (TemplateController == nullptr)
+	{
+		TemplateController = NewObject<URigVMController>(GetObject(), TEXT("TemplateController"));
+		TemplateController->SetGraph(GetTemplateModel(bIsFunctionLibrary));
+		TemplateController->EnableReporting(false);
+		TemplateController->SetFlags(RF_Transient);
+		TemplateController->SetSchemaClass(GetRigVMClientHost()->GetRigVMClient()->GetDefaultSchemaClass());
+	}
+	return TemplateController;
+#else
+	return nullptr;
+#endif
+}
+
+UEdGraph* IRigVMAssetInterface::GetEdGraph(const URigVMGraph* InModel) const
+{
+	return Cast<UEdGraph>(GetEditorObjectForRigVMGraph(InModel));
+}
+
+UEdGraph* IRigVMAssetInterface::GetEdGraph(const FString& InNodePath) const
+{
+	if (URigVMGraph* ModelForNodePath = GetModel(InNodePath))
+	{
+		return GetEdGraph(ModelForNodePath);
+	}
+	return nullptr;
+}
+
+bool IRigVMAssetInterface::IsFunctionPublic(const FName& InFunctionName) const
+{
+	return GetLocalFunctionLibrary()->IsFunctionPublic(InFunctionName);	
+}
+
+void IRigVMAssetInterface::MarkFunctionPublic(const FName& InFunctionName, bool bIsPublic)
+{
+	if(IsFunctionPublic(InFunctionName) == bIsPublic)
+	{
+		return;
+	}
+	
+	URigVMController* Controller = GetRigVMClientHost()->GetRigVMClient()->GetOrCreateController(GetLocalFunctionLibrary());
+	Controller->MarkFunctionAsPublic(InFunctionName, bIsPublic);
+}
+
+TArray<IRigVMAssetInterface*> IRigVMAssetInterface::GetDependencies(bool bRecursive) const
+{
+	TArray<IRigVMAssetInterface*> Dependencies;
+
+	TArray<URigVMGraph*> Graphs = GetAllModels();
+	for(URigVMGraph* Graph : Graphs)
+	{
+		for(URigVMNode* Node : Graph->GetNodes())
+		{
+			if(URigVMFunctionReferenceNode* FunctionReferenceNode = Cast<URigVMFunctionReferenceNode>(Node))
+			{
+				if(const URigVMLibraryNode* LibraryNode = FunctionReferenceNode->LoadReferencedNode())
+				{
+					if(IRigVMAssetInterface* DependencyBlueprint = LibraryNode->GetImplementingOuter<IRigVMAssetInterface>())
+					{
+						if(DependencyBlueprint != this)
+						{
+							if(!Dependencies.Contains(DependencyBlueprint))
+							{
+								Dependencies.Add(DependencyBlueprint);
+
+								if(bRecursive)
+								{
+									TArray<IRigVMAssetInterface*> ChildDependencies = DependencyBlueprint->GetDependencies(true);
+									for(IRigVMAssetInterface* ChildDependency : ChildDependencies)
+									{
+										Dependencies.AddUnique(ChildDependency);
+									}
+								}
+							}
+						}
+					}
+				}
+			}
+		}
+	}
+
+	return Dependencies;
+}
+
+TArray<FAssetData> IRigVMAssetInterface::GetDependentAssets() const
+{
+	TArray<FAssetData> Dependents;
+	TArray<FSoftObjectPath> AssetPaths;
+
+	if(URigVMFunctionLibrary* FunctionLibrary = GetRigVMClientHost()->GetRigVMClient()->GetFunctionLibrary())
+	{
+		IAssetRegistry& AssetRegistry = IAssetRegistry::GetChecked();
+		TArray<URigVMLibraryNode*> Functions = FunctionLibrary->GetFunctions();
+		for(URigVMLibraryNode* Function : Functions)
+		{
+			const FName FunctionName = Function->GetFName();
+			if(IsFunctionPublic(FunctionName))
+			{
+				TArray<TSoftObjectPtr<URigVMFunctionReferenceNode>> References = FunctionLibrary->GetReferencesForFunction(FunctionName);
+				for(const TSoftObjectPtr<URigVMFunctionReferenceNode>& Reference : References)
+				{
+					if (const URigVMFunctionReferenceNode* ReferencePtr = Reference.Get())
+					{
+						if (const TRigVMAssetInterfacePtr ControlRigBlueprint = IRigVMAssetInterface::GetInterfaceOuter(ReferencePtr))
+						{
+							const TSoftObjectPtr<UObject> Blueprint = ControlRigBlueprint.GetObject();
+							const FSoftObjectPath AssetPath = Blueprint.ToSoftObjectPath();
+							if(AssetPath.GetLongPackageName().StartsWith(TEXT("/Engine/Transient")))
+							{
+								continue;
+							}
+				
+							if(!AssetPaths.Contains(AssetPath))
+							{
+								AssetPaths.Add(AssetPath);
+
+								const FAssetData AssetData = AssetRegistry.GetAssetByObjectPath(AssetPath);
+								if(AssetData.IsValid())
+								{
+									Dependents.Add(AssetData);
+								}
+							}
+						}
+					}
+				}
+			}
+		}
+	}
+
+	return Dependents;
+}
+
+TArray<IRigVMAssetInterface*> IRigVMAssetInterface::GetDependentResolvedAssets(bool bRecursive, bool bOnlyLoaded) const
+{
+	TArray<FAssetData> Assets = GetDependentAssets();
+	TArray<IRigVMAssetInterface*> Dependents;
+
+	for(const FAssetData& Asset : Assets)
+	{
+		if (!bOnlyLoaded || Asset.IsAssetLoaded())
+		{
+			if(TRigVMAssetInterfacePtr Dependent = Asset.GetAsset())
+			{
+				if(!Dependents.Contains(Dependent.GetInterface()))
+				{
+					Dependents.Add(Dependent.GetInterface());
+
+					if(bRecursive && Dependent != this)
+					{
+						TArray<IRigVMAssetInterface*> ParentDependents = Dependent->GetDependentResolvedAssets(true);
+						for(IRigVMAssetInterface* ParentDependent : ParentDependents)
+						{
+							Dependents.AddUnique(ParentDependent);
+						}
+					}
+				}
+			}
+		}
+	}
+
+	return Dependents;
+}
+
+void IRigVMAssetInterface::BroadcastRefreshEditor()
+{
+	return RefreshEditorEvent.Broadcast(GetObject());
+}
+
+FOnRigVMRequestInspectMemoryStorage& IRigVMAssetInterface::OnRequestInspectMemoryStorage()
+{
+	return OnRequestInspectMemoryStorageEvent;
+}
+
+void IRigVMAssetInterface::RequestInspectMemoryStorage(const TArray<FRigVMMemoryStorageStruct*>& InMemoryStorageStructs) const
+{
+	OnRequestInspectMemoryStorageEvent.Broadcast(InMemoryStorageStructs);
+}
+
+void IRigVMAssetInterface::SetObjectBeingDebugged(UObject* NewObject)
+{
+	URigVMHost* PreviousRigBeingDebugged = Cast<URigVMHost>(GetObjectBeingDebugged());
+	if (PreviousRigBeingDebugged && PreviousRigBeingDebugged != NewObject)
+	{
+		PreviousRigBeingDebugged->DrawInterface.Reset();
+		PreviousRigBeingDebugged->RigVMLog = nullptr;
+#if WITH_EDITOR
+		PreviousRigBeingDebugged->bIsBeingDebugged = false;
+#endif
+	}
+
+	SetObjectBeingDebuggedSuper(NewObject);
+
+#if WITH_EDITOR
+	if(URigVMHost* NewRigBeingDebugged = Cast<URigVMHost>(NewObject))
+	{
+		NewRigBeingDebugged->bIsBeingDebugged = true;
+	}
+#endif
+}
+
+void IRigVMAssetInterface::PostTransacted(const FTransactionObjectEvent& TransactionEvent)
+{
+	DECLARE_SCOPE_HIERARCHICAL_COUNTER_FUNC()
+
+	if (TransactionEvent.GetEventType() == ETransactionObjectEventType::UndoRedo)
+	{
+		// The action stack undo/redo transaction should always execute first
+		// It already knows whether or not it has already executed or not
+		GetRigVMClientHost()->GetRigVMClient()->GetOrCreateActionStack()->PostTransacted(TransactionEvent);
+	}
+	
+	PostTransactedSuper(TransactionEvent);
+
+	if (TransactionEvent.GetEventType() == ETransactionObjectEventType::UndoRedo)
+	{
+		TArray<FName> PropertiesChanged = TransactionEvent.GetChangedProperties();
+
+		if (PropertiesChanged.Contains(TEXT("VMRuntimeSettings")))
+		{
+			PropagateRuntimeSettingsFromBPToInstances();
+		}
+
+		if (PropertiesChanged.Contains(TEXT("NewVariables")))
+		{
+			if (RefreshEditorEvent.IsBound())
+			{
+				RefreshEditorEvent.Broadcast(GetObject());
+			}
+			(void)MarkPackageDirty();
+		}
+
+		if (PropertiesChanged.Contains(TEXT("RigVMClient")) ||
+			PropertiesChanged.Contains(TEXT("UbergraphPages")))
+		{
+			GetUberGraphs().RemoveAll([](const UEdGraph* UberGraph) -> bool
+			{
+ 				return UberGraph == nullptr || !IsValid(UberGraph);
+			});
+			GetRigVMClientHost()->GetRigVMClient()->PostTransacted(TransactionEvent);
+			
+			RecompileVM();
+			(void)MarkPackageDirty();
+		}
+	}
+}
+
+void IRigVMAssetInterface::ReplaceDeprecatedNodes()
+{
+	TArray<UEdGraph*> EdGraphs;
+	GetAllEdGraphs(EdGraphs);
+
+	for (UEdGraph* EdGraph : EdGraphs)
+	{
+		EdGraph->Schema = GetRigVMClientHost()->GetRigVMEdGraphSchemaClass();
+	}
+
+	ReplaceDeprecatedNodesSuper();
+}
+
+void IRigVMAssetInterface::PreDuplicate(FObjectDuplicationParameters& DupParams)
+{
+	PreDuplicateSuper(DupParams);
+	PreDuplicateAssetPath = GetObject()->GetPathName();
+	PreDuplicateHostPath = GetRigVMClientHost()->GetRigVMGraphFunctionHost().GetObject();
+
+	// look for graphs which are no longer reachable and remove them. this can happen in case
+	// one of the default sub objects got renamed.
+	FRigVMClient* RigVMClient = GetRigVMClientHost()->GetRigVMClient();
+	TArray<URigVMGraph*> AllModels = RigVMClient->GetAllModels(true, true);
+	for (TObjectIterator<URigVMGraph> RigVMIt; RigVMIt; ++RigVMIt)
+	{
+		URigVMGraph* Graph = *RigVMIt;
+		if (Graph->IsInOuter(GetObject()))
+		{
+			if (!AllModels.Contains(Graph))
+			{
+				if (URigVMController* Controller = RigVMClient->GetOrCreateController(RigVMClient->GetDefaultModel()))
+				{
+					Controller->DestroyObject(Graph);
+				}
+			}
+		}
+	}
+}
+
+void IRigVMAssetInterface::ReplaceFunctionIdentifiers(const FString& InOldAssetPath, const FString& InNewAssetPath)
+{
+	if (!InOldAssetPath.Equals(GetObject()->GetPathName()))
+	{
+		const FString OldLibraryPath = InOldAssetPath + TEXT(":");
+		const FString NewLibraryPath = InNewAssetPath + TEXT(":");
+		const FString OldHostPath = InOldAssetPath + TEXT("_C");
+		const FString NewHostPath = InNewAssetPath + TEXT("_C");
+
+		auto ReplaceIdentifier = [OldLibraryPath, NewLibraryPath, OldHostPath, NewHostPath](FRigVMGraphFunctionIdentifier& Identifier)
+		{
+			FString& LibraryNodePath = Identifier.GetLibraryNodePath();
+			FSoftObjectPath& HostPath = Identifier.HostObject;
+			FString HostPathStr = HostPath.ToString();
+			if(LibraryNodePath.StartsWith(OldLibraryPath, ESearchCase::CaseSensitive))
+			{
+				LibraryNodePath = NewLibraryPath + LibraryNodePath.Mid(OldLibraryPath.Len());
+			}
+			if(HostPathStr.StartsWith(OldHostPath, ESearchCase::CaseSensitive))
+			{
+				HostPathStr = NewHostPath + HostPathStr.Mid(OldHostPath.Len());
+				HostPath = HostPathStr;
+			}
+		};
+
+		// Replace identifiers in store
+		if(TScriptInterface<IRigVMGraphFunctionHost> CRGeneratedClass = GetRigVMClientHost()->GetRigVMGraphFunctionHost())
+		{
+			FRigVMGraphFunctionStore* Store = CRGeneratedClass->GetRigVMGraphFunctionStore();
+			for (int32 i=0; i<2; ++i)
+			{
+				TArray<FRigVMGraphFunctionData>& Functions = (i == 0) ? Store->PublicFunctions : Store->PrivateFunctions;
+				for (FRigVMGraphFunctionData& Data : Functions)
+				{
+					ReplaceIdentifier(Data.Header.LibraryPointer);
+					for (TPair<FRigVMGraphFunctionIdentifier, uint32>& Pair : Data.Header.Dependencies)
+					{
+						ReplaceIdentifier(Pair.Key);
+					}
+				}
+			}
+		}
+
+		// Replace identifiers in function references
+		TArray<URigVMGraph*> AllModels = GetRigVMClientHost()->GetRigVMClient()->GetAllModels(true, true);
+		for(URigVMGraph* Model : AllModels)
+		{
+			for (URigVMNode* Node : Model->GetNodes())
+			{
+				if (URigVMFunctionReferenceNode* FunctionReferenceNode = Cast<URigVMFunctionReferenceNode>(Node))
+				{
+					ReplaceIdentifier(FunctionReferenceNode->ReferencedFunctionHeader.LibraryPointer);
+				}
+			}
+		}
+	}
+}
+
+void IRigVMAssetInterface::PostDuplicate(bool bDuplicateForPIE)
+{
+	// assuming PostDuplicate is always followed by a PostLoad:
+	// so theoretically, PostDuplicate just makes corrections to the serialized data and does nothing more,
+	// while PostLoad looks at whatever is serialized and load it into memory according to the version of the editor used
+	// note: how to know if we have corrected everything?
+	// ans: check the reference viewer for the duplicated BP and make sure that the original BP does not appear in there
+	
+	{
+		// pause compilation because we need to patch some stuff first
+		TGuardValue<bool> CompilingGuard(bIsCompiling, true);
+		// this will create the new EMPTY generated class to be used as the function store for this BP
+		// it will be filled during PostLoad based on the graph model
+		PostDuplicateSuper(bDuplicateForPIE);
+	}
+
+	const FString OldAssetPath = PreDuplicateAssetPath.ToString();
+	const FString NewAssetPath = GetObject()->GetPathName();
+	
+	ReplaceFunctionIdentifiers(OldAssetPath, NewAssetPath);
+
+	PreDuplicateAssetPath.Reset();
+	PreDuplicateHostPath.Reset();
+
+	SplitAssetVariant();
+
+	TArray<UEdGraph*> EdGraphs;
+	GetAllEdGraphs(EdGraphs);
+	for (UEdGraph* EdGraph : EdGraphs)
+	{
+		if (URigVMEdGraph* RigVMEdGraph = Cast<URigVMEdGraph>(EdGraph))
+		{
+			RigVMEdGraph->CachedModelGraph.Reset();
+		}
+	}
+
+	MarkAssetAsStructurallyModified();
+}
+
+void IRigVMAssetInterface::GetAssetRegistryTags(FAssetRegistryTagsContext Context) const
+{
+	if (CachedAssetTags.IsEmpty())
+	{
+		GetAssetRegistryTagsSuper(Context);
+		CachedAssetTags.Reset(Context.GetNumTags());
+		Context.EnumerateTags([this](const UObject::FAssetRegistryTag& Tag)
+			{
+				CachedAssetTags.Add(Tag);
+			});
+	}
+	else
+	{
+		for (const UObject::FAssetRegistryTag& Tag : CachedAssetTags)
+		{
+			Context.AddTag(Tag);
+		}
+	}
+}
+
+FRigVMGraphModifiedEvent& IRigVMAssetInterface::OnModified()
+{
+	return ModifiedEvent;
+}
+
+
+FOnRigVMCompiledEvent& IRigVMAssetInterface::OnVMCompiled()
+{
+	return VMCompiledEvent;
+}
+
+URigVMHost* IRigVMAssetInterface::CreateRigVMHostImpl()
+{
+	RecompileVMIfRequired();
+
+	URigVMHost* Host = CreateRigVMHostSuper(GetObject());
+	Host->Initialize(true);
+	return Host;
+}
+
+TArray<UStruct*> IRigVMAssetInterface::GetAvailableRigVMStructs() const
+{
+	TArray<UStruct*> Structs;
+	UStruct* BaseStruct = FRigVMStruct::StaticStruct();
+
+	for (const FRigVMFunction& Function : FRigVMRegistry::Get().GetFunctions())
+	{
+		if (Function.Struct)
+		{
+			if (Function.Struct->IsChildOf(BaseStruct))
+			{
+				Structs.Add(Function.Struct);
+				// todo: filter by available types
+				// todo: filter by execute context
+			}
+		}
+	}
+
+	return Structs;
+}
+
+#if WITH_EDITOR
+
+
+FRigVMVariantRef IRigVMAssetInterface::GetAssetVariantRefImpl() const
+{
+	return FRigVMVariantRef(FSoftObjectPath(GetObject()), GetAssetVariant());
+}
+
+bool IRigVMAssetInterface::SplitAssetVariantImpl()
+{
+	if(GetMatchingVariants().IsEmpty())
+	{
+		return false;
+	}
+
+	FScopedTransaction Transaction(LOCTEXT("SplitAssetVariant", "Split Asset Variant"));
+	GetObject()->Modify();
+
+	// prefer the path based (deterministic) guid - and fall back on random.
+	const FGuid PathBasedGuid = FRigVMVariant::GenerateGUID(GetObject()->GetPathName());
+	if(PathBasedGuid != GetAssetVariant().Guid)
+	{
+		GetAssetVariant().Guid = PathBasedGuid;
+	}
+	else
+	{
+		GetAssetVariant().Guid = FRigVMVariant::GenerateGUID();
+	}
+	
+	return true;
+}
+
+bool IRigVMAssetInterface::JoinAssetVariantImpl(const FGuid& InGuid)
+{
+	if(GetAssetVariant().Guid != InGuid)
+	{
+		FScopedTransaction Transaction(LOCTEXT("JoinAssetVariant", "Join Asset Variant"));
+		GetObject()->Modify();
+		
+		GetAssetVariant().Guid = InGuid;
+		return true;
+	}
+
+	return false;
+}
+
+TArray<FRigVMVariantRef> IRigVMAssetInterface::GetMatchingVariantsImpl() const
+{
+	if(URigVMBuildData* BuildData = URigVMBuildData::Get())
+	{
+		TArray<FRigVMVariantRef> Variants = BuildData->FindAssetVariantRefs(GetAssetVariant().Guid);
+		const FRigVMVariantRef MyVariantRef = FRigVMVariantRef(GetObject()->GetPathName(), GetAssetVariant());
+		Variants.RemoveAll([MyVariantRef](const FRigVMVariantRef& VariantRef) -> bool
+		{
+			return VariantRef == MyVariantRef;
+		});
+		return Variants;
+	}
+	return TArray<FRigVMVariantRef>();
+}
+
+#endif
+
+void IRigVMAssetInterface::RebuildGraphFromModel()
+{
+	DECLARE_SCOPE_HIERARCHICAL_COUNTER_FUNC()
+
+	TGuardValue<bool> SelfGuard(bSuspendModelNotificationsForSelf, true);
+	TGuardValue<bool> ClientIgnoreModificationsGuard(GetRigVMClientHost()->GetRigVMClient()->bIgnoreModelNotifications, true);
+	
+	verify(GetOrCreateController());
+
+	TArray<UEdGraph*> EdGraphs;
+	GetAllEdGraphs(EdGraphs);
+
+	for (UEdGraph* Graph : EdGraphs)
+	{
+		TArray<UEdGraphNode*> Nodes = Graph->Nodes;
+		for (UEdGraphNode* Node : Nodes)
+		{
+			Graph->RemoveNode(Node);
+		}
+
+		if (URigVMEdGraph* RigGraph = Cast<URigVMEdGraph>(Graph))
+		{
+			if (RigGraph->bIsFunctionDefinition)
+			{
+				GetFunctionGraphs().Remove(RigGraph);
+			}
+		}
+	}
+
+	if(GetFunctionLibraryEdGraph() && GetRigVMClientHost()->GetRigVMClient()->GetFunctionLibrary())
+	{
+		GetFunctionLibraryEdGraph()->ModelNodePath = GetRigVMClientHost()->GetRigVMClient()->GetFunctionLibrary()->GetNodePath();
+	}
+
+	TArray<URigVMGraph*> RigGraphs = GetRigVMClientHost()->GetRigVMClient()->GetAllModels(true, true);
+
+	for (int32 RigGraphIndex = 0; RigGraphIndex < RigGraphs.Num(); RigGraphIndex++)
+	{
+		GetOrCreateController(RigGraphs[RigGraphIndex])->ResendAllNotifications();
+	}
+
+	for (int32 RigGraphIndex = 0; RigGraphIndex < RigGraphs.Num(); RigGraphIndex++)
+	{
+		URigVMGraph* RigGraph = RigGraphs[RigGraphIndex];
+
+		for (URigVMNode* RigNode : RigGraph->GetNodes())
+		{
+			if (URigVMCollapseNode* CollapseNode = Cast<URigVMCollapseNode>(RigNode))
+			{
+				CreateEdGraphForCollapseNodeIfNeeded(CollapseNode, true);
+			}
+		}
+	}
+}
+
+void IRigVMAssetInterface::Notify(ERigVMGraphNotifType InNotifType, UObject* InSubject)
+{
+	GetOrCreateController()->Notify(InNotifType, InSubject);
+}
+
+void IRigVMAssetInterface::HandleModifiedEvent(ERigVMGraphNotifType InNotifType, URigVMGraph* InGraph, UObject* InSubject)
+{
+	DECLARE_SCOPE_HIERARCHICAL_COUNTER_FUNC()
+
+#if WITH_EDITOR
+
+	if (bSuspendAllNotifications)
+	{
+		return;
+	}
+
+	// since it's possible that a notification will be already sent / forwarded to the
+	// listening objects within the switch statement below - we keep a flag to mark
+	// the notify for still pending (or already sent)
+	bool bNotifForOthersPending = true;
+
+	auto MarkBlueprintAsStructurallyModified = [this]()
+	{
+		if(VMRecompilationBracket == 0)
+		{
+			if(bMarkBlueprintAsStructurallyModifiedPending)
+			{
+				bMarkBlueprintAsStructurallyModifiedPending = false;
+				MarkAssetAsStructurallyModified(bSkipDirtyBlueprintStatus);
+			}
+		}
+		else
+		{
+			bMarkBlueprintAsStructurallyModifiedPending = true;
+		}
+	};
+
+	if (!bSuspendModelNotificationsForSelf)
+	{
+		switch (InNotifType)
+		{
+			case ERigVMGraphNotifType::InteractionBracketOpened:
+			{
+				IncrementVMRecompileBracket();
+				break;
+			}
+			case ERigVMGraphNotifType::InteractionBracketClosed:
+			case ERigVMGraphNotifType::InteractionBracketCanceled:
+			{
+				DecrementVMRecompileBracket();
+				MarkBlueprintAsStructurallyModified();
+				break;
+			}
+			case ERigVMGraphNotifType::PinDefaultValueChanged:
+			{
+				if (URigVMPin* Pin = Cast<URigVMPin>(InSubject))
+				{
+					bool bRequiresRecompile = false;
+
+					URigVMPin* RootPin = Pin->GetRootPin();
+					static const FString ConstSuffix = TEXT(":Const");
+					const FString PinHash = RootPin->GetPinPath(true) + ConstSuffix;
+					
+					if (const FRigVMOperand* Operand = PinToOperandMap.Find(PinHash))
+					{
+						FRigVMASTProxy RootPinProxy = FRigVMASTProxy::MakeFromUObject(RootPin);
+						if(const FRigVMExprAST* Expression = InGraph->GetRuntimeAST()->GetExprForSubject(RootPinProxy))
+						{
+							bRequiresRecompile = Expression->NumParents() > 1;
+						}
+						else
+						{
+							bRequiresRecompile = true;
+						}
+
+						// If we are only changing a pin's default value, we need to
+						// check if there is a connection to a sub-pin of the root pin
+						// that has its value is directly stored in the root pin due to optimization, if so,
+						// we want to recompile to make sure the pin's new default value and values from other connections
+						// are both applied to the root pin because GetDefaultValue() alone cannot account for values
+						// from other connections.
+						if(!bRequiresRecompile)
+						{
+							TArray<URigVMPin*> SourcePins = RootPin->GetLinkedSourcePins(true);
+							for (const URigVMPin* SourcePin : SourcePins)
+							{
+								// check if the source node is optimized out, if so, only a recompile will allows us
+								// to re-query its value.
+								FRigVMASTProxy SourceNodeProxy = FRigVMASTProxy::MakeFromUObject(SourcePin->GetNode());
+								if (InGraph->GetRuntimeAST()->GetExprForSubject(SourceNodeProxy) == nullptr)
+								{
+									bRequiresRecompile = true;
+									break;
+								}
+							}
+						} 
+						
+						if(!bRequiresRecompile)
+						{
+							const FString DefaultValue = RootPin->GetDefaultValue();
+
+							URigVM* VM = GetVM(true);
+							FRigVMExtendedExecuteContext* Context = GetRigVMExtendedExecuteContext();
+							if (VM != nullptr && Context != nullptr)
+							{
+								VM->SetPropertyValueFromString(*Context, *Operand, DefaultValue);
+							}
+
+							TArray<UObject*> ArchetypeInstances = GetArchetypeInstances(false, false);
+							for (UObject* ArchetypeInstance : ArchetypeInstances)
+							{
+								if (URigVMHost* InstancedHost = Cast<URigVMHost>(ArchetypeInstance))
+								{
+									if (!InstancedHost->HasAllFlags(RF_ClassDefaultObject))
+									{
+										if (InstancedHost->GetVM())
+										{
+											InstancedHost->VM->SetPropertyValueFromString(InstancedHost->GetRigVMExtendedExecuteContext(), *Operand, DefaultValue);
+										}
+									}
+								}
+							}
+
+							if (Pin->IsDefinedAsConstant() || Pin->GetRootPin()->IsDefinedAsConstant())
+							{
+								// re-init the rigs
+								GetRigVMClientHost()->RequestRigVMInit();
+								bRequiresRecompile = true;
+							}
+						}
+					}
+					else
+					{
+						bRequiresRecompile = true;
+					}
+				
+					if(bRequiresRecompile)
+					{
+						RequestAutoVMRecompilation();
+					}
+				}
+				(void)MarkPackageDirty();
+				break;
+			}
+			case ERigVMGraphNotifType::NodeAdded:
+			case ERigVMGraphNotifType::NodeRemoved:
+			{
+				bool bAdded = InNotifType == ERigVMGraphNotifType::NodeAdded;
+				if (!bAdded)
+				{
+					if (URigVMNode* RigVMNode = Cast<URigVMNode>(InSubject))
+					{
+						RemoveBreakpoint(RigVMNode);
+					}
+				}
+
+				if (URigVMCollapseNode* CollapseNode = Cast<URigVMCollapseNode>(InSubject))
+				{
+					if (bAdded)
+					{
+						// If the controller for this graph already exist, make sure it is referencing the correct graph
+						if (URigVMController* Controller = GetRigVMClientHost()->GetRigVMClient()->GetController(CollapseNode->GetContainedGraph()))
+						{
+							Controller->SetGraph(CollapseNode->GetContainedGraph());
+						}
+						
+						CreateEdGraphForCollapseNodeIfNeeded(CollapseNode);
+					}
+					else
+					{
+						bNotifForOthersPending = !RemoveEdGraphForCollapseNode(CollapseNode, true);
+
+						// Cannot remove from the Controllers array because we would lose the action stack on that graph
+						// Controllers.Remove(CollapseNode->GetContainedGraph();
+					}
+
+					RequestAutoVMRecompilation();
+
+					(void)MarkPackageDirty();
+					MarkBlueprintAsStructurallyModified();
+					break;
+				}
+
+				if (URigVMNode* RigVMNode = Cast<URigVMNode>(InSubject))
+				{
+					if(RigVMNode->IsEvent() && RigVMNode->GetGraph()->IsRootGraph())
+					{
+						// let the UI know the title for the graph may have changed.
+						GetRigVMClientHost()->GetRigVMClient()->NotifyOuterOfPropertyChange();
+
+						if(URigVMEdGraph* EdGraph = Cast<URigVMEdGraph>(GetEdGraph(RigVMNode->GetGraph())))
+						{
+							// decide if this graph should be renameable
+							const int32 NumberOfEvents = static_cast<int32>(Algo::CountIf(RigVMNode->GetGraph()->GetNodes(), [](const URigVMNode* NodeToCount) -> bool
+							{
+								return NodeToCount->IsEvent() && NodeToCount->CanOnlyExistOnce();
+							}));
+							EdGraph->bAllowRenaming = NumberOfEvents != 1;
+						}
+					}
+				}
+				// fall through to the next case
+			}
+			case ERigVMGraphNotifType::LinkAdded:
+			case ERigVMGraphNotifType::LinkRemoved:
+			case ERigVMGraphNotifType::PinArraySizeChanged:
+			case ERigVMGraphNotifType::PinDirectionChanged:
+			{
+				RequestAutoVMRecompilation();
+				(void)MarkPackageDirty();
+
+				// we don't need to mark the blueprint as modified since we only
+				// need to recompile the VM here - unless we don't auto recompile.
+				if(!bAutoRecompileVM)
+				{
+					MarkBlueprintAsStructurallyModified();
+				}
+				break;
+			}
+			case ERigVMGraphNotifType::PinWatchedChanged:
+			{
+				if (URigVMHost* DebuggedHost = Cast<URigVMHost>(GetObjectBeingDebugged()))
+				{
+					URigVMPin* Pin = CastChecked<URigVMPin>(InSubject)->GetRootPin(); 
+					URigVMCompiler* Compiler = URigVMCompiler::StaticClass()->GetDefaultObject<URigVMCompiler>();
+
+					TSharedPtr<FRigVMParserAST> RuntimeAST = GetDefaultModel()->GetRuntimeAST();
+					
+					if(Pin->RequiresWatch())
+					{
+						// check if the node is optimized out - in that case we need to recompile
+						if(DebuggedHost->GetVM()->GetByteCode().GetFirstInstructionIndexForSubject(Pin->GetNode()) == INDEX_NONE)
+						{
+							RequestAutoVMRecompilation();
+							(void)MarkPackageDirty();
+						}
+						else
+						{
+							if(DebuggedHost->GetDebugMemory()->Num() == 0)
+							{
+								RequestAutoVMRecompilation();
+								(void)MarkPackageDirty();
+							}
+							else
+							{
+								Compiler->MarkDebugWatch(GetVMCompileSettings(), true, Pin, DebuggedHost->GetVM(), &PinToOperandMap, RuntimeAST);
+							}
+						}
+					}
+					else
+					{
+						Compiler->MarkDebugWatch(GetVMCompileSettings(), false, Pin, DebuggedHost->GetVM(), &PinToOperandMap, RuntimeAST);
+					}
+				}
+				// break; fall through
+			}
+			case ERigVMGraphNotifType::PinTypeChanged:
+			case ERigVMGraphNotifType::PinIndexChanged:
+			{
+				if (URigVMPin* ModelPin = Cast<URigVMPin>(InSubject))
+				{
+					if (UEdGraph* EdGraph = GetEdGraph(InGraph))
+					{							
+						if (URigVMEdGraph* Graph = Cast<URigVMEdGraph>(EdGraph))
+						{
+							if (UEdGraphNode* EdNode = Graph->FindNodeForModelNodeName(ModelPin->GetNode()->GetFName()))
+							{
+								if (UEdGraphPin* EdPin = EdNode->FindPin(*ModelPin->GetPinPath()))
+								{
+									if (ModelPin->RequiresWatch())
+									{
+										AddPinWatch(EdPin);
+									}
+									else
+									{
+										RemovePinWatch(EdPin);
+									}
+
+									if(InNotifType == ERigVMGraphNotifType::PinWatchedChanged)
+									{
+										return;
+									}
+									RequestAutoVMRecompilation();
+									(void)MarkPackageDirty();
+								}
+							}
+						}
+					}
+				}
+				// fall through another time
+			}
+			case ERigVMGraphNotifType::PinAdded:
+			case ERigVMGraphNotifType::PinRemoved:
+			case ERigVMGraphNotifType::PinRenamed:
+			{
+				// exposed pin changes like this (as well as type change etc)
+				// require to mark the blueprint as structurally modified,
+				// so that the instance actions work out.
+				if (URigVMPin* ModelPin = Cast<URigVMPin>(InSubject))
+				{
+					if(URigVMCollapseNode* CollapseNode = Cast<URigVMCollapseNode>(ModelPin->GetNode()))
+					{
+						if(Cast<URigVMFunctionLibrary>(CollapseNode->GetOuter()))
+						{
+							MarkBlueprintAsStructurallyModified();
+						}
+					}
+				}
+				break;
+			}
+			case ERigVMGraphNotifType::PinBoundVariableChanged:
+			case ERigVMGraphNotifType::VariableRemappingChanged:
+			{
+				RequestAutoVMRecompilation();
+				(void)MarkPackageDirty();
+				break;
+			}
+			case ERigVMGraphNotifType::NodeRenamed:
+			{
+				if (URigVMCollapseNode* CollapseNode = Cast<URigVMCollapseNode>(InSubject))
+				{
+					FString NewNodePath = CollapseNode->GetNodePath(true /* recursive */);
+					FString Left, Right = NewNodePath;
+					URigVMNode::SplitNodePathAtEnd(NewNodePath, Left, Right);
+					FString OldNodePath = CollapseNode->GetPreviousFName().ToString();
+					if (!Left.IsEmpty())
+					{
+						OldNodePath = URigVMNode::JoinNodePath(Left, OldNodePath);
+					}
+
+					HandleRigVMGraphRenamed(GetRigVMClientHost()->GetRigVMClient(), OldNodePath, NewNodePath);
+
+					if (UEdGraph* ContainedEdGraph = GetEdGraph(CollapseNode->GetContainedGraph()))
+					{
+						ContainedEdGraph->Rename(*CollapseNode->GetEditorSubGraphName(), nullptr);
+					}
+
+					MarkBlueprintAsStructurallyModified();
+				}
+				break;
+			}
+			case ERigVMGraphNotifType::NodeCategoryChanged:
+			case ERigVMGraphNotifType::NodeKeywordsChanged:
+			case ERigVMGraphNotifType::NodeDescriptionChanged:
+			{
+				MarkBlueprintAsStructurallyModified();
+				break;
+			}
+			default:
+			{
+				break;
+			}
+		}
+	}
+
+	// if the notification still has to be sent...
+	if (bNotifForOthersPending && !GetRigVMClientHost()->GetRigVMClient()->bSuspendModelNotificationsForOthers)
+	{
+		if (ModifiedEvent.IsBound())
+		{
+			ModifiedEvent.Broadcast(InNotifType, InGraph, InSubject);
+		}
+	}
+#endif
+}
+
+void IRigVMAssetInterface::SuspendNotifications(bool bSuspendNotifs)
+{
+	if (bSuspendAllNotifications == bSuspendNotifs)
+	{
+		return;
+	}
+
+	bSuspendAllNotifications = bSuspendNotifs;
+	if (!bSuspendNotifs)
+	{
+		RebuildGraphFromModel();
+		RefreshEditorEvent.Broadcast(GetObject());
+		RequestAutoVMRecompilation();
+	}
+}
+
+void IRigVMAssetInterface::CreateMemberVariablesOnLoad()
+{
+#if WITH_EDITOR
+
+	AddedMemberVariableMap.Reset();
+	TArray<FRigVMExternalVariable> Variables = GetExternalVariables(true);
+	for (int32 VariableIndex = 0; VariableIndex < Variables.Num(); VariableIndex++)
+	{
+		AddedMemberVariableMap.Add(Variables[VariableIndex].Name, VariableIndex);
+	}
+
+	if (GetRigVMClientHost()->GetRigVMClient()->Num() == 0)
+	{
+		return;
+	}
+
+#endif
+}
+
+void IRigVMAssetInterface::PatchVariableNodesOnLoad()
+{
+#if WITH_EDITOR
+	AddedMemberVariableMap.Reset();
+	PatchVariablesOnLoadSuper();
+#endif
+}
+
+void IRigVMAssetInterface::PatchBoundVariables()
+{
+}
+
+void IRigVMAssetInterface::PatchVariableNodesWithIncorrectType()
+{
+	TGuardValue<bool> GuardNotifsSelf(bSuspendModelNotificationsForSelf, true);
+
+	struct Local
+	{
+		static bool RefreshIfNeeded(URigVMController* Controller, URigVMVariableNode* VariableNode, const FString& CPPType, UObject* CPPTypeObject)
+		{
+			if (URigVMPin* ValuePin = VariableNode->GetValuePin())
+			{
+				if (ValuePin->GetCPPType() != CPPType || ValuePin->GetCPPTypeObject() != CPPTypeObject)
+				{
+					Controller->RefreshVariableNode(VariableNode->GetFName(), VariableNode->GetVariableName(), CPPType, CPPTypeObject, false);
+					if (RigVMTypeUtils::AreCompatible(*ValuePin->GetCPPType(), ValuePin->GetCPPTypeObject(), *CPPType, CPPTypeObject))
+					{
+						return false;
+					}
+					return true;
+				}
+			}
+			return false;
+		}
+	};
+
+	for (URigVMGraph* Graph : GetAllModels())
+	{
+		URigVMController* Controller = GetOrCreateController(Graph);
+		TArray<URigVMNode*> Nodes = Graph->GetNodes();
+		for (URigVMNode* Node : Nodes)
+		{
+			if (URigVMVariableNode* VariableNode = Cast<URigVMVariableNode>(Node))
+			{
+				if (VariableNode->IsInputArgument())
+				{
+					continue;
+				}
+				
+				FRigVMGraphVariableDescription Description = VariableNode->GetVariableDescription();
+
+				// Check for local variables
+				if (VariableNode->IsLocalVariable())
+				{
+					TArray<FRigVMGraphVariableDescription> LocalVariables = Graph->GetLocalVariables(false);
+					for (FRigVMGraphVariableDescription Variable : LocalVariables)
+					{
+						if (Variable.Name == Description.Name)
+						{
+							if (Local::RefreshIfNeeded(Controller, VariableNode, Variable.CPPType, Variable.CPPTypeObject))
+							{
+								bDirtyDuringLoad = true;
+							}
+							break;
+						}
+					}
+				}
+				else
+				{
+					for (FRigVMGraphVariableDescription& Variable : GetAssetVariables())
+					{
+						if (Variable.Name == Description.Name)
+						{
+							if (Local::RefreshIfNeeded(Controller, VariableNode, Variable.CPPType, Variable.CPPTypeObject))
+							{
+								bDirtyDuringLoad = true;
+							}
+						}
+					}
+				}
+			}
+		}
+	}
+}
+
+void IRigVMAssetInterface::PatchLinksWithCast()
+{
+#if WITH_EDITOR
+
+	{
+		TGuardValue<bool> DisableAutoCompile(bAutoRecompileVM, false);
+
+		// find all links containing a cast
+		TArray<TTuple<URigVMGraph*,TWeakObjectPtr<URigVMLink>,FString,FString>> LinksWithCast;
+		for (URigVMGraph* Graph : GetAllModels())
+		{
+			for(URigVMLink* Link : Graph->GetLinks())
+			{
+				const URigVMPin* SourcePin = Link->GetSourcePin();
+				const URigVMPin* TargetPin = Link->GetTargetPin();
+				if (SourcePin && TargetPin)
+				{
+					const TRigVMTypeIndex SourceTypeIndex = SourcePin->GetTypeIndex();
+					const TRigVMTypeIndex TargetTypeIndex = TargetPin->GetTypeIndex();
+					
+					if(SourceTypeIndex != TargetTypeIndex)
+					{
+						if(!FRigVMRegistry::Get().CanMatchTypes(SourceTypeIndex, TargetTypeIndex, true))
+						{
+							LinksWithCast.Emplace(Graph, TWeakObjectPtr<URigVMLink>(Link), SourcePin->GetPinPath(), TargetPin->GetPinPath());
+						}
+					}
+				}
+			}
+		}
+
+		// remove all of those links
+		for(const auto& Tuple : LinksWithCast)
+		{
+			URigVMController* Controller = GetController(Tuple.Get<0>());
+
+			if(URigVMLink* Link = Tuple.Get<1>().Get())
+			{
+				// the link may be detached, attach it first so that removal works.
+				const URigVMPin* SourcePin = Link->GetSourcePin();
+				URigVMPin* TargetPin = Link->GetTargetPin();
+				if(!SourcePin->IsLinkedTo(TargetPin))
+				{
+					const TArray<URigVMController::FLinkedPath> LinkedPaths = Controller->GetLinkedPaths({Link});
+					Controller->RestoreLinkedPaths(LinkedPaths);
+				}
+			}
+
+			Controller->BreakLink(Tuple.Get<2>(), Tuple.Get<3>(), false);
+
+			// notify the user that the link has been broken.
+			UE_LOG(LogRigVMDeveloper, Warning,
+				TEXT("A link was removed in %s (%s) - it contained different types on source and target pin (former cast link?)."),
+				*Controller->GetGraph()->GetNodePath(),
+				*URigVMLink::GetPinPathRepresentation(Tuple.Get<2>(), Tuple.Get<3>())
+			);
+		}
+	}
+#endif
+}
+
+void IRigVMAssetInterface::GetBackwardsCompatibilityPublicFunctions(TArray<FName>& BackwardsCompatiblePublicFunctions, TMap<URigVMLibraryNode*, FRigVMGraphFunctionHeader>& OldHeaders)
+{
+	TScriptInterface<IRigVMGraphFunctionHost> CRGeneratedClass = GetRigVMClientHost()->GetRigVMGraphFunctionHost();
+	FRigVMGraphFunctionStore* Store = CRGeneratedClass->GetRigVMGraphFunctionStore();
+	if (GetObject()->GetLinkerCustomVersion(FUE5MainStreamObjectVersion::GUID) < FUE5MainStreamObjectVersion::RigVMSaveFunctionAccessInModel)
+	{
+		for (const FRigVMGraphFunctionData& FunctionData : Store->PublicFunctions)
+		{
+			BackwardsCompatiblePublicFunctions.Add(FunctionData.Header.Name);
+			URigVMLibraryNode* LibraryNode = Cast<URigVMLibraryNode>(FunctionData.Header.LibraryPointer.GetNodeSoftPath().ResolveObject());
+			OldHeaders.Add(LibraryNode, FunctionData.Header);
+		}
+	}
+
+	// Addressing issue where PublicGraphFunctions is populated, but the model PublicFunctionNames is not
+	URigVMFunctionLibrary* FunctionLibrary = GetLocalFunctionLibrary();
+	if (FunctionLibrary)
+	{
+		if (GetPublicGraphFunctions().Num() > FunctionLibrary->PublicFunctionNames.Num())
+		{
+			for (const FRigVMGraphFunctionHeader& PublicHeader : GetPublicGraphFunctions())
+			{
+				BackwardsCompatiblePublicFunctions.Add(PublicHeader.Name);
+			}
+		}
+	}
+}
+
+void IRigVMAssetInterface::PropagateRuntimeSettingsFromBPToInstances()
+{
+	TArray<UObject*> ArchetypeInstances = GetArchetypeInstances(true, false);
+	for (UObject* ArchetypeInstance : ArchetypeInstances)
+	{
+		if (URigVMHost* InstanceHost = Cast<URigVMHost>(ArchetypeInstance))
+		{
+			InstanceHost->VMRuntimeSettings = GetVMRuntimeSettings();
+		}
+	}
+
+	TArray<UEdGraph*> EdGraphs;
+	GetAllEdGraphs(EdGraphs);
+
+	for (UEdGraph* Graph : EdGraphs)
+	{
+		TArray<UEdGraphNode*> Nodes = Graph->Nodes;
+		for (UEdGraphNode* Node : Nodes)
+		{
+			if(URigVMEdGraphNode* RigNode = Cast<URigVMEdGraphNode>(Node))
+			{
+				RigNode->ReconstructNode_Internal(true);
+			}
+		}
+	}
+}
+
+void IRigVMAssetInterface::InitializeArchetypeInstances()
+{
+	if (URigVM* VM = GetVM(true))
+	{
+		TArray<UObject*> ArchetypeInstances = GetArchetypeInstances(false, false);
+		for (UObject* Instance : ArchetypeInstances)
+		{
+			if (!Instance->HasAllFlags(RF_ClassDefaultObject))
+			{
+				if (URigVMHost* InstanceHost = Cast<URigVMHost>(Instance))
+				{
+					// No objects should be created during load, so PostInitInstanceIfRequired, which creates a new VM and
+					// DynamicHierarchy, should not be called during load
+					if (!InstanceHost->HasAllFlags(RF_NeedPostLoad))
+					{
+						InstanceHost->PostInitInstanceIfRequired();
+					}
+					InstanceHost->InstantiateVMFromCDO();
+					InstanceHost->CopyExternalVariableDefaultValuesFromCDO();
+				}
+			}
+		}
+	}
+}
+
+#if WITH_EDITOR
+
+void IRigVMAssetInterface::OnVariableAdded(const FName& InVarName)
+{
+	FRigVMGraphVariableDescription Variable;
+	for (FRigVMGraphVariableDescription& NewVariable : GetAssetVariables())
+	{
+		if (NewVariable.Name == InVarName)
+		{
+			Variable = NewVariable;
+			break;
+		}
+	}
+
+	const FRigVMExternalVariable ExternalVariable = RigVMTypeUtils::ExternalVariableFromRigVMVariableDescription(Variable);
+    FString CPPType;
+    UObject* CPPTypeObject = nullptr;
+    RigVMTypeUtils::CPPTypeFromExternalVariable(ExternalVariable, CPPType, &CPPTypeObject);
+	if (CPPTypeObject)
+	{
+		if (ExternalVariable.bIsArray)
+		{
+			CPPType = RigVMTypeUtils::ArrayTypeFromBaseType(CPPTypeObject->GetPathName());
+		}
+		else
+		{
+			CPPType = CPPTypeObject->GetPathName();
+		}
+	}
+
+	// register the type in the registry
+	FRigVMRegistry::Get().FindOrAddType(FRigVMTemplateArgumentType(*CPPType, CPPTypeObject));
+	
+    RigVMPythonUtils::Print(GetObject()->GetFName().ToString(),
+		FString::Printf(TEXT("blueprint.add_member_variable('%s', '%s', %s, %s, '%s')"),
+			*InVarName.ToString(),
+			*CPPType,
+			(ExternalVariable.bIsPublic) ? TEXT("False") : TEXT("True"), 
+			(ExternalVariable.bIsReadOnly) ? TEXT("True") : TEXT("False"), 
+			*Variable.DefaultValue)); 
+	
+	BroadcastExternalVariablesChangedEvent();
+}
+
+void IRigVMAssetInterface::OnVariableRemoved(const FName& InVarName)
+{
+	TArray<URigVMGraph*> AllGraphs = GetAllModels();
+	for (URigVMGraph* Graph : AllGraphs)
+	{
+		if (URigVMController* Controller = GetOrCreateController(Graph))
+		{
+#if WITH_EDITOR
+			const bool bSetupUndoRedo = !GIsTransacting;
+#else
+			const bool bSetupUndoRedo = false;
+#endif
+			Controller->OnExternalVariableRemoved(InVarName, bSetupUndoRedo);
+		}
+	}
+
+	RigVMPythonUtils::Print(GetObject()->GetFName().ToString(),
+		FString::Printf(TEXT("blueprint.remove_member_variable('%s')"),
+			*InVarName.ToString()));
+	
+	BroadcastExternalVariablesChangedEvent();
+}
+
+void IRigVMAssetInterface::OnVariableRenamed(const FName& InOldVarName, const FName& InNewVarName)
+{
+	TArray<URigVMGraph*> AllGraphs = GetAllModels();
+	for (URigVMGraph* Graph : AllGraphs)
+	{
+		if (URigVMController* Controller = GetOrCreateController(Graph))
+		{
+#if WITH_EDITOR
+			const bool bSetupUndoRedo = !GIsTransacting;
+#else
+			const bool bSetupUndoRedo = false;
+#endif
+			Controller->OnExternalVariableRenamed(InOldVarName, InNewVarName, bSetupUndoRedo);
+		}
+	}
+
+	RigVMPythonUtils::Print(GetObject()->GetFName().ToString(),
+		FString::Printf(TEXT("blueprint.rename_member_variable('%s', '%s')"),
+			*InOldVarName.ToString(),
+			*InNewVarName.ToString()));
+	
+	BroadcastExternalVariablesChangedEvent();
+}
+
+void IRigVMAssetInterface::OnVariableTypeChanged(const FName& InVarName, FEdGraphPinType InOldPinType, FEdGraphPinType InNewPinType)
+{
+	FString CPPType;
+	UObject* CPPTypeObject = nullptr;
+	RigVMTypeUtils::CPPTypeFromPinType(InNewPinType, CPPType, &CPPTypeObject);
+	
+	TArray<URigVMGraph*> AllGraphs = GetAllModels();
+	for (URigVMGraph* Graph : AllGraphs)
+	{
+		if (URigVMController* Controller = GetOrCreateController(Graph))
+		{
+#if WITH_EDITOR
+			const bool bSetupUndoRedo = !GIsTransacting;
+#else
+			const bool bSetupUndoRedo = false;
+#endif
+
+			if (!CPPType.IsEmpty())
+			{
+				Controller->OnExternalVariableTypeChanged(InVarName, CPPType, CPPTypeObject, bSetupUndoRedo);
+			}
+			else
+			{
+				Controller->OnExternalVariableRemoved(InVarName, bSetupUndoRedo);
+			}
+		}
+	}
+
+	if(UScriptStruct* ScriptStruct = Cast<UScriptStruct>(CPPTypeObject))
+	{
+		for (auto Var : GetAssetVariables())
+		{
+			if (Var.Name == InVarName)
+			{
+				CPPType = ScriptStruct->GetName();
+			}
+		}
+	}
+	else if (UEnum* Enum = Cast<UEnum>(CPPTypeObject))
+	{
+		for (auto Var : GetAssetVariables())
+		{
+			if (Var.Name == InVarName)
+			{
+				CPPType = Enum->GetName();
+			}
+		}
+	}
+
+	// register the type in the registry
+	FRigVMRegistry::Get().FindOrAddType(FRigVMTemplateArgumentType(*CPPType, CPPTypeObject));
+
+	RigVMPythonUtils::Print(GetObject()->GetFName().ToString(),
+		FString::Printf(TEXT("blueprint.change_member_variable_type('%s', '%s')"),
+		*InVarName.ToString(),
+		*CPPType));
+
+	BroadcastExternalVariablesChangedEvent();
+}
+
+FName IRigVMAssetInterface::AddAssetVariableFromPinType(const FName& InName, const FEdGraphPinType& InType, bool bIsPublic, bool bIsReadOnly, FString InDefaultValue)
+{
+	FString CPPType;
+	UObject* CPPTypeObject = nullptr;
+	RigVMTypeUtils::CPPTypeFromPinType(InType, CPPType, &CPPTypeObject);
+	return AddMemberVariable(InName, CPPType, bIsPublic, bIsReadOnly, InDefaultValue);
+}
+
+void IRigVMAssetInterface::BroadcastExternalVariablesChangedEvent()
+{
+	ExternalVariablesChangedEvent.Broadcast(GetExternalVariables(false));
+}
+
+void IRigVMAssetInterface::BroadcastNodeDoubleClicked(URigVMNode* InNode)
+{
+	NodeDoubleClickedEvent.Broadcast(GetObject(), InNode);
+}
+
+void IRigVMAssetInterface::BroadcastGraphImported(UEdGraph* InGraph)
+{
+	GraphImportedEvent.Broadcast(InGraph);
+}
+
+void IRigVMAssetInterface::BroadcastPostEditChangeChainProperty(FPropertyChangedChainEvent& PropertyChangedChainEvent)
+{
+	PostEditChangeChainPropertyEvent.Broadcast(PropertyChangedChainEvent);
+}
+
+void IRigVMAssetInterface::RefreshAllNodes()
+{
+	DECLARE_SCOPE_HIERARCHICAL_COUNTER_FUNC()
+#if WITH_EDITORONLY_DATA
+	// Avoid refreshing EdGraph nodes during cook
+	if (GIsCookerLoadingPackage)
+	{
+		return;
+	}
+	
+	
+	// Avoid refreshing EdGraph if PostLoad() hasn't been called, since reconstruct node later
+	// can access model data that hasn't been fully loaded. And it is ok to skip here because
+	// the EdGraph will be reconstructed later when the CR editor
+	// initializes, as that is when the EdGraph is actually used.
+	if (GetObject()->HasAnyFlags(RF_NeedPostLoad))
+	{
+		return;
+	}
+	
+	if (GetDefaultModel() == nullptr)
+	{
+		return;
+	}
+
+	TArray<URigVMEdGraphNode*> AllNodes;
+	TArray<UEdGraph*> AllGraphs;
+	GetAllEdGraphs(AllGraphs);
+	for(int32 i=0; i<AllGraphs.Num(); i++)
+	{
+		check(AllGraphs[i] != NULL);
+		TArray<URigVMEdGraphNode*> GraphNodes;
+		AllGraphs[i]->GetNodesOfClass<URigVMEdGraphNode>(GraphNodes);
+		AllNodes.Append(GraphNodes);
+	}
+
+	for (URigVMEdGraphNode* Node : AllNodes)
+	{
+		Node->SetFlags(RF_Transient);
+	}
+
+	for(URigVMEdGraphNode* Node : AllNodes)
+	{
+		Node->ReconstructNode();
+	}
+
+	for (URigVMEdGraphNode* Node : AllNodes)
+	{
+		Node->ClearFlags(RF_Transient);
+	}
+	
+#endif
+}
+
+void IRigVMAssetInterface::BroadcastRequestLocalizeFunctionDialog(FRigVMGraphFunctionIdentifier InFunction, bool bForce)
+{
+	RequestLocalizeFunctionDialog.Broadcast(InFunction, GetController(GetDefaultModel()), GetRigVMClientHost()->GetRigVMGraphFunctionHost().GetInterface(), bForce);
+}
+
+const FCompilerResultsLog& IRigVMAssetInterface::GetCompileLog() const
+{
+	return CompileLog;
+}
+
+FCompilerResultsLog& IRigVMAssetInterface::GetCompileLog()
+{
+	return CompileLog;
+}
+
+void IRigVMAssetInterface::BroadCastReportCompilerMessage(EMessageSeverity::Type InSeverity, UObject* InSubject, const FString& InMessage)
+{
+	ReportCompilerMessageEvent.Broadcast(InSeverity, InSubject, InMessage);
+}
+
+#endif
+
+UEdGraph* IRigVMAssetInterface::CreateEdGraph(URigVMGraph* InModel, bool bForce)
+{
+	check(InModel);
+
+#if WITH_EDITORONLY_DATA
+	if(InModel->IsA<URigVMFunctionLibrary>())
+	{
+		return GetFunctionLibraryEdGraph();
+	}
+#endif
+	
+	if(bForce)
+	{
+		RemoveEdGraph(InModel);
+	}
+
+	FString GraphName = InModel->GetName();
+	GraphName.RemoveFromStart(FRigVMClient::RigVMModelPrefix);
+	GraphName.TrimStartAndEndInline();
+
+	if(GraphName.IsEmpty())
+	{
+		GraphName = URigVMEdGraphSchema::GraphName_RigVM.ToString();
+	}
+
+	GraphName = GetRigVMClientHost()->GetRigVMClient()->GetUniqueName(*GraphName).ToString();
+
+	URigVMEdGraph* RigVMEdGraph = NewObject<URigVMEdGraph>(GetObject(), GetRigVMClientHost()->GetRigVMEdGraphClass(), *GraphName, RF_Transactional);
+	RigVMEdGraph->Schema = GetRigVMClientHost()->GetRigVMEdGraphSchemaClass();
+	RigVMEdGraph->bAllowDeletion = true;
+	RigVMEdGraph->ModelNodePath = InModel->GetNodePath();
+	RigVMEdGraph->InitializeFromAsset(GetObject());
+	
+	AddUbergraphPage(RigVMEdGraph);
+	AddLastEditedDocument(RigVMEdGraph);
+
+	return RigVMEdGraph;
+}
+
+bool IRigVMAssetInterface::RemoveEdGraph(URigVMGraph* InModel)
+{
+	if(URigVMEdGraph* RigGraph = Cast<URigVMEdGraph>(GetEdGraph(InModel)))
+	{
+		TArray<TObjectPtr<UEdGraph>>& UbergraphPages = GetUberGraphs();
+		if(UbergraphPages.Contains(RigGraph))
+		{
+			GetObject()->Modify();
+			UbergraphPages.Remove(RigGraph);
+		}
+		DestroyObject(RigGraph);
+		return true;
+	}
+	return false;
+}
+
+void IRigVMAssetInterface::DestroyObject(UObject* InObject)
+{
+	GetRigVMClientHost()->GetRigVMClient()->DestroyObject(InObject);
+}
+
+void IRigVMAssetInterface::RenameGraph(const FString& InNodePath, const FName& InNewName)
+{
+	FName OldName = NAME_None;
+	UEdGraph* EdGraph = GetEdGraph(InNodePath);
+	if(EdGraph)
+	{
+		OldName = EdGraph->GetFName();
+	}
+	
+	GetRigVMClientHost()->GetRigVMClient()->RenameModel(InNodePath, InNewName, true);
+
+	if(EdGraph)
+	{
+		NotifyGraphRenamedSuper(EdGraph, OldName, EdGraph->GetFName());
+	}
+}
+
+void IRigVMAssetInterface::CreateEdGraphForCollapseNodeIfNeeded(URigVMCollapseNode* InNode, bool bForce)
+{
+	check(InNode);
+
+	if (bForce)
+	{
+		RemoveEdGraphForCollapseNode(InNode, false);
+	}
+
+	if (InNode->GetGraph()->IsA<URigVMFunctionLibrary>())
+	{
+		if (URigVMGraph* ContainedGraph = InNode->GetContainedGraph())
+		{
+			bool bFunctionGraphExists = false;
+			for (UEdGraph* FunctionGraph : GetFunctionGraphs())
+			{
+				if (URigVMEdGraph* RigFunctionGraph = Cast<URigVMEdGraph>(FunctionGraph))
+				{
+					if (RigFunctionGraph->ModelNodePath == ContainedGraph->GetNodePath())
+					{
+						bFunctionGraphExists = true;
+						break;
+					}
+				}
+			}
+
+			if (!bFunctionGraphExists)
+			{
+				// create a sub graph
+				URigVMEdGraph* RigFunctionGraph = NewObject<URigVMEdGraph>(GetObject(), GetRigVMClientHost()->GetRigVMEdGraphClass(), *InNode->GetName(), RF_Transactional);
+				RigFunctionGraph->Schema = GetRigVMClientHost()->GetRigVMEdGraphSchemaClass();
+				RigFunctionGraph->bAllowRenaming = 1;
+				RigFunctionGraph->bEditable = 1;
+				RigFunctionGraph->bAllowDeletion = 1;
+				RigFunctionGraph->ModelNodePath = ContainedGraph->GetNodePath();
+				RigFunctionGraph->bIsFunctionDefinition = true;
+
+				GetFunctionGraphs().Add(RigFunctionGraph);
+
+				RigFunctionGraph->InitializeFromAsset(GetObject());
+
+				GetOrCreateController(ContainedGraph)->ResendAllNotifications();
+			}
+
+		}
+	}
+	else if (URigVMEdGraph* RigGraph = Cast<URigVMEdGraph>(GetEdGraph(InNode->GetGraph())))
+	{
+		if (URigVMGraph* ContainedGraph = InNode->GetContainedGraph())
+		{
+			bool bSubGraphExists = false;
+			for (UEdGraph* SubGraph : RigGraph->SubGraphs)
+			{
+				if (URigVMEdGraph* SubRigGraph = Cast<URigVMEdGraph>(SubGraph))
+				{
+					if (SubRigGraph->ModelNodePath == ContainedGraph->GetNodePath())
+					{
+						bSubGraphExists = true;
+						break;
+					}
+				}
+			}
+
+			if (!bSubGraphExists)
+			{
+				bool bEditable = true;
+				if (InNode->IsA<URigVMAggregateNode>())
+				{
+					bEditable = false;
+				}
+				
+				// create a sub graph
+				URigVMEdGraph* SubRigGraph = NewObject<URigVMEdGraph>(RigGraph, GetRigVMClientHost()->GetRigVMEdGraphClass(), *InNode->GetEditorSubGraphName(), RF_Transactional);
+				SubRigGraph->Schema = GetRigVMClientHost()->GetRigVMEdGraphSchemaClass();
+				SubRigGraph->bAllowRenaming = 1;
+				SubRigGraph->bEditable = bEditable;
+				SubRigGraph->bAllowDeletion = 1;
+				SubRigGraph->ModelNodePath = ContainedGraph->GetNodePath();
+				SubRigGraph->bIsFunctionDefinition = false;
+
+				RigGraph->SubGraphs.Add(SubRigGraph);
+
+				SubRigGraph->InitializeFromAsset(GetObject());
+
+				GetOrCreateController(ContainedGraph)->ResendAllNotifications();
+			}
+		}
+	}
+}
+
+bool IRigVMAssetInterface::RemoveEdGraphForCollapseNode(URigVMCollapseNode* InNode, bool bNotify)
+{
+	check(InNode);
+
+	if (InNode->GetGraph()->IsA<URigVMFunctionLibrary>())
+	{
+		if (URigVMGraph* ContainedGraph = InNode->GetContainedGraph())
+		{
+			for (UEdGraph* FunctionGraph : GetFunctionGraphs())
+			{
+				if (URigVMEdGraph* RigFunctionGraph = Cast<URigVMEdGraph>(FunctionGraph))
+				{
+					if (RigFunctionGraph->ModelNodePath == ContainedGraph->GetNodePath())
+					{
+						if (URigVMController* SubController = GetController(ContainedGraph))
+						{
+							SubController->OnModified().RemoveAll(RigFunctionGraph);
+						}
+
+						if (ModifiedEvent.IsBound() && bNotify)
+						{
+							ModifiedEvent.Broadcast(ERigVMGraphNotifType::NodeRemoved, InNode->GetGraph(), InNode);
+						}
+
+						GetFunctionGraphs().Remove(RigFunctionGraph);
+						RigFunctionGraph->Rename(nullptr, GetTransientPackage(), REN_DontCreateRedirectors);
+						if(RigFunctionGraph->IsRooted())
+						{
+							RigFunctionGraph->RemoveFromRoot();
+						}
+						RigFunctionGraph->MarkAsGarbage();
+						return bNotify;
+					}
+				}
+			}
+		}
+	}
+	else if (URigVMEdGraph* RigGraph = Cast<URigVMEdGraph>(GetEdGraph(InNode->GetGraph())))
+	{
+		if (URigVMGraph* ContainedGraph = InNode->GetContainedGraph())
+		{
+			for (UEdGraph* SubGraph : RigGraph->SubGraphs)
+			{
+				if (URigVMEdGraph* SubRigGraph = Cast<URigVMEdGraph>(SubGraph))
+				{
+					if (SubRigGraph->ModelNodePath == ContainedGraph->GetNodePath())
+					{
+						if (URigVMController* SubController = GetController(ContainedGraph))
+						{
+							SubController->OnModified().RemoveAll(SubRigGraph);
+						}
+
+						if (ModifiedEvent.IsBound() && bNotify)
+						{
+							ModifiedEvent.Broadcast(ERigVMGraphNotifType::NodeRemoved, InNode->GetGraph(), InNode);
+						}
+
+						RigGraph->SubGraphs.Remove(SubRigGraph);
+						SubRigGraph->Rename(nullptr, GetTransientPackage(), REN_DontCreateRedirectors);
+						if(SubRigGraph->IsRooted())
+						{
+							SubRigGraph->RemoveFromRoot();
+						}
+						SubRigGraph->MarkAsGarbage();
+						return bNotify;
+					}
+				}
+			}
+		}
+	}
+
+	return false;
+}
+
+#if WITH_EDITOR
+
+void IRigVMAssetInterface::QueueCompilerMessageDelegate(const FOnRigVMReportCompilerMessage::FDelegate& InDelegate)
+{
+	FScopeLock Lock(&QueuedCompilerMessageDelegatesMutex);
+	QueuedCompilerMessageDelegates.Add(InDelegate);
+}
+
+void IRigVMAssetInterface::ClearQueuedCompilerMessageDelegates()
+{
+	FScopeLock Lock(&QueuedCompilerMessageDelegatesMutex);
+	QueuedCompilerMessageDelegates.Reset();
+}
+
+#endif
+
+FRigVMBlueprintCompileScope::FRigVMBlueprintCompileScope(TRigVMAssetInterfacePtr InBlueprint): Blueprint(InBlueprint)
+{
+	check(Blueprint);
+	Blueprint->IncrementVMRecompileBracket();
+}
+
+FRigVMBlueprintCompileScope::~FRigVMBlueprintCompileScope()
+{
+	Blueprint->DecrementVMRecompileBracket();
+}
+
+#undef LOCTEXT_NAMESPACE
+
+
