RigVM: Adding IRigVMAssetInterface as parent of RigVMBlueprint. First step to separation of RigVM and Blueprint.

diff --git a/Engine/Plugins/Runtime/RigVM/Source/RigVMDeveloper/Private/RigVMBlueprintCompiler.cpp b/Engine/Plugins/Runtime/RigVM/Source/RigVMDeveloper/Private/RigVMBlueprintCompiler.cpp
--- a/Engine/Plugins/Runtime/RigVM/Source/RigVMDeveloper/Private/RigVMBlueprintCompiler.cpp
+++ b/Engine/Plugins/Runtime/RigVM/Source/RigVMDeveloper/Private/RigVMBlueprintCompiler.cpp
@@ -3,7 +3,7 @@
 #include "RigVMBlueprintCompiler.h"
 #include "RigVMHost.h"
 #include "KismetCompiler.h"
-#include "RigVMBlueprint.h"
+#include "RigVMBlueprintLegacy.h"
 #include "RigVMCore/RigVMStruct.h"
 #include "EdGraph/RigVMEdGraph.h"
 #include "EdGraph/RigVMEdGraphNode.h"
diff --git a/Engine/Plugins/Runtime/RigVM/Source/RigVMDeveloper/Private/RigVMBlueprintUtils.cpp b/Engine/Plugins/Runtime/RigVM/Source/RigVMDeveloper/Private/RigVMBlueprintUtils.cpp
--- a/Engine/Plugins/Runtime/RigVM/Source/RigVMDeveloper/Private/RigVMBlueprintUtils.cpp
+++ b/Engine/Plugins/Runtime/RigVM/Source/RigVMDeveloper/Private/RigVMBlueprintUtils.cpp
@@ -3,17 +3,19 @@
 #include "RigVMBlueprintUtils.h"
 
 #include "BlueprintActionDatabase.h"
+#include "INotifyFieldValueChanged.h"
 #include "Kismet2/BlueprintEditorUtils.h"
 #include "RigVMCore/RigVMStruct.h"
 #include "UObject/UObjectIterator.h"
-#include "EdGraph/RigVMEdGraphNode.h"
-#include "RigVMBlueprint.h"
+#include "RigVMAsset.h"
+#include "EdGraph/RigVMEdGraph.h"
+#include "EdGraph/RigVMEdGraphSchema.h"
 #include "Kismet2/Kismet2NameValidators.h"
 #include "Stats/StatsHierarchical.h"
 
 #define LOCTEXT_NAMESPACE "RigVMBlueprintUtils"
 
-FName FRigVMBlueprintUtils::ValidateName(UBlueprint* InBlueprint, const FString& InName)
+FName FRigVMBlueprintUtils::ValidateName(TRigVMAssetInterfacePtr InBlueprint, const FString& InName)
 {
 	DECLARE_SCOPE_HIERARCHICAL_COUNTER_FUNC()
 
@@ -27,51 +29,121 @@ FName FRigVMBlueprintUtils::ValidateName(UBlueprint* InBlueprint, const FString&
 		Name.RightChopInline(12, EAllowShrinking::No);
 	}
 
-	TSharedPtr<FKismetNameValidator> NameValidator;
-	NameValidator = MakeShareable(new FKismetNameValidator(InBlueprint));
-
-	// Clean up BaseName to not contain any invalid characters, which will mean we can never find a legal name no matter how many numbers we add
-	if (NameValidator->IsValid(Name) == EValidatorResult::ContainsInvalidCharacters)
+	if (UBlueprint* Blueprint = Cast<UBlueprint>(InBlueprint.GetObject()))
 	{
-		for (TCHAR& TestChar : Name)
+		TSharedPtr<FKismetNameValidator> NameValidator;
+		NameValidator = MakeShareable(new FKismetNameValidator(Blueprint));
+
+		// Clean up BaseName to not contain any invalid characters, which will mean we can never find a legal name no matter how many numbers we add
+		if (NameValidator->IsValid(Name) == EValidatorResult::ContainsInvalidCharacters)
 		{
-			for (TCHAR BadChar : UE_BLUEPRINT_INVALID_NAME_CHARACTERS)
+			for (TCHAR& TestChar : Name)
 			{
-				if (TestChar == BadChar)
+				for (TCHAR BadChar : UE_BLUEPRINT_INVALID_NAME_CHARACTERS)
 				{
-					TestChar = TEXT('_');
-					break;
+					if (TestChar == BadChar)
+					{
+						TestChar = TEXT('_');
+						break;
+					}
 				}
 			}
 		}
+
+		if (UClass* ParentClass = Blueprint->ParentClass)
+		{
+			FFieldVariant ExistingField = FindUFieldOrFProperty(ParentClass, *Name);
+			if (ExistingField)
+			{
+				Name = FString::Printf(TEXT("%s_%d"), *Name, 0);
+			}
+		}
+
+		int32 Count = 0;
+		FString BaseName = Name;
+		while (NameValidator->IsValid(Name) != EValidatorResult::Ok)
+		{
+			// Calculate the number of digits in the number, adding 2 (1 extra to correctly count digits, another to account for the '_' that will be added to the name
+			int32 CountLength = Count > 0 ? (int32)log((double)Count) + 2 : 2;
+
+			// If the length of the final string will be too long, cut off the end so we can fit the number
+			if (CountLength + BaseName.Len() > NameValidator->GetMaximumNameLength())
+			{
+				BaseName.LeftInline(NameValidator->GetMaximumNameLength() - CountLength);
+			}
+			Name = FString::Printf(TEXT("%s_%d"), *BaseName, Count);
+			Count++;
+		}
+	}
+
+	return *Name;
+}
+
+UEdGraphNode* FRigVMBlueprintUtils::GetNodeByGUID(const TRigVMAssetInterfacePtr InBlueprint, const FGuid& InNodeGuid)
+{
+	TArray<UEdGraphNode*> Nodes;
+	TArray<UEdGraph*> AllGraphs;
+	InBlueprint->GetAllEdGraphs(AllGraphs);
+	for(int32 i=0; i<AllGraphs.Num(); i++)
+	{
+		check(AllGraphs[i] != NULL);
+		TArray<UEdGraphNode*> GraphNodes;
+		AllGraphs[i]->GetNodesOfClass(GraphNodes);
+		Nodes.Append(GraphNodes);
 	}
 
-	if (UClass* ParentClass = InBlueprint->ParentClass)
+	for(UEdGraphNode* Node : Nodes)
 	{
-		FFieldVariant ExistingField = FindUFieldOrFProperty(ParentClass, *Name);
-		if (ExistingField)
+		if(Node->NodeGuid == InNodeGuid)
 		{
-			Name = FString::Printf(TEXT("%s_%d"), *Name, 0);
+			return Node;
 		}
 	}
+	return nullptr;
+}
 
-	int32 Count = 0;
-	FString BaseName = Name;
-	while (NameValidator->IsValid(Name) != EValidatorResult::Ok)
+TRigVMAssetInterfacePtr FRigVMBlueprintUtils::FindAssetForGraph(const UEdGraph* Graph)
+{
+	for (UObject* TestOuter = Graph ? Graph->GetOuter() : nullptr; TestOuter; TestOuter = TestOuter->GetOuter())
 	{
-		// Calculate the number of digits in the number, adding 2 (1 extra to correctly count digits, another to account for the '_' that will be added to the name
-		int32 CountLength = Count > 0 ? (int32)log((double)Count) + 2 : 2;
+		if (TestOuter->Implements<URigVMAssetInterface>())
+		{
+			return TestOuter;
+		}
+	}
+
+	return nullptr;
+}
 
-		// If the length of the final string will be too long, cut off the end so we can fit the number
-		if (CountLength + BaseName.Len() > NameValidator->GetMaximumNameLength())
+TRigVMAssetInterfacePtr FRigVMBlueprintUtils::FindAssetForNode(const UEdGraphNode* Node)
+{
+	if (Node)
+	{
+		if (UEdGraph* Graph = Node->GetGraph())
 		{
-			BaseName.LeftInline(NameValidator->GetMaximumNameLength() - CountLength);
+			return FindAssetForGraph(Node->GetGraph());
 		}
-		Name = FString::Printf(TEXT("%s_%d"), *BaseName, Count);
-		Count++;
 	}
+	return nullptr;
+}
 
-	return *Name;
+FName FRigVMBlueprintUtils::FindUniqueVariableName(const IRigVMAssetInterface* InBlueprint, const FString& InBaseName)
+{
+	DECLARE_SCOPE_HIERARCHICAL_COUNTER_FUNC()
+
+	const TArray<FRigVMGraphVariableDescription> Variables = InBlueprint->GetAssetVariables();
+	FString NewName = InBaseName;
+	int32 Index = 0;
+	while (Variables.ContainsByPredicate([NewName](const FRigVMGraphVariableDescription& Variable)
+		{
+			return Variable.Name == NewName;
+		}))
+	{
+		Index++;
+		NewName = FString::Printf(TEXT("%s_%d"), *InBaseName, Index);
+	}
+
+	return *NewName;
 }
 
 void FRigVMBlueprintUtils::ForAllRigVMStructs(TFunction<void(UScriptStruct*)> InFunction)
@@ -91,63 +163,39 @@ void FRigVMBlueprintUtils::ForAllRigVMStructs(TFunction<void(UScriptStruct*)> In
 	}
 }
 
-void FRigVMBlueprintUtils::HandleReconstructAllNodes(UBlueprint* InBlueprint)
+void FRigVMBlueprintUtils::HandleReconstructAllBlueprintNodes(UBlueprint* InBlueprint)
 {
-	DECLARE_SCOPE_HIERARCHICAL_COUNTER_FUNC()
-
-	return HandleRefreshAllNodes(InBlueprint);
+	if (TRigVMAssetInterfacePtr Interface = InBlueprint)
+	{
+		HandleReconstructAllBlueprintNodes(Interface);
+	}
 }
 
-void FRigVMBlueprintUtils::HandleRefreshAllNodes(UBlueprint* InBlueprint)
+void FRigVMBlueprintUtils::HandleReconstructAllBlueprintNodes(TRigVMAssetInterfacePtr InBlueprint)
 {
 	DECLARE_SCOPE_HIERARCHICAL_COUNTER_FUNC()
-#if WITH_EDITORONLY_DATA
-	// Avoid refreshing EdGraph nodes during cook
-	if (GIsCookerLoadingPackage)
-	{
-		return;
-	}
-	
-	if(const URigVMBlueprint* RigVMBlueprint = Cast<URigVMBlueprint>(InBlueprint))
-	{
-		// Avoid refreshing EdGraph if PostLoad() hasn't been called, since reconstruct node later
-		// can access model data that hasn't been fully loaded. And it is ok to skip here because
-		// the EdGraph will be reconstructed later when the CR editor
-		// initializes, as that is when the EdGraph is actually used.
-		if (RigVMBlueprint->HasAnyFlags(RF_NeedPostLoad))
-		{
-			return;
-		}
-		
-		if (RigVMBlueprint->GetDefaultModel() == nullptr)
-		{
-			return;
-		}
 
-		TArray<URigVMEdGraphNode*> AllNodes;
-		FBlueprintEditorUtils::GetAllNodesOfClass(RigVMBlueprint, AllNodes);
+	return HandleRefreshAllNodes(InBlueprint);
+}
 
-		for (URigVMEdGraphNode* Node : AllNodes)
-		{
-			Node->SetFlags(RF_Transient);
-		}
+void FRigVMBlueprintUtils::HandleRefreshAllNodes(TRigVMAssetInterfacePtr InBlueprint)
+{
+	InBlueprint->RefreshAllNodes();
+}
 
-		for(URigVMEdGraphNode* Node : AllNodes)
-		{
-			Node->ReconstructNode();
-		}
+void FRigVMBlueprintUtils::HandleRefreshAllBlueprintNodes(UBlueprint* InBlueprint)
+{
+	DECLARE_SCOPE_HIERARCHICAL_COUNTER_FUNC()
 
-		for (URigVMEdGraphNode* Node : AllNodes)
-		{
-			Node->ClearFlags(RF_Transient);
-		}
+	if (TRigVMAssetInterfacePtr Blueprint = InBlueprint)
+	{
+		Blueprint->RefreshAllNodes();
 	}
-#endif
 }
 
 void FRigVMBlueprintUtils::HandleAssetDeleted(const FAssetData& InAssetData)
 {
-	if (InAssetData.GetClass() && InAssetData.GetClass()->IsChildOf(URigVMBlueprint::StaticClass()))
+	if (InAssetData.GetClass() && InAssetData.GetClass()->ImplementsInterface(URigVMAssetInterface::StaticClass()))
 	{
 		// Make sure any RigVMBlueprint removes any TypeActions related to this asset (e.g. public functions)
 		if (FBlueprintActionDatabase* ActionDatabase = FBlueprintActionDatabase::TryGet())
@@ -157,13 +205,54 @@ void FRigVMBlueprintUtils::HandleAssetDeleted(const FAssetData& InAssetData)
 	}
 }
 
-void FRigVMBlueprintUtils::RemoveMemberVariableIfNotUsed(UBlueprint* Blueprint, const FName VarName)
+void FRigVMBlueprintUtils::RemoveMemberVariableIfNotUsed(TRigVMAssetInterfacePtr Blueprint, const FName VarName)
 {
 	DECLARE_SCOPE_HIERARCHICAL_COUNTER_FUNC()
 
-	if (Blueprint->IsA<URigVMBlueprint>())
+	if (Blueprint->GetObject()->IsA<UBlueprint>())
+	{
+		FBlueprintEditorUtils::RemoveMemberVariable(Cast<UBlueprint>(Blueprint->GetObject()), VarName);
+	}
+}
+
+URigVMEdGraph* FRigVMBlueprintUtils::CreateNewGraph(UObject* ParentScope, const FName& GraphName, TSubclassOf<class URigVMEdGraph> GraphClass, TSubclassOf<class URigVMEdGraphSchema> SchemaClass)
+{
+	// Copy paste from FBlueprintEditorUtils::CreateNewGraph
+	
+	URigVMEdGraph* NewGraph = nullptr;
+	bool bRename = false;
+
+	// Ensure this name isn't already being used for a graph
+	if (GraphName != NAME_None)
+	{
+		if (UObject* ExistingObject = FindObject<UObject>(ParentScope, *(GraphName.ToString())))
+		{
+			if (ExistingObject->IsA<URigVMEdGraph>())
+			{
+				// Rename the old graph out of the way - this may confuse the user somewhat - and even
+				// break their logic. But name collisions are not avoidable e.g. someone can add
+				// a function to an interface that conflicts with something in a class hierarchy
+				ExistingObject->Rename(nullptr, ExistingObject->GetOuter(), REN_DoNotDirty);
+			}
+		}
+
+		// Construct new graph with the supplied name
+		NewGraph = NewObject<URigVMEdGraph>(ParentScope, GraphClass, NAME_None, RF_Transactional);
+		bRename = true;
+	}
+	else
+	{
+		// Construct a new graph with a default name
+		NewGraph = NewObject<URigVMEdGraph>(ParentScope, GraphClass, NAME_None, RF_Transactional);
+	}
+
+	NewGraph->Schema = SchemaClass;
+
+	// Now move to where we want it to. Workaround to ensure transaction buffer is correctly utilized
+	if (bRename)
 	{
-		FBlueprintEditorUtils::RemoveMemberVariable(Blueprint, VarName);
+		NewGraph->Rename(*(GraphName.ToString()), ParentScope, REN_DoNotDirty);
 	}
+	return NewGraph;
 }
 #undef LOCTEXT_NAMESPACE
\ No newline at end of file
diff --git a/Engine/Plugins/Runtime/RigVM/Source/RigVMDeveloper/Private/RigVMCompiler/RigVMCompiler.cpp b/Engine/Plugins/Runtime/RigVM/Source/RigVMDeveloper/Private/RigVMCompiler/RigVMCompiler.cpp
--- a/Engine/Plugins/Runtime/RigVM/Source/RigVMDeveloper/Private/RigVMCompiler/RigVMCompiler.cpp
+++ b/Engine/Plugins/Runtime/RigVM/Source/RigVMDeveloper/Private/RigVMCompiler/RigVMCompiler.cpp
@@ -1621,7 +1621,7 @@ bool URigVMCompiler::CompileFunction(const FRigVMCompileSettings& InSettings, co
 	// Update the compilation data of this library, and the hashes of the compilation data of its dependencies used for this compilation
 	if (IRigVMClientHost* ClientHost = InLibraryNode->GetImplementingOuter<IRigVMClientHost>())
 	{
-		if (IRigVMGraphFunctionHost* FunctionHost = ClientHost->GetRigVMGraphFunctionHost())
+		if (TScriptInterface<IRigVMGraphFunctionHost> FunctionHost = ClientHost->GetRigVMGraphFunctionHost())
 		{
 			if (FRigVMGraphFunctionStore* Store = FunctionHost->GetRigVMGraphFunctionStore())
 			{
diff --git a/Engine/Plugins/Runtime/RigVM/Source/RigVMDeveloper/Private/RigVMDeveloperModule.cpp b/Engine/Plugins/Runtime/RigVM/Source/RigVMDeveloper/Private/RigVMDeveloperModule.cpp
--- a/Engine/Plugins/Runtime/RigVM/Source/RigVMDeveloper/Private/RigVMDeveloperModule.cpp
+++ b/Engine/Plugins/Runtime/RigVM/Source/RigVMDeveloper/Private/RigVMDeveloperModule.cpp
@@ -5,7 +5,7 @@
 =============================================================================*/
 
 #include "RigVMDeveloperModule.h"
-#include "RigVMBlueprint.h"
+#include "RigVMBlueprintLegacy.h"
 #include "RigVMBlueprintCompiler.h"
 #include "RigVMModel/RigVMBuildData.h"
 
diff --git a/Engine/Plugins/Runtime/RigVM/Source/RigVMDeveloper/Private/RigVMModel/Nodes/RigVMUnitNode.cpp b/Engine/Plugins/Runtime/RigVM/Source/RigVMDeveloper/Private/RigVMModel/Nodes/RigVMUnitNode.cpp
--- a/Engine/Plugins/Runtime/RigVM/Source/RigVMDeveloper/Private/RigVMModel/Nodes/RigVMUnitNode.cpp
+++ b/Engine/Plugins/Runtime/RigVM/Source/RigVMDeveloper/Private/RigVMModel/Nodes/RigVMUnitNode.cpp
@@ -6,7 +6,7 @@
 #include "RigVMModel/RigVMControllerActions.h"
 #include "RigVMUserWorkflowRegistry.h"
 #include "RigVMHost.h"
-#include "RigVMBlueprint.h"
+#include "RigVMBlueprintLegacy.h"
 
 #include UE_INLINE_GENERATED_CPP_BY_NAME(RigVMUnitNode)
 
@@ -484,7 +484,7 @@ void URigVMUnitNode::EnumeratePropertiesOnHostAndStructInstance(URigVMHost* InHo
 	check(InHost);
 	check(InInstance.IsValid() && InInstance->IsValid());
 
-	const URigVMBlueprint* Blueprint = GetTypedOuter<URigVMBlueprint>();
+	const IRigVMAssetInterface* Blueprint = GetImplementingOuter<IRigVMAssetInterface>();
 	if(Blueprint == nullptr)
 	{
 		return;
@@ -768,7 +768,7 @@ void URigVMUnitNode::ComputePinValueDifferences(TSharedPtr<FStructOnScope> InDes
 
 bool URigVMUnitNode::IsPartOfRuntime() const
 {
-	if(const URigVMBlueprint* Blueprint = GetTypedOuter<URigVMBlueprint>())
+	if(const IRigVMAssetInterface* Blueprint = GetImplementingOuter<IRigVMAssetInterface>())
 	{
 		if(URigVMHost* DebuggedHost = Cast<URigVMHost>(Blueprint->GetObjectBeingDebugged()))
 		{
diff --git a/Engine/Plugins/Runtime/RigVM/Source/RigVMDeveloper/Private/RigVMModel/RigVMBuildData.cpp b/Engine/Plugins/Runtime/RigVM/Source/RigVMDeveloper/Private/RigVMModel/RigVMBuildData.cpp
--- a/Engine/Plugins/Runtime/RigVM/Source/RigVMDeveloper/Private/RigVMModel/RigVMBuildData.cpp
+++ b/Engine/Plugins/Runtime/RigVM/Source/RigVMDeveloper/Private/RigVMModel/RigVMBuildData.cpp
@@ -1,14 +1,13 @@
 // Copyright Epic Games, Inc. All Rights Reserved.
 
 #include "RigVMModel/RigVMBuildData.h"
-#include "RigVMBlueprint.h"
+#include "RigVMAsset.h"
 #include "RigVMBlueprintGeneratedClass.h"
 #include "AssetRegistry/AssetRegistryModule.h"
 #include "RigVMModel/RigVMFunctionLibrary.h"
 #include "UObject/UObjectIterator.h"
 #include "RigVMModel/RigVMClient.h"
 #include "UObject/StrongObjectPtr.h"
-#include "RigVMBlueprint.h"
 
 #if WITH_EDITOR
 #include "Editor.h"
@@ -78,10 +77,10 @@ TArray<UClass*> URigVMBuildData::FindAllRigVMAssetClasses()
 		{
 			ImplementedClasses.AddUnique(*ClassIterator);
 
-			const UObject* CDO = ClassIterator->GetDefaultObject();
-			if(const URigVMBlueprint* BlueprintCDO = Cast<URigVMBlueprint>(CDO))
+			UObject* CDO = ClassIterator->GetDefaultObject();
+			if(const TRigVMAssetInterfacePtr BlueprintCDO = CDO)
 			{
-				if(UClass* GeneratedClass = BlueprintCDO->GetRigVMBlueprintGeneratedClass())
+				if(UClass* GeneratedClass = BlueprintCDO->GetRigVMGeneratedClass())
 				{
 					ImplementedClasses.AddUnique(GeneratedClass);
 				}
@@ -146,25 +145,25 @@ void URigVMBuildData::SetupRigVMGraphFunctionPointers()
 		}
 		if(IRigVMClientHost* ClientHost = Cast<IRigVMClientHost>(InObject))
 		{
-			if(IRigVMGraphFunctionHost* FunctionHost = ClientHost->GetRigVMGraphFunctionHost())
+			if(TScriptInterface<IRigVMGraphFunctionHost> FunctionHost = ClientHost->GetRigVMGraphFunctionHost())
 			{
-				return FunctionHost;
+				return FunctionHost.GetInterface();
 			}
 		}
 		if(IRigVMClientHost* OuterClientHost = InObject->GetImplementingOuter<IRigVMClientHost>())
 		{
-			if(IRigVMGraphFunctionHost* FunctionHost = OuterClientHost->GetRigVMGraphFunctionHost())
+			if(TScriptInterface<IRigVMGraphFunctionHost> FunctionHost = OuterClientHost->GetRigVMGraphFunctionHost())
 			{
-				return FunctionHost;
+				return FunctionHost.GetInterface();
 			}
 		}
-		if(URigVMBlueprint* Blueprint = Cast<URigVMBlueprint>(InObject))
+		if(IRigVMClientHost* Blueprint = Cast<IRigVMClientHost>(InObject))
 		{
-			return Blueprint->GetRigVMGraphFunctionHost();
+			return Blueprint->GetRigVMGraphFunctionHost().GetInterface();
 		}
-		if(URigVMBlueprint* Blueprint = InObject->GetTypedOuter<URigVMBlueprint>())
+		if(URigVMClientHost* Blueprint = InObject->GetTypedOuter<URigVMClientHost>())
 		{
-			return Blueprint->GetRigVMGraphFunctionHost();
+			return Cast<IRigVMClientHost>(Blueprint)->GetRigVMGraphFunctionHost().GetInterface();
 		}
 		return nullptr;
 	};
@@ -180,8 +179,8 @@ void URigVMBuildData::TearDownRigVMGraphFunctionPointers()
 TArray<FRigVMGraphFunctionHeader> URigVMBuildData::GetFunctionHeadersForAsset(const FAssetData& InAssetData)
 {
 	TArray<FRigVMGraphFunctionHeader> Result;
-	if (!InAssetData.IsInstanceOf(URigVMBlueprint::StaticClass()) &&
-	!InAssetData.IsInstanceOf(URigVMBlueprintGeneratedClass::StaticClass()))
+	if (!InAssetData.GetClass()->ImplementsInterface(URigVMAssetInterface::StaticClass()) &&
+		!InAssetData.IsInstanceOf(URigVMBlueprintGeneratedClass::StaticClass()))
 	{
 		return Result;
 	}
@@ -191,12 +190,12 @@ TArray<FRigVMGraphFunctionHeader> URigVMBuildData::GetFunctionHeadersForAsset(co
 	if (InAssetData.IsAssetLoaded())
 	{
 		UObject* AssetObject = InAssetData.GetAsset();
-		URigVMBlueprintGeneratedClass* GeneratedClass = Cast<URigVMBlueprintGeneratedClass>(AssetObject);
+		TScriptInterface<IRigVMGraphFunctionHost> GeneratedClass = AssetObject;
 		if (!GeneratedClass)
 		{
-			if (URigVMBlueprint* RigVMBlueprint = Cast<URigVMBlueprint>(AssetObject))
+			if (TRigVMAssetInterfacePtr RigVMBlueprint = AssetObject)
 			{
-				GeneratedClass = RigVMBlueprint->GetRigVMBlueprintGeneratedClass();
+				GeneratedClass = RigVMBlueprint->GetRigVMGeneratedClass();
 			}
 		}
 
@@ -205,8 +204,8 @@ TArray<FRigVMGraphFunctionHeader> URigVMBuildData::GetFunctionHeadersForAsset(co
 			for (int32 Pass=0; Pass<2; Pass++)
 			{
 				TArray<FRigVMGraphFunctionData>& Functions = (Pass == 0)
-					? GeneratedClass->GraphFunctionStore.PrivateFunctions
-					: GeneratedClass->GraphFunctionStore.PublicFunctions;
+					? GeneratedClass->GetRigVMGraphFunctionStore()->PrivateFunctions
+					: GeneratedClass->GetRigVMGraphFunctionStore()->PublicFunctions;
 				
 				for (FRigVMGraphFunctionData& Data : Functions)
 				{
@@ -222,8 +221,8 @@ TArray<FRigVMGraphFunctionHeader> URigVMBuildData::GetFunctionHeadersForAsset(co
 	const FString PublicGraphFunctionsString = InAssetData.GetTagValueRef<FString>(TEXT("PublicGraphFunctions"));
 	if(!PublicGraphFunctionsString.IsEmpty())
 	{
-		static const FName FunctionsPropertyName = GET_MEMBER_NAME_CHECKED(URigVMBlueprint, PublicGraphFunctions);
-		const FArrayProperty* PublicGraphFunctionsProperty = CastField<FArrayProperty>(URigVMBlueprint::StaticClass()->FindPropertyByName(FunctionsPropertyName));
+		static const FName FunctionsPropertyName = TEXT("PublicGraphFunctions");
+		const FArrayProperty* PublicGraphFunctionsProperty = CastField<FArrayProperty>(InAssetData.GetClass()->FindPropertyByName(FunctionsPropertyName));
 		PublicGraphFunctionsProperty->ImportText_Direct(*PublicGraphFunctionsString, &Result, nullptr, EPropertyPortFlags::PPF_None);
 	}
 	return Result;
@@ -532,7 +531,7 @@ TArray<FRigVMVariantRef> URigVMBuildData::GatherAllFunctionVariantRefs() const
 	const FAssetRegistryModule& AssetRegistryModule = FModuleManager::LoadModuleChecked<FAssetRegistryModule>("AssetRegistry");
 	TArray<FAssetData> Assets;
 	FARFilter AssetFilter;
-	AssetFilter.ClassPaths.Add(URigVMBlueprint::StaticClass()->GetClassPathName());
+	AssetFilter.ClassPaths.Add(URigVMAssetInterface::StaticClass()->GetClassPathName());
 	AssetFilter.ClassPaths.Add(URigVMBlueprintGeneratedClass::StaticClass()->GetClassPathName());
 	AssetFilter.bRecursiveClasses = true;
 	AssetRegistryModule.Get().GetAssets(AssetFilter, Assets);
@@ -622,7 +621,7 @@ FRigVMVariantRef URigVMBuildData::CreateFunctionVariant(const FRigVMGraphFunctio
 	{
 		if(URigVMFunctionLibrary* FunctionLibrary = Cast<URigVMFunctionLibrary>(LibraryNode->GetGraph()))
 		{
-			if(URigVMBlueprint* RigVMBlueprint = FunctionLibrary->GetTypedOuter<URigVMBlueprint>())
+			if(IRigVMAssetInterface* RigVMBlueprint = FunctionLibrary->GetImplementingOuter<IRigVMAssetInterface>())
 			{
 				if(URigVMController* Controller = RigVMBlueprint->GetOrCreateController(FunctionLibrary))
 				{
@@ -643,7 +642,7 @@ TArray<FRigVMVariantRef> URigVMBuildData::GatherAllAssetVariantRefs() const
 	const FAssetRegistryModule& AssetRegistryModule = FModuleManager::LoadModuleChecked<FAssetRegistryModule>("AssetRegistry");
 	TArray<FAssetData> Assets;
 	FARFilter AssetFilter;
-	AssetFilter.ClassPaths.Add(URigVMBlueprint::StaticClass()->GetClassPathName());
+	AssetFilter.ClassPaths.Add(URigVMAssetInterface::StaticClass()->GetClassPathName());
 	AssetFilter.bRecursiveClasses = true;
 	AssetRegistryModule.Get().GetAssets(AssetFilter, Assets);
 
@@ -700,7 +699,7 @@ FRigVMVariantRef URigVMBuildData::CreateAssetVariant(const FAssetData& InAssetDa
 	
 	if(UObject* DuplicatedAsset = EditorAssetSubsystem->DuplicateAsset(SourcePackageLongName, DestinationPackageLongName))
 	{
-		if(URigVMBlueprint* RigVMBlueprint = Cast<URigVMBlueprint>(DuplicatedAsset))
+		if(TRigVMAssetInterfacePtr RigVMBlueprint = DuplicatedAsset)
 		{
 			// since we duplicated the asset the variant guid will be the same too,
 			// but to be sure we'll run the join code path as well.
@@ -727,7 +726,7 @@ FRigVMVariantRef URigVMBuildData::GetVariantRefForAsset(const FAssetData& InAsse
 	FRigVMVariant AssetVariant;
 	if(InAssetData.IsAssetLoaded())
 	{
-		if(const URigVMBlueprint* Blueprint = Cast<URigVMBlueprint>(InAssetData.GetAsset()))
+		if(const IRigVMAssetInterface* Blueprint = Cast<IRigVMAssetInterface>(InAssetData.GetAsset()))
 		{
 			AssetVariant = Blueprint->GetAssetVariant();
 		}
@@ -735,8 +734,8 @@ FRigVMVariantRef URigVMBuildData::GetVariantRefForAsset(const FAssetData& InAsse
 		
 	if (!AssetVariant.Guid.IsValid())
 	{
-		static const FName AssetVariantPropertyName = GET_MEMBER_NAME_CHECKED(URigVMBlueprint, AssetVariant);
-		const FProperty* AssetVariantProperty = CastField<FProperty>(URigVMBlueprint::StaticClass()->FindPropertyByName(AssetVariantPropertyName));
+		static const FName AssetVariantPropertyName = TEXT("AssetVariant");
+		const FProperty* AssetVariantProperty = CastField<FProperty>(InAssetData.GetClass()->FindPropertyByName(AssetVariantPropertyName));
 		const FString VariantStr = InAssetData.GetTagValueRef<FString>(AssetVariantPropertyName);
 		if(!VariantStr.IsEmpty())
 		{
@@ -763,11 +762,12 @@ FRigVMVariantRef URigVMBuildData::SplitVariantFromSet(const FRigVMVariantRef& In
 	if(Result.IsValid())
 	{
 		FAssetData AssetData = GetAssetDataForVariant(Result);
-		if(URigVMBlueprint* RigVMBlueprint = Cast<URigVMBlueprint>(AssetData.GetAsset()))
+		if (AssetData.GetClass()->ImplementsInterface(URigVMAssetInterface::StaticClass()))
 		{
+			TRigVMAssetInterfacePtr RigVMBlueprint(AssetData.GetAsset());
 			if(UObject* Subject = InVariantRef.ObjectPath.TryLoad())
 			{
-				if(Subject == RigVMBlueprint)
+				if(Subject == RigVMBlueprint->GetObject())
 				{
 					if(RigVMBlueprint->SplitAssetVariant())
 					{
@@ -799,11 +799,12 @@ FRigVMVariantRef URigVMBuildData::JoinVariantSet(const FRigVMVariantRef& InVaria
 	if(Result.IsValid() && InGuid.IsValid())
 	{
 		FAssetData AssetData = GetAssetDataForVariant(Result);
-		if(URigVMBlueprint* RigVMBlueprint = Cast<URigVMBlueprint>(AssetData.GetAsset()))
+		if (AssetData.GetClass()->ImplementsInterface(URigVMAssetInterface::StaticClass()))
 		{
+			TRigVMAssetInterfacePtr RigVMBlueprint = AssetData.GetAsset();
 			if(UObject* Subject = InVariantRef.ObjectPath.TryLoad())
 			{
-				if(Subject == RigVMBlueprint)
+				if(Subject == RigVMBlueprint->GetObject())
 				{
 					if(RigVMBlueprint->JoinAssetVariant(InGuid))
 					{
@@ -866,7 +867,7 @@ TArray<FRigVMGraphFunctionIdentifier> URigVMBuildData::GetAllFunctionIdentifiers
 			{
 				if(IRigVMClientHost* ClientHost = Cast<IRigVMClientHost>(AssetData.GetAsset()))
 				{
-					FunctionHost = ClientHost->GetRigVMGraphFunctionHost();
+					FunctionHost = ClientHost->GetRigVMGraphFunctionHost().GetInterface();
 				}
 			}
 			if(FunctionHost)
diff --git a/Engine/Plugins/Runtime/RigVM/Source/RigVMDeveloper/Private/RigVMModel/RigVMClient.cpp b/Engine/Plugins/Runtime/RigVM/Source/RigVMDeveloper/Private/RigVMModel/RigVMClient.cpp
--- a/Engine/Plugins/Runtime/RigVM/Source/RigVMDeveloper/Private/RigVMModel/RigVMClient.cpp
+++ b/Engine/Plugins/Runtime/RigVM/Source/RigVMDeveloper/Private/RigVMModel/RigVMClient.cpp
@@ -1,6 +1,8 @@
 // Copyright Epic Games, Inc. All Rights Reserved.
 
 #include "RigVMModel/RigVMClient.h"
+
+#include "RigVMAsset.h"
 #include "Misc/TransactionObjectEvent.h"
 #include "Misc/ScopedSlowTask.h"
 #include "Exporters/Exporter.h"
@@ -697,7 +699,7 @@ static void SetGetFunctionHostObjectPathDelegate(FRigVMClient* RigVMClient, URig
 					{
 						if (IRigVMClientHost* ClientHost = Cast<IRigVMClientHost>(WeakClientHost.Get()))
 						{
-							return Cast<UObject>(ClientHost->GetRigVMGraphFunctionHost());
+							return ClientHost->GetRigVMGraphFunctionHost().GetObject();
 						}
 					}
 					return nullptr;
@@ -1460,7 +1462,7 @@ void FRigVMClient::PreSave(FObjectPreSaveContext ObjectSaveContext)
 	{
 		if (IRigVMClientHost* ClientHost = Cast<IRigVMClientHost>(GetOuter()))
 		{
-			if (IRigVMGraphFunctionHost* FunctionHost = ClientHost->GetRigVMGraphFunctionHost())
+			if (TScriptInterface<IRigVMGraphFunctionHost> FunctionHost = ClientHost->GetRigVMGraphFunctionHost())
 			{
 				if (FRigVMGraphFunctionStore* Store = FunctionHost->GetRigVMGraphFunctionStore())
 				{
@@ -1491,7 +1493,7 @@ void FRigVMClient::HandleGraphModifiedEvent(ERigVMGraphNotifType InNotifType, UR
 #if WITH_EDITOR
 
 	IRigVMClientHost* ClientHost = Cast<IRigVMClientHost>(GetOuter());
-	IRigVMGraphFunctionHost* FunctionHost = nullptr;
+	TScriptInterface<IRigVMGraphFunctionHost> FunctionHost = nullptr;
 	FRigVMGraphFunctionStore* FunctionStore = nullptr;
 	if (ClientHost)
 	{
@@ -1515,7 +1517,7 @@ void FRigVMClient::HandleGraphModifiedEvent(ERigVMGraphNotifType InNotifType, UR
 					{
 						if (FunctionStore && FunctionHost)
 						{
-							FunctionStore->AddFunction(CollapseNode->GetFunctionHeader(FunctionHost), false);
+							FunctionStore->AddFunction(CollapseNode->GetFunctionHeader(FunctionHost.GetInterface()), false);
 						}
 					}
 				}
@@ -1542,7 +1544,7 @@ void FRigVMClient::HandleGraphModifiedEvent(ERigVMGraphNotifType InNotifType, UR
 					{
 						if (FunctionStore && FunctionHost)
 						{
-							FunctionStore->RemoveFunction(FRigVMGraphFunctionIdentifier (Cast<UObject>(FunctionHost), CollapseNode->GetPathName()));
+							FunctionStore->RemoveFunction(FRigVMGraphFunctionIdentifier (FunctionHost.GetObject(), CollapseNode->GetPathName()));
 						}
 					}
 				}
@@ -1733,7 +1735,7 @@ FRigVMGraphFunctionStore* FRigVMClient::FindFunctionStore(const URigVMLibraryNod
 {
 	if (IRigVMClientHost* ClientHost = InLibraryNode->GetImplementingOuter<IRigVMClientHost>())
 	{
-		if (IRigVMGraphFunctionHost* FunctionHost = ClientHost->GetRigVMGraphFunctionHost())
+		if (TScriptInterface<IRigVMGraphFunctionHost> FunctionHost = ClientHost->GetRigVMGraphFunctionHost())
 		{
 			return FunctionHost->GetRigVMGraphFunctionStore();
 		}
@@ -1786,11 +1788,11 @@ bool FRigVMClient::UpdateGraphFunctionData(const URigVMLibraryNode* InLibraryNod
 {
 	if (IRigVMClientHost* ClientHost = InLibraryNode->GetImplementingOuter<IRigVMClientHost>())
 	{
-		if (IRigVMGraphFunctionHost* FunctionHost = ClientHost->GetRigVMGraphFunctionHost())
+		if (TScriptInterface<IRigVMGraphFunctionHost> FunctionHost = ClientHost->GetRigVMGraphFunctionHost())
 		{
 			if (FRigVMGraphFunctionStore* Store = FindFunctionStore(InLibraryNode))
 			{
-				if (FRigVMGraphFunctionData* Data = Store->UpdateFunctionInterface(InLibraryNode->GetFunctionHeader(FunctionHost)))
+				if (FRigVMGraphFunctionData* Data = Store->UpdateFunctionInterface(InLibraryNode->GetFunctionHeader(FunctionHost.GetInterface())))
 				{
 					UpdateFunctionReferences(Data->Header, false, false);
 					return true;
diff --git a/Engine/Plugins/Runtime/RigVM/Source/RigVMDeveloper/Private/RigVMModel/RigVMController.cpp b/Engine/Plugins/Runtime/RigVM/Source/RigVMDeveloper/Private/RigVMModel/RigVMController.cpp
--- a/Engine/Plugins/Runtime/RigVM/Source/RigVMDeveloper/Private/RigVMModel/RigVMController.cpp
+++ b/Engine/Plugins/Runtime/RigVM/Source/RigVMDeveloper/Private/RigVMModel/RigVMController.cpp
@@ -1,7 +1,6 @@
 // Copyright Epic Games, Inc. All Rights Reserved.
 
 #include "RigVMModel/RigVMController.h"
-#include "RigVMBlueprint.h"
 #include "RigVMModel/RigVMControllerActions.h"
 #include "RigVMModel/Nodes/RigVMFunctionEntryNode.h"
 #include "RigVMModel/Nodes/RigVMFunctionReturnNode.h"
@@ -27,6 +26,7 @@
 #include "Algo/Reverse.h"
 #include "Algo/Transform.h"
 #include "RigVMPythonUtils.h"
+#include "RigVMSettings.h"
 #include "RigVMTypeUtils.h"
 #include "RigVMStringUtils.h"
 #include "StructUtils/UserDefinedStruct.h"
@@ -1452,7 +1452,7 @@ FRigVMGraphFunctionStore* URigVMController::GetGraphFunctionStore() const
 		return nullptr;
 	}
 
-	IRigVMGraphFunctionHost* FunctionHost = ClientHost->GetRigVMGraphFunctionHost();
+	TScriptInterface<IRigVMGraphFunctionHost> FunctionHost = ClientHost->GetRigVMGraphFunctionHost();
 	if (!FunctionHost)
 	{
 		ReportErrorf(TEXT("Cannot find function host."));
@@ -6933,7 +6933,7 @@ void URigVMController::SetReferencedFunction(URigVMFunctionReferenceNode* InFunc
 	
 	FRigVMGraphFunctionHeader OldReferencedNode = InFunctionRefNode->GetReferencedFunctionHeader();
 	IRigVMClientHost* ClientHost = GetImplementingOuter<IRigVMClientHost>();
-	InFunctionRefNode->ReferencedFunctionHeader = InNewReferencedNode->GetFunctionHeader(ClientHost ? ClientHost->GetRigVMGraphFunctionHost() : nullptr);
+	InFunctionRefNode->ReferencedFunctionHeader = InNewReferencedNode->GetFunctionHeader(ClientHost ? ClientHost->GetRigVMGraphFunctionHost().GetInterface() : nullptr);
 	
 	if(!(OldReferencedNode == InFunctionRefNode->ReferencedFunctionHeader))
 	{
@@ -14481,7 +14481,7 @@ URigVMFunctionReferenceNode* URigVMController::AddFunctionReferenceNodeFromDescr
 		bool bLocal = false;
 		if(IRigVMClientHost* ClientHost = GetImplementingOuter<IRigVMClientHost>())
 		{
-			if (FunctionHeader->LibraryPointer.HostObject == Cast<UObject>(ClientHost->GetRigVMGraphFunctionHost()))
+			if (FunctionHeader->LibraryPointer.HostObject == ClientHost->GetRigVMGraphFunctionHost().GetObject())
 			{
 				bLocal = true;
 				RigVMPythonUtils::Print(GetSchema()->GetGraphOuterName(GetGraph()), 
@@ -14987,7 +14987,7 @@ URigVMLibraryNode* URigVMController::CreateFunctionVariant(const FName& InFuncti
 		return Result;
 	}
 
-	IRigVMGraphFunctionHost* FunctionHost = ClientHost->GetRigVMGraphFunctionHost();
+	TScriptInterface<IRigVMGraphFunctionHost> FunctionHost = ClientHost->GetRigVMGraphFunctionHost();
 	if (!FunctionHost)
 	{
 		ReportErrorf(TEXT("Cannot find function host."));
diff --git a/Engine/Plugins/Runtime/RigVM/Source/RigVMDeveloper/Private/RigVMModel/RigVMSchema.cpp b/Engine/Plugins/Runtime/RigVM/Source/RigVMDeveloper/Private/RigVMModel/RigVMSchema.cpp
--- a/Engine/Plugins/Runtime/RigVM/Source/RigVMDeveloper/Private/RigVMModel/RigVMSchema.cpp
+++ b/Engine/Plugins/Runtime/RigVM/Source/RigVMDeveloper/Private/RigVMModel/RigVMSchema.cpp
@@ -349,9 +349,10 @@ bool URigVMSchema::CanAddNode(URigVMController* InController, const URigVMNode*
 	}
 	else if(const URigVMCollapseNode* CollapseNode = Cast<URigVMCollapseNode>(InNode))
 	{
-		const URigVMGraph* ContainedGraph = CollapseNode->GetContainedGraph();
+		URigVMGraph* ContainedGraph = CollapseNode->GetContainedGraph();
 		const TArray<URigVMNode*> ContainedNodes = CollapseNode->GetContainedNodes();
 		URigVMController* ContainedController = InController->GetControllerForGraph(ContainedGraph);
+		TGuardValue<bool> GuardEditGraph(ContainedGraph->bEditable, true);
 		for(const URigVMNode* ContainedNode : ContainedNodes)
 		{
 			if(!CanAddNode(ContainedController, ContainedNode))
diff --git a/Engine/Plugins/Runtime/RigVM/Source/RigVMDeveloper/Public/EdGraph/RigVMEdGraph.h b/Engine/Plugins/Runtime/RigVM/Source/RigVMDeveloper/Public/EdGraph/RigVMEdGraph.h
--- a/Engine/Plugins/Runtime/RigVM/Source/RigVMDeveloper/Public/EdGraph/RigVMEdGraph.h
+++ b/Engine/Plugins/Runtime/RigVM/Source/RigVMDeveloper/Public/EdGraph/RigVMEdGraph.h
@@ -2,6 +2,7 @@
 
 #pragma once
 
+#include "RigVMAsset.h"
 #include "EdGraph/EdGraph.h"
 #include "EdGraph/RigVMEdGraphNode.h"
 #include "RigVMModel/RigVMGraph.h"
@@ -11,9 +12,10 @@
 
 #define UE_API RIGVMDEVELOPER_API
 
-class URigVMBlueprint;
+class IRigVMAssetInterface;
 class URigVMEdGraphSchema;
 class URigVMController;
+class URigVMBlueprint;
 
 struct FRigVMStringTag
 {
@@ -143,9 +145,13 @@ public:
 	UE_API virtual void HandleRigVMGraphRenamed(const FString& InOldNodePath, const FString& InNewNodePath) override;
 
 	/** Set up this graph */
+	UE_DEPRECATED(5.7, "Please use const TRigVMAssetInterfacePtr GetAssetDefaultObject() const")
 	UE_API const URigVMBlueprint* GetBlueprintDefaultObject() const;
+	UE_API const TRigVMAssetInterfacePtr GetAssetDefaultObject() const;
 	UE_API void SetBlueprintClass(const UClass* InClass);
-	UE_API virtual void InitializeFromBlueprint(URigVMBlueprint* InBlueprint);
+	UE_DEPRECATED(5.7, "Please use void InitializeFromAsset(TRigVMAssetInterfacePtr InAsset);")
+	UE_API virtual void InitializeFromBlueprint(URigVMBlueprint* InBlueprint) {}
+	UE_API virtual void InitializeFromAsset(TRigVMAssetInterfacePtr InAsset);
 	UE_API bool IsPreviewGraph() const;
 
 	/** Get the ed graph schema */
@@ -164,7 +170,9 @@ public:
 
 	UE_API UEdGraphNode* FindNodeForModelNodeName(const FName& InModelNodeName, const bool bCacheIfRequired = true);
 
+	UE_DEPRECATED(5.7, "Please use TRigVMAssetInterfacePtr GetAsset() const")
 	UE_API URigVMBlueprint* GetBlueprint() const;
+	UE_API TRigVMAssetInterfacePtr GetAsset() const;
 	UE_API URigVMGraph* GetModel() const;
 	UE_API URigVMController* GetController() const;
 	bool IsRootGraph() const { return GetRootGraph() == this; }
@@ -211,11 +219,11 @@ private:
 
 #endif
 	friend class URigVMEdGraphNode;
-	friend class URigVMBlueprint;
+	friend class IRigVMAssetInterface;
 	friend class FRigVMEditorBase;
 	friend class SRigVMGraphNode;
-	friend class URigVMBlueprint;
 
+	friend class URigVMEdGraphNodeSpawner;
 	friend class URigVMEdGraphUnitNodeSpawner;
 	friend class URigVMEdGraphVariableNodeSpawner;
 	friend class URigVMEdGraphParameterNodeSpawner;
diff --git a/Engine/Plugins/Runtime/RigVM/Source/RigVMDeveloper/Public/EdGraph/RigVMEdGraphNode.h b/Engine/Plugins/Runtime/RigVM/Source/RigVMDeveloper/Public/EdGraph/RigVMEdGraphNode.h
--- a/Engine/Plugins/Runtime/RigVM/Source/RigVMDeveloper/Public/EdGraph/RigVMEdGraphNode.h
+++ b/Engine/Plugins/Runtime/RigVM/Source/RigVMDeveloper/Public/EdGraph/RigVMEdGraphNode.h
@@ -4,6 +4,7 @@
 #pragma once
 
 #include "CoreMinimal.h"
+#include "RigVMAsset.h"
 #include "UObject/WeakObjectPtrTemplates.h"
 #include "RigVMModel/RigVMGraph.h"
 #include "Blueprint/BlueprintExtension.h"
@@ -11,10 +12,11 @@
 
 #define UE_API RIGVMDEVELOPER_API
 
+class IRigVMAssetInterface;
+class URigVMBlueprint;
 class FBlueprintActionDatabaseRegistrar;
 class UEdGraph;
 struct FSlateIcon;
-class URigVMBlueprint;
 
 /** Base class for RigVM editor side nodes */
 UCLASS(MinimalAPI)
@@ -156,7 +158,9 @@ public:
 	UE_API void CopyPinDefaultsToModel(UEdGraphPin* Pin, bool bUndo = false, bool bPrintPythonCommand = false);
 
 	/** Get the blueprint that this node is contained within */
+	UE_DEPRECATED(5.7, "Please use TRigVMAssetInterfacePtr GetAsset() const")
 	UE_API URigVMBlueprint* GetBlueprint() const;
+	UE_API TRigVMAssetInterfacePtr GetAsset() const;
 
 	/** Get the VM model this node lives within */
 	UE_API URigVMGraph* GetModel() const;
@@ -272,7 +276,6 @@ protected:
 	UE_API void SetColorFromModel(const FLinearColor& InColor);
 
 	UE_API UClass* GetRigVMGeneratedClass() const;
-	UE_API UClass* GetRigVMSkeletonGeneratedClass() const;
 
 	static UE_API FEdGraphPinType GetPinTypeForModelPin(const URigVMPin* InModelPin);
 	static UE_API FEdGraphPinType GetPinTypeForCategoryPin();
diff --git a/Engine/Plugins/Runtime/RigVM/Source/RigVMDeveloper/Public/EdGraph/RigVMEdGraphSchema.h b/Engine/Plugins/Runtime/RigVM/Source/RigVMDeveloper/Public/EdGraph/RigVMEdGraphSchema.h
--- a/Engine/Plugins/Runtime/RigVM/Source/RigVMDeveloper/Public/EdGraph/RigVMEdGraphSchema.h
+++ b/Engine/Plugins/Runtime/RigVM/Source/RigVMDeveloper/Public/EdGraph/RigVMEdGraphSchema.h
@@ -2,6 +2,7 @@
 
 #pragma once
 
+#include "RigVMAsset.h"
 #include "EdGraph/RigVMEdGraphNode.h"
 #include "EdGraph/EdGraphSchema.h"
 #include "GraphEditorDragDropAction.h"
@@ -14,7 +15,6 @@
 
 #define UE_API RIGVMDEVELOPER_API
 
-class URigVMBlueprint;
 class URigVMEdGraph;
 class URigVMEdGraphNode;
 
@@ -87,7 +87,7 @@ class FRigVMLocalVariableNameValidator : public FStringSetNameValidator
 {
 
 public:
-	UE_API FRigVMLocalVariableNameValidator(const class UBlueprint* Blueprint, const URigVMGraph* Graph, FName InExistingName = NAME_None);
+	UE_API FRigVMLocalVariableNameValidator(const TRigVMAssetInterfacePtr Blueprint, const URigVMGraph* Graph, FName InExistingName = NAME_None);
 
 	// Begin FNameValidatorInterface
 	UE_API virtual EValidatorResult IsValid(const FString& Name, bool bOriginal) override;
@@ -102,6 +102,7 @@ class FRigVMNameValidator : public FStringSetNameValidator
 
 public:
 	UE_API FRigVMNameValidator(const class UBlueprint* Blueprint, const UStruct* ValidationScope, FName InExistingName = NAME_None);
+	UE_API FRigVMNameValidator(const TRigVMAssetInterfacePtr Blueprint, const UStruct* ValidationScope, FName InExistingName = NAME_None);
 
 	// Begin FNameValidatorInterface
 	UE_API virtual EValidatorResult IsValid(const FString& Name, bool bOriginal) override;
@@ -282,11 +283,11 @@ protected:
 	/** Constructor */
 	UE_API FRigVMFunctionDragDropAction();
 
-	static UE_API TSharedRef<FRigVMFunctionDragDropAction> New(TSharedPtr<FEdGraphSchemaAction> InAction, URigVMBlueprint* InRigBlueprint, URigVMEdGraph* InRigGraph);
+	static UE_API TSharedRef<FRigVMFunctionDragDropAction> New(TSharedPtr<FEdGraphSchemaAction> InAction, TRigVMAssetInterfacePtr InRigBlueprint, URigVMEdGraph* InRigGraph);
 
 protected:
 
-	URigVMBlueprint* SourceRigBlueprint;
+	TScriptInterface<IRigVMAssetInterface> SourceRigBlueprint;
 	URigVMEdGraph* SourceRigGraph;
 	bool bControlDrag;
 	bool bAltDrag;
@@ -360,6 +361,7 @@ public:
 	virtual void SetPinBeingDroppedOnNode(UEdGraphPin* InSourcePin) const override { PinBeingDropped = InSourcePin; }
 	UE_API virtual void InsertAdditionalActions(TArray<UBlueprint*> InBlueprints, TArray<UEdGraph*> InGraphs, TArray<UEdGraphPin*> InPins, FGraphActionListBuilderBase& OutAllActions) const override;
 	UE_API virtual TSharedPtr<INameValidatorInterface> GetNameValidator(const UBlueprint* BlueprintObj, const FName& OriginalName, const UStruct* ValidationScope, const FName& ActionTypeId) const override;
+	UE_API virtual TSharedPtr<INameValidatorInterface> GetNameValidator(const TRigVMAssetInterfacePtr Asset, const FName& OriginalName, const UStruct* ValidationScope) const;
 	UE_API bool CanPromotePinToVariable(const UEdGraphPin& Pin, bool bCond) const;
 
 	/** Returns true if the schema supports the script type */
diff --git a/Engine/Plugins/Runtime/RigVM/Source/RigVMDeveloper/Public/RigVMAsset.h b/Engine/Plugins/Runtime/RigVM/Source/RigVMDeveloper/Public/RigVMAsset.h
new file mode 100644
--- /dev/null
+++ b/Engine/Plugins/Runtime/RigVM/Source/RigVMDeveloper/Public/RigVMAsset.h
@@ -0,0 +1,813 @@
+// Copyright Epic Games, Inc. All Rights Reserved.
+
+#pragma once
+
+#include "CoreMinimal.h"
+#include "RigVMCore/RigVM.h"
+#include "RigVMHost.h"
+#include "RigVMModel/RigVMClient.h"
+#include "RigVMModel/RigVMExternalDependency.h"
+#include "RigVMCompiler/RigVMCompiler.h"
+#include "RigVMCore/RigVMGraphFunctionDefinition.h"
+#include "RigVMSettings.h"
+#include "Blueprint/BlueprintExtension.h"
+#if WITH_EDITOR
+#include "HAL/CriticalSection.h"
+#include "Kismet2/CompilerResultsLog.h"
+#endif
+
+#include "RigVMAsset.generated.h"
+
+#define UE_API RIGVMDEVELOPER_API
+
+#if WITH_EDITOR
+class URigVMEdGraph;
+class IRigVMEditorModule;
+namespace UE::RigVM::Editor::Tools
+{
+	class FFilterByAssetTag;
+}
+#endif
+struct FEndLoadPackageContext;
+struct FRigVMMemoryStorageStruct;
+struct FGuardSkipDirtyBlueprintStatus;
+
+class IRigVMAssetInterface;
+typedef TScriptInterface<IRigVMAssetInterface> TRigVMAssetInterfacePtr;
+
+DECLARE_EVENT_ThreeParams(IRigVMAssetInterface, FOnRigVMCompiledEvent, UObject*, URigVM*, FRigVMExtendedExecuteContext&);
+DECLARE_EVENT_OneParam(IRigVMAssetInterface, FOnRigVMRefreshEditorEvent, TRigVMAssetInterfacePtr);
+DECLARE_EVENT_FourParams(IRigVMAssetInterface, FOnRigVMVariableDroppedEvent, UObject*, FProperty*, const FVector2D&, const FVector2D&);
+DECLARE_EVENT_OneParam(IRigVMAssetInterface, FOnRigVMExternalVariablesChanged, const TArray<FRigVMExternalVariable>&);
+DECLARE_EVENT_TwoParams(IRigVMAssetInterface, FOnRigVMNodeDoubleClicked, TRigVMAssetInterfacePtr, URigVMNode*);
+DECLARE_EVENT_OneParam(IRigVMAssetInterface, FOnRigVMGraphImported, UEdGraph*);
+DECLARE_EVENT_OneParam(IRigVMAssetInterface, FOnRigVMPostEditChangeChainProperty, FPropertyChangedChainEvent&);
+DECLARE_EVENT_FourParams(IRigVMAssetInterface, FOnRigVMLocalizeFunctionDialogRequested, FRigVMGraphFunctionIdentifier&, URigVMController*, IRigVMGraphFunctionHost*, bool);
+DECLARE_EVENT_ThreeParams(IRigVMAssetInterface, FOnRigVMReportCompilerMessage, EMessageSeverity::Type, UObject*, const FString&);
+DECLARE_DELEGATE_RetVal_FourParams(FRigVMController_BulkEditResult, FRigVMOnBulkEditDialogRequestedDelegate, TRigVMAssetInterfacePtr, URigVMController*, URigVMLibraryNode*, ERigVMControllerBulkEditType);
+DECLARE_DELEGATE_RetVal_OneParam(bool, FRigVMOnBreakLinksDialogRequestedDelegate, TArray<URigVMLink*>);
+DECLARE_DELEGATE_RetVal_OneParam(TRigVMTypeIndex, FRigVMOnPinTypeSelectionRequestedDelegate, const TArray<TRigVMTypeIndex>&);
+DECLARE_EVENT(IRigVMAssetInterface, FOnRigVMBreakpointAdded);
+DECLARE_EVENT_OneParam(IRigVMAssetInterface, FOnRigVMRequestInspectObject, const TArray<UObject*>& );
+DECLARE_EVENT_OneParam(IRigVMAssetInterface, FOnRigVMRequestInspectMemoryStorage, const TArray<FRigVMMemoryStorageStruct*>&);
+DECLARE_EVENT_OneParam( IRigVMAssetInterface, FOnRigVMAssetChangedEvent, class UObject * );
+DECLARE_EVENT_OneParam(IRigVMAssetInterface, FOnRigVMSetObjectBeingDebugged, UObject* /*InDebugObj*/);
+
+USTRUCT()
+struct UE_API FRigVMPythonSettings
+{
+	GENERATED_BODY();
+
+	FRigVMPythonSettings()
+	{
+	}
+};
+
+UENUM(BlueprintType)
+enum class ERigVMTagDisplayMode : uint8
+{
+	None = 0,
+	All = 0x001,
+	DeprecationOnly = 0x002,
+	Last = DeprecationOnly UMETA(Hidden), 
+};
+
+USTRUCT()
+struct UE_API FRigVMEdGraphDisplaySettings
+{
+	GENERATED_BODY();
+
+	FRigVMEdGraphDisplaySettings();
+
+	// When enabled shows the first node instruction index
+	// matching the execution stack window.
+	UPROPERTY(EditAnywhere, Category = "Graph Display Settings")
+	bool bShowNodeInstructionIndex;
+
+	// When enabled shows the node counts both in the graph view as
+	// we as in the execution stack window.
+	// The number on each node represents how often the node has been run.
+	// Keep in mind when looking at nodes in a function the count
+	// represents the sum of all counts for each node based on all
+	// references of the function currently running.
+	UPROPERTY(EditAnywhere, Category = "Graph Display Settings")
+	bool bShowNodeRunCounts;
+
+	// A lower limit for counts for nodes used for debugging.
+	// Any node lower than this count won't show the run count.
+	UPROPERTY(EditAnywhere, Category = "Graph Display Settings")
+	int32 NodeRunLowerBound;
+
+	// A upper limit for counts for nodes used for debugging.
+	// If a node reaches this count a warning will be issued for the
+	// node and displayed both in the execution stack as well as in the
+	// graph. Setting this to <= 1 disables the warning.
+	// Note: The count limit doesn't apply to functions / collapse nodes.
+	UPROPERTY(EditAnywhere, Category = "Graph Display Settings")
+	int32 NodeRunLimit;
+
+	// The duration in microseconds of the fastest instruction / node
+	UPROPERTY(EditAnywhere, Category = "Graph Display Settings", transient, meta = (EditCondition = "!bAutoDetermineRange"))
+	double MinMicroSeconds;
+
+	// The duration in microseconds of the slowest instruction / node
+	UPROPERTY(EditAnywhere, Category = "Graph Display Settings", transient, meta = (EditCondition = "!bAutoDetermineRange"))
+	double MaxMicroSeconds;
+
+	// The total duration of the last execution of the rig
+	UPROPERTY(VisibleAnywhere, Category = "Graph Display Settings", transient)
+	double TotalMicroSeconds;
+
+	// If you set this to more than 1 the results will be averaged across multiple frames
+	UPROPERTY(EditAnywhere, Category = "Graph Display Settings", meta = (UIMin=1, UIMax=256))
+	int32 AverageFrames;
+
+	TArray<double> MinMicroSecondsFrames;
+	TArray<double> MaxMicroSecondsFrames;
+	TArray<double> TotalMicroSecondsFrames;
+
+	UPROPERTY(EditAnywhere, Category = "Graph Display Settings")
+	bool bAutoDetermineRange;
+
+	UPROPERTY(transient)
+	double LastMinMicroSeconds;
+
+	UPROPERTY(transient)
+	double LastMaxMicroSeconds;
+
+	// The color of the fastest instruction / node
+	UPROPERTY(EditAnywhere, Category = "Graph Display Settings")
+	FLinearColor MinDurationColor;
+
+	// The color of the slowest instruction / node
+	UPROPERTY(EditAnywhere, Category = "Graph Display Settings")
+	FLinearColor MaxDurationColor;
+
+	// The color of the slowest instruction / node
+	UPROPERTY(EditAnywhere, Category = "Graph Display Settings")
+	ERigVMTagDisplayMode TagDisplayMode;
+
+	void SetTotalMicroSeconds(double InTotalMicroSeconds);
+	void SetLastMinMicroSeconds(double InMinMicroSeconds);
+	void SetLastMaxMicroSeconds(double InMaxMicroSeconds);
+	double AggregateAverage(TArray<double>& InFrames, double InPrevious, double InNext) const;
+};
+
+/**
+ * Enumerates states a RigVMAsset can be in.
+ * Copied from EBlueprintStatus.
+ */
+UENUM()
+enum ERigVMAssetStatus : int
+{
+	/** Asset is in an unknown state. */
+	RVMA_Unknown,
+	/** Asset has been modified but not recompiled. */
+	RVMA_Dirty,
+	/** Asset tried but failed to be compiled. */
+	RVMA_Error,
+	/** Asset has been compiled since it was last modified. */
+	RVMA_UpToDate,
+	/** Asset is in the process of being created for the first time. */
+	RVMA_BeingCreated,
+	/** Asset has been compiled since it was last modified. There are warnings. */
+	RVMA_UpToDateWithWarnings,
+	RVMA_MAX,
+};
+
+UINTERFACE(BlueprintType)
+class UE_API URigVMAssetInterface : public UInterface
+{
+	GENERATED_BODY()
+};
+
+class UE_API IRigVMAssetInterface 
+{
+	GENERATED_BODY()
+
+#if WITH_EDITOR
+public:
+	static TRigVMAssetInterfacePtr GetInterfaceOuter(const UObject* InObject);
+
+	virtual UObject* GetObject() = 0;
+	const UObject* GetObject() const { return const_cast<IRigVMAssetInterface*>(this)->GetObject(); }
+	virtual void Clear();
+
+	TScriptInterface<IRigVMClientHost> GetRigVMClientHost();
+	const TScriptInterface<IRigVMClientHost> GetRigVMClientHost() const { return const_cast<IRigVMAssetInterface*>(this)->GetRigVMClientHost(); }
+	void RequestClientHostRigVMInit();
+	
+	static inline const FLazyName RigVMPanelNodeFactoryName = FLazyName(TEXT("FRigVMEdGraphPanelNodeFactory"));
+	static inline const FLazyName RigVMPanelPinFactoryName = FLazyName(TEXT("FRigVMEdGraphPanelPinFactory"));
+
+	virtual void GetAllEdGraphs(TArray<UEdGraph*>& Graphs) const = 0;
+	virtual TArray<FRigVMGraphVariableDescription> GetAssetVariables() const = 0;
+	virtual bool IsRegeneratingOnLoad() const = 0;
+	virtual FProperty* FindGeneratedPropertyByName(const FName& InName) const = 0; 
+	virtual FName AddHostMemberVariableFromExternal(FRigVMExternalVariable InVariableToCreate, FString InDefaultValue = FString()) = 0;
+	virtual FName AddMemberVariable(const FName& InName, const FString& InCPPType, bool bIsPublic = false, bool bIsReadOnly = false, FString InDefaultValue = TEXT("")) = 0;
+	virtual bool RemoveMemberVariable(const FName& InName) = 0;
+	virtual bool RenameMemberVariable(const FName& InOldName, const FName& InNewName) = 0;
+	virtual bool ChangeMemberVariableType(const FName& InName, const FString& InCPPType, bool bIsPublic = false, bool bIsReadOnly = false, FString InDefaultValue = TEXT("")) = 0;
+	virtual bool ChangeMemberVariableType(const FName& InName, const FEdGraphPinType& InType) = 0;
+	virtual FText GetVariableTooltip(const FName& InName) const = 0;
+	virtual FString GetVariableCategory(const FName& InName) = 0;
+	virtual FString GetVariableMetadataValue(const FName& InName, const FName& InKey) = 0;
+	virtual bool SetVariableCategory(const FName& InName, const FString& InCategory) = 0;
+	virtual bool SetVariableTooltip(const FName& InName, const FText& InTooltip) = 0; 
+	virtual bool SetVariableExposeOnSpawn(const FName& InName, const bool bInExposeOnSpawn) = 0; 
+	virtual bool SetVariableExposeToCinematics(const FName& InName, const bool bInExposeToCinematics) = 0; 
+	virtual bool SetVariablePrivate(const FName& InName, const bool bInPrivate) = 0; 
+	virtual bool SetVariablePublic(const FName& InName, const bool bIsPublic) = 0; 
+	virtual FString OnCopyVariable(const FName& InName) const = 0;
+	virtual bool OnPasteVariable(const FString& InText) = 0;
+	virtual UObject* GetObjectBeingDebugged() = 0;
+	virtual UObject* GetObjectBeingDebugged() const = 0;
+	virtual const FString& GetObjectBeingDebuggedPath() const = 0;
+	virtual TArray<UObject*> GetArchetypeInstances(bool bIncludeCDO, bool bIncludeDerivedClass) const = 0;
+	virtual UWorld* GetWorldBeingDebugged() const = 0;
+	virtual void SetWorldBeingDebugged(UWorld* NewWorld) = 0;
+	virtual ERigVMAssetStatus GetAssetStatus() const = 0;
+	virtual bool IsUpToDate() const = 0;
+	virtual UClass* GetRigVMGeneratedClass() = 0;
+	virtual const UClass* GetRigVMGeneratedClass() const { return const_cast<IRigVMAssetInterface*>(this)->GetRigVMGeneratedClass(); }
+	virtual URigVM* GetVM(bool bCreateIfNeeded) const = 0;
+	
+	virtual bool ExportGraphToText(UEdGraph* InEdGraph, FString& OutText);
+	virtual bool TryImportGraphFromText(const FString& InClipboardText, UEdGraph** OutGraphPtr = nullptr);
+	virtual bool CanImportGraphFromText(const FString& InClipboardText);
+
+	virtual void RefreshAllNodes() ;
+	virtual void IncrementVMRecompileBracket();
+	virtual void DecrementVMRecompileBracket();
+
+	virtual void SetObjectBeingDebugged(UObject* NewObject);
+	virtual FOnRigVMSetObjectBeingDebugged& OnSetObjectBeingDebugged()  { return SetObjectBeingDebuggedEvent; }
+
+	/** Returns the editor module to be used for this blueprint */
+	virtual IRigVMEditorModule* GetEditorModule() const;
+
+	/** Returns true if a given panel node factory is compatible this blueprint */
+	virtual const FLazyName& GetPanelNodeFactoryName() const;
+
+	virtual FRigVMRuntimeSettings& GetVMRuntimeSettings() = 0;
+	virtual const FRigVMRuntimeSettings& GetVMRuntimeSettings() const { return const_cast<IRigVMAssetInterface*>(this)->GetVMRuntimeSettings();};
+	
+	
+	static void QueueCompilerMessageDelegate(const FOnRigVMReportCompilerMessage::FDelegate& InDelegate);
+	static void ClearQueuedCompilerMessageDelegates();
+
+	virtual FRigVMGraphModifiedEvent& OnModified();
+	FOnRigVMCompiledEvent& OnVMCompiled();
+
+	virtual FOnRigVMReportCompilerMessage& OnReportCompilerMessage()  { return ReportCompilerMessageEvent; }
+
+	/** Returns true if a given panel pin factory is compatible this blueprint */
+	virtual const FLazyName& GetPanelPinFactoryName() const;
+
+	virtual void AddPinWatch(UEdGraphPin* InPin) = 0;
+	virtual bool IsPinBeingWatched(const UEdGraphPin* InPin) const = 0;
+
+	TMap<FString, FRigVMOperand> PinToOperandMap;
+
+	/** Removes all the breakpoints from the blueprint and the VM */
+	void ClearBreakpoints();
+
+	/** Adds a breakpoint to all loaded blueprints which use the node indicated by InBreakpointNodePath 
+	  * If the node is inside a public function, it will add a breakpoint to all blueprints calling this function. */
+	bool AddBreakpoint(const FString& InBreakpointNodePath);
+
+	/** Adds a breakpoint to all loaded blueprints which use the InBreakpointNode. 
+	  * If LibraryNode is not null, it indicates that the library uses the InBreapointNode, and the function will add
+	  * breakpoints to any other loaded blueprint that references this library. */
+	bool AddBreakpoint(URigVMNode* InBreakpointNode, URigVMLibraryNode* LibraryNode = nullptr);
+
+	/** Removes the given breakpoint from all the loaded blueprints that use this node, and recomputes all breakpoints
+	  * in the VM  */
+	bool RemoveBreakpoint(const FString& InBreakpointNodePath);
+	bool RemoveBreakpoint(URigVMNode* InBreakpointNode);
+
+	virtual FRigVMController_RequestJumpToHyperlinkDelegate& OnRequestJumpToHyperlink()  { return RequestJumpToHyperlink; };
+
+	/** Returns the settings defaults for this blueprint */
+	URigVMEditorSettings* GetRigVMEditorSettings() const;
+
+	virtual UEdGraph* GetEdGraph(const URigVMGraph* InModel) const;
+	virtual UEdGraph* GetEdGraph(const FString& InNodePath) const;
+	
+	void AddVariableSearchMetaDataInfo(const FName InVariableName, TArray<UBlueprintExtension::FSearchTagDataPair>& OutTaggedMetaData) const;
+
+	virtual FRigVMEdGraphDisplaySettings& GetRigGraphDisplaySettings() = 0;
+	virtual const FRigVMEdGraphDisplaySettings& GetRigGraphDisplaySettings() const { return const_cast<IRigVMAssetInterface*>(this)->GetRigGraphDisplaySettings(); };
+	virtual FRigVMCompileSettings& GetVMCompileSettings() = 0;
+	virtual const FRigVMCompileSettings& GetVMCompileSettings() const { return const_cast<IRigVMAssetInterface*>(this)->GetVMCompileSettings(); };
+	
+
+	virtual URigVMHost* GetDebuggedRigVMHost();
+
+	virtual URigVMGraph* GetFocusedModel() const;
+	virtual TArray<URigVMGraph*> GetAllModels() const;
+	virtual URigVMFunctionLibrary* GetLocalFunctionLibrary() const;
+	virtual URigVMGraph* GetModel(const UEdGraph* InEdGraph = nullptr) const;
+	virtual URigVMGraph* GetModel(const FString& InNodePath) const;
+	virtual URigVMController* GetControllerByName(const FString InGraphName = TEXT("")) const;
+	virtual URigVMController* GetOrCreateController(URigVMGraph* InGraph = nullptr);
+	virtual URigVMController* GetController(const UEdGraph* InEdGraph) const;
+	virtual URigVMController* GetOrCreateController(const UEdGraph* InGraph);
+	virtual URigVMController* GetController(const URigVMGraph* InGraph = nullptr) const;
+
+	virtual FOnRigVMLocalizeFunctionDialogRequested& OnRequestLocalizeFunctionDialog()  { return RequestLocalizeFunctionDialog; }
+	void BroadcastRequestLocalizeFunctionDialog(FRigVMGraphFunctionIdentifier InFunction, bool bForce = false);
+
+	const FCompilerResultsLog& GetCompileLog() const;
+	FCompilerResultsLog& GetCompileLog();
+	void HandleReportFromCompiler(EMessageSeverity::Type InSeverity, UObject* InSubject, const FString& InMessage);
+
+	// Returns a list of dependencies of this blueprint.
+	// Dependencies are blueprints that contain functions used in this blueprint
+	TArray<IRigVMAssetInterface*> GetDependencies(bool bRecursive = false) const;
+
+	// Returns a list of dependents as unresolved soft object pointers.
+	// A dependent is a blueprint which uses a function defined in this blueprint.
+	// This function is not recursive, since it avoids opening the asset.
+	// Use GetDependentBlueprints as an alternative.
+	TArray<FAssetData> GetDependentAssets() const;
+
+	// Returns a list of dependents as resolved blueprints.
+	// A dependent is a blueprint which uses a function defined in this blueprint.
+	// If bOnlyLoaded is false, this function loads the dependent assets and can introduce a large cost
+	// depending on the size / count of assets in the project.
+	TArray<IRigVMAssetInterface*> GetDependentResolvedAssets(bool bRecursive = false, bool bOnlyLoaded = false) const;
+
+	void BroadcastRefreshEditor();
+	virtual FOnRigVMRefreshEditorEvent& OnRefreshEditor()  { return RefreshEditorEvent; }
+
+	FOnRigVMRequestInspectMemoryStorage& OnRequestInspectMemoryStorage();
+	void RequestInspectMemoryStorage(const TArray<FRigVMMemoryStorageStruct*>& InMemoryStorageStructs) const;
+
+	FOnRigVMPostEditChangeChainProperty& OnPostEditChangeChainProperty() { return PostEditChangeChainPropertyEvent; }
+	void BroadcastPostEditChangeChainProperty(FPropertyChangedChainEvent& PropertyChangedChainEvent);
+	virtual void MarkAssetAsStructurallyModified(bool bSkipDirtyAssetStatus = false) = 0;
+	virtual void MarkAssetAsModified(FPropertyChangedEvent PropertyChangedEvent = FPropertyChangedEvent(nullptr)) = 0;
+	
+protected:
+	IRigVMAssetInterface() {}
+	IRigVMAssetInterface(const FObjectInitializer& ObjectInitializer);
+
+	virtual FRigVMVariant& GetAssetVariant() = 0;
+	virtual const FRigVMVariant& GetAssetVariant() const { return const_cast<IRigVMAssetInterface*>(this)->GetAssetVariant(); }
+	virtual FRigVMVariantRef GetAssetVariantRef() const = 0;
+	virtual TArray<FRigVMGraphFunctionHeader>& GetPublicGraphFunctions() = 0;
+	
+	virtual void RemovePinWatch(UEdGraphPin* InPin) = 0;
+	virtual void ClearPinWatches() = 0;
+	virtual void ForeachPinWatch(TFunctionRef<void(UEdGraphPin*)> Task) = 0;
+	
+	virtual void SetAssetStatus(const ERigVMAssetStatus& InStatus) = 0;
+	
+	virtual TArray<TObjectPtr<UEdGraph>>& GetUberGraphs() = 0;
+	virtual const TArray<TObjectPtr<UEdGraph>>& GetUberGraphs() const { return const_cast<IRigVMAssetInterface*>(this)->GetUberGraphs(); };
+	
+	
+	virtual UClass* GetRigVMGeneratedClassPrototype() const = 0;
+		virtual TArray<struct FEditedDocumentInfo>& GetLastEditedDocuments() = 0;
+	virtual UObject* GetDefaultsObject() = 0;
+	virtual void SetObjectBeingDebuggedSuper(UObject* NewObject) = 0;
+	virtual void PreSaveSuper(FObjectPreSaveContext ObjectSaveContext) = 0;
+	virtual void PostLoadSuper() = 0;
+	virtual void SetUberGraphs(const TArray<TObjectPtr<UEdGraph>>& InGraphs) = 0;
+	virtual bool MarkPackageDirty() = 0;
+	virtual UPackage* GetPackage() = 0;
+	virtual void PostTransactedSuper(const FTransactionObjectEvent& TransactionEvent) = 0;
+	virtual void ReplaceDeprecatedNodesSuper() = 0;
+	virtual void PreDuplicateSuper(FObjectDuplicationParameters& DupParams) = 0;
+	virtual void PostDuplicateSuper(bool bDuplicateForPIE) = 0;
+	virtual void GetAssetRegistryTagsSuper(FAssetRegistryTagsContext Context) const = 0;
+	virtual bool RequiresForceLoadMembersSuper(UObject* InObject) const = 0;
+	virtual void PostEditChangeChainPropertySuper(struct FPropertyChangedChainEvent& PropertyChangedEvent) = 0;
+	virtual void PostRenameSuper(UObject* OldOuter, const FName OldName) = 0;
+	virtual void GetPreloadDependenciesSuper(TArray<UObject*>& OutDeps) = 0;
+	virtual TArray<FRigVMExternalVariable> GetExternalVariables(bool bFallbackToBlueprint) const = 0;
+	virtual FString GetVariableDefaultValue(const FName& InName, bool bFromDebuggedObject) const = 0;
+	virtual FRigVMExtendedExecuteContext* GetRigVMExtendedExecuteContext() = 0;
+	virtual void PreCompile() = 0;
+	virtual TMap<FString, FSoftObjectPath>& GetUserDefinedStructGuidToPathName(bool bFromCDO) = 0;
+	virtual const TMap<FString, FSoftObjectPath>& GetUserDefinedStructGuidToPathName(bool bFromCDO) const {return const_cast<IRigVMAssetInterface*>(this)->GetUserDefinedStructGuidToPathName(bFromCDO);}
+	virtual TMap<FString, FSoftObjectPath>& GetUserDefinedEnumToPathName(bool bFromCDO) = 0;
+	virtual const TMap<FString, FSoftObjectPath>& GetUserDefinedEnumToPathName(bool bFromCDO) const {return const_cast<IRigVMAssetInterface*>(this)->GetUserDefinedEnumToPathName(bFromCDO);}
+	virtual TSet<TObjectPtr<UObject>>& GetUserDefinedTypesInUse(bool bFromCDO) = 0;
+	virtual FCompilerResultsLog* GetCurrentMessageLog() const = 0;
+	virtual FRigVMDebugInfo& GetDebugInfo() = 0;
+	virtual TObjectPtr<URigVMEdGraph>& GetFunctionLibraryEdGraph() = 0;
+	virtual const TObjectPtr<URigVMEdGraph>& GetFunctionLibraryEdGraph() const { return const_cast<IRigVMAssetInterface*>(this)->GetFunctionLibraryEdGraph(); }
+	virtual TArray<TObjectPtr<UEdGraph>>& GetFunctionGraphs() = 0;
+	virtual const TArray<TObjectPtr<UEdGraph>>& GetFunctionGraphs() const { return const_cast<IRigVMAssetInterface*>(this)->GetFunctionGraphs();};
+	virtual bool& IsReferencedObjectPathsStored() = 0;
+	virtual TArray<FSoftObjectPath>& GetReferencedObjectPaths() = 0;
+	virtual TArray<FName>& GetSupportedEventNames() = 0;
+	virtual TArray<FRigVMReferenceNodeData>& GetFunctionReferenceNodeData() = 0;
+	virtual void NotifyGraphRenamedSuper(class UEdGraph* Graph, FName OldName, FName NewName) = 0;
+	virtual URigVMHost* CreateRigVMHostSuper(UObject* InOuter) = 0;
+	virtual void AddUbergraphPage(URigVMEdGraph* RigVMEdGraph) = 0;
+	virtual void AddLastEditedDocument(URigVMEdGraph* RigVMEdGraph) = 0;
+	virtual void Compile() = 0;
+	virtual FCompilerResultsLog CompileBlueprint() = 0;
+	virtual void PostEditChangeBlueprintActors() = 0;
+	virtual bool SplitAssetVariant() = 0;
+	virtual bool JoinAssetVariant(const FGuid& InGuid) = 0;
+	virtual TArray<FRigVMVariantRef> GetMatchingVariants() const = 0;
+	virtual void PatchVariablesOnLoadSuper() = 0;
+	virtual void SerializeSuper(FArchive& Ar) = 0;
+	virtual void PostSerialize(FArchive& Ar) = 0;
+
+
+	// Interface to URigVMExternalDependencyManager
+	void CollectExternalDependencies(TArray<FRigVMExternalDependency>& OutDependencies, const FName& InCategory, const FRigVMClient* InClient) const;
+	void CollectExternalDependencies(TArray<FRigVMExternalDependency>& OutDependencies, const FName& InCategory, const FRigVMGraphFunctionStore* InFunctionStore) const;
+	void CollectExternalDependencies(TArray<FRigVMExternalDependency>& OutDependencies, const FName& InCategory, const FRigVMGraphFunctionData* InFunction) const;
+	void CollectExternalDependencies(TArray<FRigVMExternalDependency>& OutDependencies, const FName& InCategory, const FRigVMGraphFunctionHeader* InHeader) const;
+	void CollectExternalDependencies(TArray<FRigVMExternalDependency>& OutDependencies, const FName& InCategory, const FRigVMFunctionCompilationData* InCompilationData) const;
+	void CollectExternalDependencies(TArray<FRigVMExternalDependency>& OutDependencies, const FName& InCategory, const URigVMGraph* InGraph) const;
+	void CollectExternalDependencies(TArray<FRigVMExternalDependency>& OutDependencies, const FName& InCategory, const URigVMNode* InNode) const;
+	void CollectExternalDependencies(TArray<FRigVMExternalDependency>& OutDependencies, const FName& InCategory, const URigVMPin* InPin) const;
+	void CollectExternalDependencies(TArray<FRigVMExternalDependency>& OutDependencies, const FName& InCategory, const UStruct* InStruct) const;
+	void CollectExternalDependencies(TArray<FRigVMExternalDependency>& OutDependencies, const FName& InCategory, const UEnum* InEnum) const;
+	void CollectExternalDependenciesForCPPTypeObject(TArray<FRigVMExternalDependency>& OutDependencies, const FName& InCategory, const UObject* InObject) const;
+	// end of Interface to URigVMExternalDependencyManager
+
+	void CommonInitialization(const FObjectInitializer& ObjectInitializer);
+	
+	void InitializeModelIfRequired(bool bRecompileVM = true);
+
+
+	void SerializeImpl(FArchive& Ar);
+
+
+	
+	virtual void PreSave(FObjectPreSaveContext ObjectSaveContext);
+	virtual void PostLoad();
+	static void DeclareConstructClasses(TArray<FTopLevelAssetPath>& OutConstructClasses, const UClass* SpecificSubclass);
+	virtual bool IsPostLoadThreadSafe() const { return false; }
+	virtual void PostTransacted(const FTransactionObjectEvent& TransactionEvent);
+	virtual void ReplaceDeprecatedNodes();
+	virtual void PreDuplicate(FObjectDuplicationParameters& DupParams);
+	virtual void PostDuplicate(bool bDuplicateForPIE);
+	virtual void GetAssetRegistryTags(FAssetRegistryTagsContext Context) const;
+	
+	
+	virtual bool RequiresForceLoadMembers(UObject* InObject) const;
+
+	// UObject interface
+	virtual void PostEditChangeChainProperty(struct FPropertyChangedChainEvent& PropertyChangedEvent);
+	virtual void PostRename(UObject* OldOuter, const FName OldName);
+	/** Called during cooking. Must return all objects that will be Preload()ed when this is serialized at load time. */
+	virtual void GetPreloadDependencies(TArray<UObject*>& OutDeps);
+
+	//  --- IRigVMClientHost interface Start---
+	virtual UObject* GetEditorObjectForRigVMGraph(const URigVMGraph* InVMGraph) const;
+	virtual URigVMGraph* GetRigVMGraphForEditorObject(UObject* InObject) const;
+	virtual void HandleRigVMGraphAdded(const FRigVMClient* InClient, const FString& InNodePath);
+	virtual void HandleRigVMGraphRemoved(const FRigVMClient* InClient, const FString& InNodePath);
+	virtual void HandleRigVMGraphRenamed(const FRigVMClient* InClient, const FString& InOldNodePath, const FString& InNewNodePath);
+	virtual void HandleConfigureRigVMController(const FRigVMClient* InClient, URigVMController* InControllerToConfigure);
+	virtual UObject* ResolveUserDefinedTypeById(const FString& InTypeName) const;
+
+
+	void GenerateUserDefinedDependenciesData(FRigVMExtendedExecuteContext& Context);
+	virtual void RecompileVM() ;
+	virtual void RecompileVMIfRequired() ;
+	virtual void RequestAutoVMRecompilation() ;
+	virtual void SetAutoVMRecompile(bool bAutoRecompile)  { bAutoRecompileVM = bAutoRecompile; }
+	virtual bool GetAutoVMRecompile() const  { return bAutoRecompileVM; }
+
+	
+
+	// this is needed since even after load
+	// model data can change while the RigVM BP is not opened
+	// for example, if a user defined struct changed after BP load,
+	// any pin that references the struct needs to be regenerated
+	virtual void RefreshAllModels(ERigVMLoadType InLoadType = ERigVMLoadType::PostLoad);
+
+	// RigVMRegistry changes can be triggered when new user defined types(structs/enums) are added/removed
+	// in which case we have to refresh the model
+	virtual void OnRigVMRegistryChanged();
+
+	void RequestRigVMInit() const;
+
+	
+
+	virtual URigVMGraph* GetDefaultModel() const;
+
+
+	
+
+	virtual URigVMFunctionLibrary* GetOrCreateLocalFunctionLibrary(bool bSetupUndoRedo = true);
+
+	virtual URigVMGraph* AddModel(FString InName = TEXT("Rig Graph"), bool bSetupUndoRedo = true, bool bPrintPythonCommand = true);
+
+	virtual bool RemoveModel(FString InName = TEXT("Rig Graph"), bool bSetupUndoRedo = true, bool bPrintPythonCommand = true);
+
+	virtual FRigVMGetFocusedGraph& OnGetFocusedGraph();
+	virtual const FRigVMGetFocusedGraph& OnGetFocusedGraph() const;
+
+
+	virtual TArray<FString> GeneratePythonCommands(const FString InNewBlueprintName);
+
+	virtual void SetupPinRedirectorsForBackwardsCompatibility() {};
+
+	
+
+	virtual bool IsFunctionPublic(const FName& InFunctionName) const;
+	virtual void MarkFunctionPublic(const FName& InFunctionName, bool bIsPublic = true);
+
+	virtual void RenameGraph(const FString& InNodePath, const FName& InNewName);
+
+	//  --- IRigVMClientHost interface End ---
+
+	//  --- IRigVMExternalDependencyManager interface Start ---
+
+	virtual TArray<FRigVMExternalDependency> GetExternalDependenciesForCategory(const FName& InCategory) const;
+	
+	//  --- IRigVMExternalDependencyManager interface End ---
+
+
+	virtual FOnRigVMRequestInspectObject& OnRequestInspectObject() { return OnRequestInspectObjectEvent; }
+	void RequestInspectObject(const TArray<UObject*>& InObjects) { OnRequestInspectObjectEvent.Broadcast(InObjects); }
+
+	
+
+	URigVMGraph* GetTemplateModel(bool bIsFunctionLibrary = false);
+	URigVMController* GetTemplateController(bool bIsFunctionLibrary = false);
+
+
+	
+
+	
+
+	TObjectPtr<URigVMGraph> TemplateModel;
+	TObjectPtr<URigVMController> TemplateController;
+	mutable TArray<UObject::FAssetRegistryTag> CachedAssetTags;
+
+	
+
+	bool bSuspendModelNotificationsForSelf;
+	bool bSuspendAllNotifications;
+
+	void RebuildGraphFromModel();
+
+	URigVMHost* CreateRigVMHostImpl();
+	
+
+	virtual TArray<UStruct*> GetAvailableRigVMStructs() const;
+
+	FRigVMVariantRef GetAssetVariantRefImpl() const;
+
+	/** Resets the asset's guid to a new one and splits it from the former variant set */
+	bool SplitAssetVariantImpl();
+
+	/** Merges the asset's guid with a provided one to join the variant set */
+	bool JoinAssetVariantImpl(const FGuid& InGuid);
+
+	TArray<FRigVMVariantRef> GetMatchingVariantsImpl() const;
+
+	
+
+	bool bAutoRecompileVM;
+
+	bool bVMRecompilationRequired;
+
+	bool bIsCompiling;
+
+	int32 VMRecompilationBracket;
+	
+	bool bSkipDirtyBlueprintStatus;
+
+	FRigVMGraphModifiedEvent ModifiedEvent;
+	void Notify(ERigVMGraphNotifType InNotifType, UObject* InSubject);
+	void HandleModifiedEvent(ERigVMGraphNotifType InNotifType, URigVMGraph* InGraph, UObject* InSubject);
+
+	FOnRigVMAssetChangedEvent ChangedEvent;
+	virtual FOnRigVMAssetChangedEvent& OnChanged()  { return ChangedEvent; }
+
+	FOnRigVMSetObjectBeingDebugged SetObjectBeingDebuggedEvent;
+	
+	
+	void ReplaceFunctionIdentifiers(const FString& OldAssetPath, const FString& NewAssetPath);
+
+	FOnRigVMRefreshEditorEvent RefreshEditorEvent;
+	FOnRigVMVariableDroppedEvent VariableDroppedEvent;
+	FOnRigVMBreakpointAdded BreakpointAddedEvent;
+
+	void SuspendNotifications(bool bSuspendNotifs);
+	
+	virtual FOnRigVMVariableDroppedEvent& OnVariableDropped()  { return VariableDroppedEvent; }
+	virtual FOnRigVMBreakpointAdded& OnBreakpointAdded()  { return BreakpointAddedEvent; }
+
+
+	FOnRigVMCompiledEvent VMCompiledEvent;
+
+	virtual void PathDomainSpecificContentOnLoad() {}
+	virtual void PatchBoundVariables();
+	virtual void PatchVariableNodesWithIncorrectType();
+	virtual void PatchParameterNodesOnLoad() {}
+	virtual void PatchLinksWithCast();
+	virtual void GetBackwardsCompatibilityPublicFunctions(TArray<FName> &BackwardsCompatiblePublicFunctions, TMap<URigVMLibraryNode*, FRigVMGraphFunctionHeader>& OldHeaders);
+
+	virtual void CreateMemberVariablesOnLoad();
+	virtual void PatchVariableNodesOnLoad();
+	TMap<FName, int32> AddedMemberVariableMap;
+
+	void PropagateRuntimeSettingsFromBPToInstances();
+	void InitializeArchetypeInstances();
+
+	void HandlePackageDone(const FEndLoadPackageContext& Context);
+
+	virtual void HandlePackageDone();
+
+	/** Our currently running rig vm instance */
+	virtual TObjectPtr<URigVMHost>& GetEditorHost() = 0;
+
+	// RigVMBP, once end-loaded, will inform other RigVM-Dependent systems that Host instances are ready.
+	void BroadcastRigVMPackageDone();
+
+	// Previously some memory classes were parented to the asset object
+	// however it is no longer supported since classes are now identified 
+	// with only package name + class name, see FTopLevelAssetPath
+	// this function removes those deprecated class.
+	// new classes should be created by RecompileVM and parented to the Package
+	// during PostLoad
+	void RemoveDeprecatedVMMemoryClass();
+
+	// During load, we do not want the GC to destroy the generator classes until all URigVMMemoryStorage objects
+	// are loaded, so we need to keep a pointer to the classes. These pointers will be removed on PreSave so that the
+	// GC can do its work.
+	TArray<TObjectPtr<URigVMMemoryStorageGeneratorClass>> OldMemoryStorageGeneratorClasses;
+
+	FOnRigVMExternalVariablesChanged& OnExternalVariablesChanged() { return ExternalVariablesChangedEvent; }
+
+	virtual void OnVariableAdded(const FName& InVarName);
+	virtual void OnVariableRemoved(const FName& InVarName);
+	virtual void OnVariableRenamed(const FName& InOldVarName, const FName& InNewVarName);
+	virtual void OnVariableTypeChanged(const FName& InVarName, FEdGraphPinType InOldPinType, FEdGraphPinType InNewPinType);
+
+	FName AddAssetVariableFromPinType(const FName& InName, const FEdGraphPinType& InType, bool bIsPublic = false, bool bIsReadOnly = false, FString InDefaultValue = FString());
+	
+	virtual FOnRigVMNodeDoubleClicked& OnNodeDoubleClicked()  { return NodeDoubleClickedEvent; }
+	void BroadcastNodeDoubleClicked(URigVMNode* InNode);
+
+	virtual FOnRigVMGraphImported& OnGraphImported()  { return GraphImportedEvent; }
+	void BroadcastGraphImported(UEdGraph* InGraph);
+
+	
+
+	
+	virtual FRigVMOnBulkEditDialogRequestedDelegate& OnRequestBulkEditDialog()  { return RequestBulkEditDialog; }
+
+	virtual FRigVMOnBreakLinksDialogRequestedDelegate& OnRequestBreakLinksDialog()  { return RequestBreakLinksDialog; }
+
+	virtual FRigVMOnPinTypeSelectionRequestedDelegate& OnRequestPinTypeSelectionDialog()  { return RequestPinTypeSelectionDialog; }
+
+	
+
+	
+	void BroadCastReportCompilerMessage(EMessageSeverity::Type InSeverity, UObject* InSubject, const FString& InMessage);
+
+	
+	FOnRigVMExternalVariablesChanged ExternalVariablesChangedEvent;
+	void BroadcastExternalVariablesChangedEvent();
+
+	FOnRigVMNodeDoubleClicked NodeDoubleClickedEvent;
+	FOnRigVMGraphImported GraphImportedEvent;
+	FOnRigVMPostEditChangeChainProperty PostEditChangeChainPropertyEvent;
+	FOnRigVMLocalizeFunctionDialogRequested RequestLocalizeFunctionDialog;
+	FOnRigVMReportCompilerMessage ReportCompilerMessageEvent;
+	FRigVMOnBulkEditDialogRequestedDelegate RequestBulkEditDialog;
+	FRigVMOnBreakLinksDialogRequestedDelegate RequestBreakLinksDialog;
+	FRigVMOnPinTypeSelectionRequestedDelegate RequestPinTypeSelectionDialog;
+	FRigVMController_RequestJumpToHyperlinkDelegate RequestJumpToHyperlink;
+
+	UEdGraph* CreateEdGraph(URigVMGraph* InModel, bool bForce = false);
+	bool RemoveEdGraph(URigVMGraph* InModel);
+	void DestroyObject(UObject* InObject);
+	void CreateEdGraphForCollapseNodeIfNeeded(URigVMCollapseNode* InNode, bool bForce = false);
+	bool RemoveEdGraphForCollapseNode(URigVMCollapseNode* InNode, bool bNotify = false);
+
+	FCompilerResultsLog CompileLog;
+	
+	TArray<IRigVMGraphFunctionHost*> GetReferencedFunctionHosts(bool bForceLoad) const;
+
+
+	bool bCompileInDebugMode;
+
+	virtual TArray<TObjectPtr<URigVMNode>>& GetRigVMBreakpointNodes() = 0;
+
+	FOnRigVMRequestInspectObject OnRequestInspectObjectEvent;
+	FOnRigVMRequestInspectMemoryStorage OnRequestInspectMemoryStorageEvent;
+
+	/** Sets the execution mode. In Release mode the rig will ignore all breakpoints. */
+	void SetDebugMode(const bool bValue);
+
+	/** Returns the execution mode */
+	bool IsInDebugMode() const { return bCompileInDebugMode; }
+
+	
+	/** Adds a breakpoint to the first instruction of each callpath related to the InBreakpointNode */
+	bool AddBreakpointToHost(URigVMNode* InBreakpointNode);
+
+	
+
+	/** Recomputes the instruction breakpoints given the node breakpoints in the blueprint */
+	void RefreshBreakpoints();
+
+	/** Shape libraries to load during package load completed */ 
+	TArray<FString> ShapeLibrariesToLoadOnPackageLoaded;
+
+	TArray<FRigVMReferenceNodeData> GetReferenceNodeData() const;
+	virtual void SetupDefaultObjectDuringCompilation(URigVMHost* InCDO);
+
+
+	static FSoftObjectPath PreDuplicateAssetPath;
+	static FSoftObjectPath PreDuplicateHostPath;
+	static TArray<TRigVMAssetInterfacePtr> sCurrentlyOpenedRigVMBlueprints;
+
+	void MarkDirtyDuringLoad();
+
+	bool IsMarkedDirtyDuringLoad() const;
+
+	bool bDirtyDuringLoad;
+	bool bErrorsDuringCompilation;
+	bool bSuspendPythonMessagesForRigVMClient;
+	bool bMarkBlueprintAsStructurallyModifiedPending;
+
+	
+
+	static FCriticalSection QueuedCompilerMessageDelegatesMutex;
+	static TArray<FOnRigVMReportCompilerMessage::FDelegate> QueuedCompilerMessageDelegates;
+
+
+
+#endif
+public:
+
+protected:
+	
+
+
+	friend class FRigVMBlueprintCompilerContext;
+	friend class FRigVMLegacyEditor;
+	friend class FRigVMNewEditor;
+	friend class FRigVMEditorBase;
+	friend class FRigVMEditorModule;
+	friend class URigVMEdGraphSchema;
+	friend struct FRigVMEdGraphSchemaAction_PromoteToVariable;
+	friend class URigVMBuildData;
+	friend class FRigVMVariantDetailCustomization;
+	friend class FRigVMTreeAssetVariantFilter;
+	friend class FRigVMTreePackageNode;
+	friend class SRigVMGraphNode;
+	friend struct FGuardSkipDirtyBlueprintStatus;
+	friend class SRigModuleAssetBrowser;
+	friend class UE::RigVM::Editor::Tools::FFilterByAssetTag;
+};
+
+class RIGVMDEVELOPER_API FRigVMBlueprintCompileScope
+{
+public:
+   
+	FRigVMBlueprintCompileScope(TRigVMAssetInterfacePtr InBlueprint);
+
+	~FRigVMBlueprintCompileScope();
+
+private:
+
+	TRigVMAssetInterfacePtr Blueprint;
+};
+
+struct FGuardSkipDirtyBlueprintStatus : private FNoncopyable
+{
+	[[nodiscard]] FGuardSkipDirtyBlueprintStatus(TRigVMAssetInterfacePtr InBlueprint, bool bNewValue)
+	{
+		if (InBlueprint)
+		{
+			WeakBlueprint = InBlueprint;
+			bOldValue = InBlueprint->bSkipDirtyBlueprintStatus;
+			InBlueprint->bSkipDirtyBlueprintStatus = bNewValue;
+		}
+	}
+	~FGuardSkipDirtyBlueprintStatus()
+	{
+		if (TStrongObjectPtr<UObject> Asset = WeakBlueprint.GetWeakObjectPtr().Pin())
+		{
+			if (TRigVMAssetInterfacePtr Interface = Asset.Get())
+			{
+				Interface->bSkipDirtyBlueprintStatus = bOldValue;
+			}
+		}
+	}
+
+private:
+	TWeakInterfacePtr<IRigVMAssetInterface> WeakBlueprint;
+	bool bOldValue = false;
+};
+
+#undef UE_API
\ No newline at end of file
