RigVM: Adding IRigVMAssetInterface as parent of RigVMBlueprint. First step to separation of RigVM and Blueprint.

diff --git a/Engine/Plugins/Runtime/RigVM/Source/RigVMDeveloper/Public/RigVMBlueprintLegacy.h b/Engine/Plugins/Runtime/RigVM/Source/RigVMDeveloper/Public/RigVMBlueprintLegacy.h
new file mode 100644
--- /dev/null
+++ b/Engine/Plugins/Runtime/RigVM/Source/RigVMDeveloper/Public/RigVMBlueprintLegacy.h
@@ -0,0 +1,506 @@
+// Copyright Epic Games, Inc. All Rights Reserved.
+
+#pragma once
+
+#include "CoreMinimal.h"
+#include "RigVMBlueprintGeneratedClass.h"
+#include "Engine/Blueprint.h"
+#include "RigVMCore/RigVM.h"
+#include "RigVMHost.h"
+#include "RigVMModel/RigVMClient.h"
+#include "RigVMCompiler/RigVMCompiler.h"
+#include "RigVMCore/RigVMGraphFunctionDefinition.h"
+#include "EdGraph/RigVMEdGraph.h"
+#include "EdGraph/RigVMEdGraphSchema.h"
+#if WITH_EDITOR
+#include "HAL/CriticalSection.h"
+#endif
+
+#include "RigVMAsset.h"
+#include "UObject/AssetRegistryTagsContext.h"
+#include "UObject/ObjectSaveContext.h"
+
+#include "RigVMBlueprintLegacy.generated.h"
+
+#define UE_API RIGVMDEVELOPER_API
+
+
+class URigVMBlueprintGeneratedClass;
+
+USTRUCT(meta = (Deprecated = "5.2"))
+struct RIGVMDEVELOPER_API FRigVMOldPublicFunctionArg
+{
+	GENERATED_BODY();
+	
+	FRigVMOldPublicFunctionArg()
+	: Name(NAME_None)
+	, CPPType(NAME_None)
+	, CPPTypeObjectPath(NAME_None)
+	, bIsArray(false)
+	, Direction(ERigVMPinDirection::Input)
+	{}
+
+	UPROPERTY()
+	FName Name;
+
+	UPROPERTY()
+	FName CPPType;
+
+	UPROPERTY()
+	FName CPPTypeObjectPath;
+
+	UPROPERTY()
+	bool bIsArray;
+
+	UPROPERTY()
+	ERigVMPinDirection Direction;
+
+	FEdGraphPinType GetPinType() const;
+};
+
+USTRUCT(meta = (Deprecated = "5.2"))
+struct RIGVMDEVELOPER_API FRigVMOldPublicFunctionData
+{
+	GENERATED_BODY();
+
+	FRigVMOldPublicFunctionData()
+		:Name(NAME_None)
+	{}
+
+	virtual ~FRigVMOldPublicFunctionData() {}
+
+	UPROPERTY()
+	FName Name;
+
+	UPROPERTY()
+	FString DisplayName;
+
+	UPROPERTY()
+	FString Category;
+
+	UPROPERTY()
+	FString Keywords;
+
+	UPROPERTY()
+	FRigVMOldPublicFunctionArg ReturnValue;
+
+	UPROPERTY()
+	TArray<FRigVMOldPublicFunctionArg> Arguments;
+
+	bool IsMutable() const;
+};
+
+UCLASS(BlueprintType, meta=(IgnoreClassThumbnail))
+class UE_API URigVMBlueprint : public UBlueprint, public IRigVMAssetInterface, public IRigVMClientHost, public IRigVMExternalDependencyManager
+{
+	GENERATED_UCLASS_BODY()
+
+public:
+	virtual void BeginDestroy() override;
+	virtual UObject* GetObject() override { return this; }
+	virtual TArray<FRigVMGraphFunctionHeader>& GetPublicGraphFunctions() override { return PublicGraphFunctions; }
+	virtual FRigVMClient* GetRigVMClient() override { return &RigVMClient; }
+	virtual const FRigVMClient* GetRigVMClient() const override { return &RigVMClient; }
+	virtual FRigVMVariant& GetAssetVariant() override { return AssetVariant; }
+	virtual const FRigVMVariant& GetAssetVariant() const override { return AssetVariant; }
+	virtual FRigVMCompileSettings& GetVMCompileSettings() override { return VMCompileSettings; }
+	virtual FRigVMRuntimeSettings& GetVMRuntimeSettings() override { return VMRuntimeSettings; }
+	virtual bool IsRegeneratingOnLoad() const override { return bIsRegeneratingOnLoad;}
+	virtual TArray<FRigVMGraphVariableDescription> GetAssetVariables() const override;
+	virtual TArray<struct FEditedDocumentInfo>& GetLastEditedDocuments() override { return LastEditedDocuments;}
+
+private:
+	virtual UClass* GetRigVMGeneratedClass() override { return GeneratedClass; }
+	virtual TArray<FRigVMExternalVariable> GetExternalVariables(bool bFallbackToBlueprint) const override;
+	virtual FString GetVariableDefaultValue(const FName& InName, bool bFromDebuggedObject) const override;
+	virtual URigVM* GetVM(bool bCreateIfNeeded = true) const override;
+	virtual FRigVMExtendedExecuteContext* GetRigVMExtendedExecuteContext() override;
+	virtual TScriptInterface<IRigVMGraphFunctionHost> GetRigVMGraphFunctionHost() override { return GetRigVMBlueprintGeneratedClass(); }
+	virtual TScriptInterface<const IRigVMGraphFunctionHost> GetRigVMGraphFunctionHost() const override { return GetRigVMBlueprintGeneratedClass(); }
+	virtual FCompilerResultsLog* GetCurrentMessageLog() const override { return CurrentMessageLog; }
+	virtual void SetAssetStatus(const ERigVMAssetStatus& InStatus) override;
+	virtual ERigVMAssetStatus GetAssetStatus() const override;
+	virtual bool IsUpToDate() const override { return UBlueprint::IsUpToDate(); }
+	virtual TArray<UObject*> GetArchetypeInstances(bool bIncludeCDO, bool bIncludeDerivedClasses) const override;
+	virtual const FString& GetObjectBeingDebuggedPath() const override { return UBlueprint::GetObjectPathToDebug(); }
+	virtual UWorld* GetWorldBeingDebugged() const override { return UBlueprint::GetWorldBeingDebugged(); }
+	virtual void SetWorldBeingDebugged(UWorld* NewWorld) override { return UBlueprint::SetWorldBeingDebugged(NewWorld); }
+	virtual FRigVMDebugInfo& GetDebugInfo() override;
+	virtual TObjectPtr<URigVMEdGraph>& GetFunctionLibraryEdGraph() override { return FunctionLibraryEdGraph; }
+	virtual URigVMHost* CreateRigVMHostSuper(UObject* InOuter) override;
+	virtual void MarkAssetAsStructurallyModified(bool bSkipDirtyAssetStatus = false) override;
+	virtual void MarkAssetAsModified(FPropertyChangedEvent PropertyChangedEvent = FPropertyChangedEvent(nullptr)) override;
+	virtual void AddUbergraphPage(URigVMEdGraph* RigVMEdGraph) override;
+	virtual void AddLastEditedDocument(URigVMEdGraph* RigVMEdGraph) override;
+	virtual void Compile() override;
+	virtual FCompilerResultsLog CompileBlueprint() override;
+	virtual void PatchVariablesOnLoadSuper() override;
+	virtual void AddPinWatch(UEdGraphPin* InPin) override;
+	virtual void RemovePinWatch(UEdGraphPin* InPin) override;
+	virtual void ClearPinWatches() override;
+	virtual bool IsPinBeingWatched(const UEdGraphPin* InPin) const override;
+	virtual void ForeachPinWatch(TFunctionRef<void(UEdGraphPin*)> Task) override;
+	virtual TMap<FString, FSoftObjectPath>& GetUserDefinedStructGuidToPathName(bool bFromCDO) override;
+	virtual TMap<FString, FSoftObjectPath>& GetUserDefinedEnumToPathName(bool bFromCDO) override;
+	virtual TSet<TObjectPtr<UObject>>& GetUserDefinedTypesInUse(bool bFromCDO) override;
+	virtual FRigVMEdGraphDisplaySettings& GetRigGraphDisplaySettings() override { return RigGraphDisplaySettings; }
+	virtual TArray<TObjectPtr<UEdGraph>>& GetFunctionGraphs() override { return FunctionGraphs; }
+	virtual bool& IsReferencedObjectPathsStored() override { return ReferencedObjectPathsStored; }
+	virtual TArray<FSoftObjectPath>& GetReferencedObjectPaths() override { return ReferencedObjectPaths; }
+	virtual TArray<FName>& GetSupportedEventNames() override { return SupportedEventNames; }
+	virtual TArray<FRigVMReferenceNodeData>& GetFunctionReferenceNodeData() override { return FunctionReferenceNodeData; }
+	virtual void NotifyGraphRenamedSuper(class UEdGraph* Graph, FName OldName, FName NewName) override { UBlueprint::NotifyGraphRenamed(Graph, OldName, NewName); }
+	virtual UObject* GetDefaultsObject() override;
+	virtual void PostEditChangeBlueprintActors() override;
+
+public:
+	URigVMBlueprint();
+
+	/** Get the (full) generated class for this rigvm blueprint */
+	URigVMBlueprintGeneratedClass* GetRigVMBlueprintGeneratedClass() const;
+	/** Returns the class used as the super class for all generated classes */
+	virtual UClass* GetRigVMGeneratedClassPrototype() const { return URigVMBlueprintGeneratedClass::StaticClass(); }
+
+	virtual void Serialize(FArchive& Ar) override { IRigVMAssetInterface ::SerializeImpl(Ar); }
+	virtual void SerializeSuper(FArchive& Ar) override { UBlueprint ::Serialize(Ar); }
+	virtual void PostSerialize(FArchive& Ar) override;
+
+#if WITH_EDITOR
+
+	// UBlueprint interface
+	virtual UClass* GetBlueprintClass() const override;
+	virtual UClass* RegenerateClass(UClass* ClassToRegenerate, UObject* PreviousCDO) override;
+	virtual bool SupportedByDefaultBlueprintFactory() const override { return false; }
+	virtual bool IsValidForBytecodeOnlyRecompile() const override { return false; }
+	virtual void LoadModulesRequiredForCompilation() override {}
+	virtual void GetTypeActions(FBlueprintActionDatabaseRegistrar& ActionRegistrar) const override;
+	virtual void GetInstanceActions(FBlueprintActionDatabaseRegistrar& ActionRegistrar) const override;
+	virtual void SetObjectBeingDebugged(UObject* NewObject) override { IRigVMAssetInterface::SetObjectBeingDebugged(NewObject); }
+	virtual void SetObjectBeingDebuggedSuper(UObject* NewObject) override { UBlueprint::SetObjectBeingDebugged(NewObject); }
+	virtual UObject* GetObjectBeingDebugged() override { return UBlueprint::GetObjectBeingDebugged(); }
+	virtual UObject* GetObjectBeingDebugged() const override { return UBlueprint::GetObjectBeingDebugged(); }
+	virtual void PreSave(FObjectPreSaveContext ObjectSaveContext) override { IRigVMAssetInterface::PreSave(ObjectSaveContext); }
+	virtual void PreSaveSuper(FObjectPreSaveContext ObjectSaveContext) override { UBlueprint::PreSave(ObjectSaveContext); }
+	virtual void PostLoad() override;
+	virtual void PostLoadSuper() override { UBlueprint::PostLoad(); };
+#if WITH_EDITORONLY_DATA
+	static void DeclareConstructClasses(TArray<FTopLevelAssetPath>& OutConstructClasses, const UClass* SpecificSubclass);
+#endif
+	virtual bool IsPostLoadThreadSafe() const override { return IRigVMAssetInterface::IsPostLoadThreadSafe(); }
+	virtual void PostTransacted(const FTransactionObjectEvent& TransactionEvent) override { IRigVMAssetInterface::PostTransacted(TransactionEvent); }
+	virtual void PostTransactedSuper(const FTransactionObjectEvent& TransactionEvent) override { UBlueprint::PostTransacted(TransactionEvent); };
+	virtual void ReplaceDeprecatedNodes() override { IRigVMAssetInterface::ReplaceDeprecatedNodes(); }
+	virtual void ReplaceDeprecatedNodesSuper() override { UBlueprint::ReplaceDeprecatedNodes(); };
+	virtual void PreDuplicate(FObjectDuplicationParameters& DupParams) override { IRigVMAssetInterface::PreDuplicate(DupParams); }
+	virtual void PreDuplicateSuper(FObjectDuplicationParameters& DupParams) override { UBlueprint::PreDuplicate(DupParams); }
+	virtual void PostDuplicate(bool bDuplicateForPIE) override { IRigVMAssetInterface::PostDuplicate(bDuplicateForPIE); }
+	virtual void PostDuplicateSuper(bool bDuplicateForPIE) override { UBlueprint::PostDuplicate(bDuplicateForPIE); };
+	virtual void GetAssetRegistryTags(FAssetRegistryTagsContext Context) const override { IRigVMAssetInterface::GetAssetRegistryTags(Context); }
+	virtual void GetAssetRegistryTagsSuper(FAssetRegistryTagsContext Context) const override { UBlueprint::GetAssetRegistryTags(Context); }
+	UE_DEPRECATED(5.4, "Implement the version that takes FAssetRegistryTagsContext instead.")
+	virtual void GetAssetRegistryTags(TArray<FAssetRegistryTag>& OutTags) const override;
+
+	virtual bool SupportsGlobalVariables() const override { return true; }
+	virtual bool SupportsLocalVariables() const override { return true; }
+	virtual bool SupportsFunctions() const override { return true; }
+	virtual bool SupportsMacros() const override { return false; }
+	virtual bool SupportsDelegates() const override { return false; }
+	virtual bool SupportsEventGraphs() const override { return true; }
+	virtual bool SupportsAnimLayers() const override { return false; }
+	virtual bool ExportGraphToText(UEdGraph* InEdGraph, FString& OutText) override { IRigVMAssetInterface::ExportGraphToText(InEdGraph, OutText); return true; }
+	virtual bool TryImportGraphFromText(const FString& InClipboardText, UEdGraph** OutGraphPtr = nullptr) override { return IRigVMAssetInterface::TryImportGraphFromText(InClipboardText, OutGraphPtr); }
+	virtual bool CanImportGraphFromText(const FString& InClipboardText) override { return IRigVMAssetInterface::CanImportGraphFromText(InClipboardText); }
+	virtual bool RequiresForceLoadMembers(UObject* InObject) const override { return IRigVMAssetInterface::RequiresForceLoadMembers(InObject); }
+	virtual bool RequiresForceLoadMembersSuper(UObject* InObject) const override { return UBlueprint::RequiresForceLoadMembers(InObject); }
+	
+
+	// UObject interface
+	virtual void PostEditChangeChainProperty(struct FPropertyChangedChainEvent& PropertyChangedEvent) override { IRigVMAssetInterface::PostEditChangeChainProperty(PropertyChangedEvent); }
+	virtual void PostEditChangeChainPropertySuper(struct FPropertyChangedChainEvent& PropertyChangedEvent) override { UBlueprint::PostEditChangeChainProperty(PropertyChangedEvent); }
+	virtual void PostRename(UObject* OldOuter, const FName OldName) override { IRigVMAssetInterface::PostRename(OldOuter, OldName); }
+	virtual void PostRenameSuper(UObject* OldOuter, const FName OldName) override { UBlueprint::PostRename(OldOuter, OldName); }
+	/** Called during cooking. Must return all objects that will be Preload()ed when this is serialized at load time. */
+	virtual void GetPreloadDependencies(TArray<UObject*>& OutDeps) override { IRigVMAssetInterface::GetPreloadDependencies(OutDeps);}
+	virtual void GetPreloadDependenciesSuper(TArray<UObject*>& OutDeps) override { UBlueprint::GetPreloadDependencies(OutDeps); }
+
+	virtual bool MarkPackageDirty() override { return UBlueprint::MarkPackageDirty(); }
+	virtual UPackage* GetPackage() override { return UBlueprint::GetPackage(); }
+	virtual TArray<TObjectPtr<UEdGraph>>& GetUberGraphs() override { return UbergraphPages; }
+	virtual void GetAllEdGraphs(TArray<UEdGraph*>& Graphs) const override { UBlueprint::GetAllGraphs(Graphs); }
+	virtual void SetUberGraphs(const TArray<TObjectPtr<UEdGraph>>& InGraphs) override { UbergraphPages = InGraphs; }
+	virtual void PostSaveRoot(FObjectPostSaveRootContext ObjectSaveContext) override;
+
+	//  --- IRigVMClientHost interface Start---
+	virtual FString GetAssetName() const override { return GetName(); }
+	virtual UClass* GetRigVMSchemaClass() const override { return URigVMSchema::StaticClass(); }
+	virtual UScriptStruct* GetRigVMExecuteContextStruct() const override { return FRigVMExecuteContext::StaticStruct(); }
+	virtual UClass* GetRigVMEdGraphClass() const override { return URigVMEdGraph::StaticClass(); }
+	virtual UClass* GetRigVMEdGraphNodeClass() const override { return URigVMEdGraphNode::StaticClass(); }
+	virtual UClass* GetRigVMEdGraphSchemaClass() const override { return URigVMEdGraphSchema::StaticClass(); }
+	virtual UClass* GetRigVMEditorSettingsClass() const override { return URigVMEditorSettings::StaticClass(); }
+	virtual UObject* GetEditorObjectForRigVMGraph(const URigVMGraph* InVMGraph) const override { return IRigVMAssetInterface::GetEditorObjectForRigVMGraph(InVMGraph); }
+	virtual URigVMGraph* GetRigVMGraphForEditorObject(UObject* InObject) const override { return IRigVMAssetInterface::GetRigVMGraphForEditorObject(InObject); }
+	virtual void HandleRigVMGraphAdded(const FRigVMClient* InClient, const FString& InNodePath) override { return IRigVMAssetInterface::HandleRigVMGraphAdded(InClient, InNodePath); }
+	virtual void HandleRigVMGraphRemoved(const FRigVMClient* InClient, const FString& InNodePath) override { return IRigVMAssetInterface::HandleRigVMGraphRemoved(InClient, InNodePath); }
+	virtual void HandleRigVMGraphRenamed(const FRigVMClient* InClient, const FString& InOldNodePath, const FString& InNewNodePath) override { return IRigVMAssetInterface::HandleRigVMGraphRenamed(InClient, InOldNodePath, InNewNodePath); }
+	virtual void HandleConfigureRigVMController(const FRigVMClient* InClient, URigVMController* InControllerToConfigure) override { return IRigVMAssetInterface::HandleConfigureRigVMController(InClient, InControllerToConfigure); }
+	virtual void IncrementVMRecompileBracket() override { return IRigVMAssetInterface::IncrementVMRecompileBracket(); }
+	virtual void DecrementVMRecompileBracket() override { return IRigVMAssetInterface::DecrementVMRecompileBracket(); }
+	virtual void RefreshAllModels(ERigVMLoadType InLoadType = ERigVMLoadType::PostLoad) override { return IRigVMAssetInterface::RefreshAllModels(InLoadType); }
+	virtual void OnRigVMRegistryChanged() override { return IRigVMAssetInterface::OnRigVMRegistryChanged(); }
+	virtual URigVMGraph* GetModel(const FString& InNodePath) const override { return IRigVMAssetInterface::GetModel(InNodePath); }
+	virtual FRigVMGetFocusedGraph& OnGetFocusedGraph() override { return IRigVMAssetInterface::OnGetFocusedGraph(); }
+	virtual const FRigVMGetFocusedGraph& OnGetFocusedGraph() const override { return IRigVMAssetInterface::OnGetFocusedGraph(); }
+	virtual void SetupPinRedirectorsForBackwardsCompatibility() override { return IRigVMAssetInterface::SetupPinRedirectorsForBackwardsCompatibility(); }
+	virtual FRigVMGraphModifiedEvent& OnModified() override { return IRigVMAssetInterface::OnModified(); }
+	virtual bool IsFunctionPublic(const FName& InFunctionName) const override { return IRigVMAssetInterface::IsFunctionPublic(InFunctionName); }
+	virtual void MarkFunctionPublic(const FName& InFunctionName, bool bIsPublic = true) override { return IRigVMAssetInterface::MarkFunctionPublic(InFunctionName, bIsPublic); }
+	virtual void RenameGraph(const FString& InNodePath, const FName& InNewName) override { return IRigVMAssetInterface::RenameGraph(InNodePath, InNewName); }
+	//  --- IRigVMClientHost interface End---
+	
+	UFUNCTION(BlueprintCallable, Category = "RigVM Blueprint")
+	virtual void RecompileVM() override { IRigVMAssetInterface::RecompileVM(); }
+	UFUNCTION(BlueprintCallable, Category = "RigVM Blueprint")
+	virtual void RecompileVMIfRequired() override { return IRigVMAssetInterface::RecompileVMIfRequired(); }
+	UFUNCTION(BlueprintCallable, Category = "RigVM Blueprint")
+	virtual void RequestAutoVMRecompilation() override { return IRigVMAssetInterface::RequestAutoVMRecompilation(); }
+	UFUNCTION(BlueprintCallable, Category = "RigVM Blueprint")
+	virtual void SetAutoVMRecompile(bool bAutoRecompile) override { return IRigVMAssetInterface::SetAutoVMRecompile(bAutoRecompile); }
+	UFUNCTION(BlueprintPure, Category = "RigVM Blueprint")
+	virtual bool GetAutoVMRecompile() const override { return IRigVMAssetInterface::GetAutoVMRecompile(); }
+
+	UFUNCTION(BlueprintCallable, Category = "RigVM Blueprint")
+	virtual void RequestRigVMInit() override { IRigVMAssetInterface::RequestRigVMInit(); }
+
+	UFUNCTION(BlueprintCallable, Category = "RigVM Blueprint")
+	virtual URigVMGraph* GetModel(const UEdGraph* InEdGraph = nullptr) const override { return IRigVMAssetInterface::GetModel(InEdGraph); }
+
+	UFUNCTION(BlueprintCallable, Category = "RigVM Blueprint")
+	virtual URigVMGraph* GetDefaultModel() const override { return IRigVMAssetInterface::GetDefaultModel(); }
+
+	UFUNCTION(BlueprintCallable, Category = "RigVM Blueprint")
+	virtual TArray<URigVMGraph*> GetAllModels() const override { return IRigVMAssetInterface::GetAllModels(); }
+
+	UFUNCTION(BlueprintCallable, Category = "RigVM Blueprint")
+	virtual URigVMFunctionLibrary* GetLocalFunctionLibrary() const override { return IRigVMAssetInterface::GetLocalFunctionLibrary(); }
+
+	UFUNCTION(BlueprintCallable, Category = "RigVM Blueprint")
+	virtual URigVMFunctionLibrary* GetOrCreateLocalFunctionLibrary(bool bSetupUndoRedo = true) override { return IRigVMAssetInterface::GetOrCreateLocalFunctionLibrary(bSetupUndoRedo); }
+
+	UFUNCTION(BlueprintCallable, Category = "RigVM Blueprint")
+	virtual URigVMGraph* AddModel(FString InName = TEXT("Rig Graph"), bool bSetupUndoRedo = true, bool bPrintPythonCommand = true) override { return IRigVMAssetInterface::AddModel(InName, bSetupUndoRedo, bPrintPythonCommand); }
+
+	UFUNCTION(BlueprintCallable, Category = "RigVM Blueprint")
+	virtual bool RemoveModel(FString InName = TEXT("Rig Graph"), bool bSetupUndoRedo = true, bool bPrintPythonCommand = true) override { return IRigVMAssetInterface::RemoveModel(InName, bSetupUndoRedo, bPrintPythonCommand); }
+
+
+	UFUNCTION(BlueprintCallable, Category = "RigVM Blueprint")
+	virtual URigVMGraph* GetFocusedModel() const override { return IRigVMAssetInterface::GetFocusedModel(); }
+
+	UFUNCTION(BlueprintCallable, Category = "RigVM Blueprint")
+	virtual URigVMController* GetController(const URigVMGraph* InGraph = nullptr) const override { return IRigVMAssetInterface::GetController(InGraph); }
+	virtual URigVMController* GetController(const UEdGraph* InGraph) const override { return IRigVMAssetInterface::GetController(InGraph); }
+
+	UFUNCTION(BlueprintCallable, Category = "RigVM Blueprint")
+	virtual URigVMController* GetControllerByName(const FString InGraphName = TEXT("")) const override { return IRigVMAssetInterface::GetControllerByName(InGraphName); }
+
+	UFUNCTION(BlueprintCallable, Category = "RigVM Blueprint")
+	virtual URigVMController* GetOrCreateController(URigVMGraph* InGraph = nullptr) override { return IRigVMAssetInterface::GetOrCreateController(InGraph); }
+	virtual URigVMController* GetOrCreateController(const UEdGraph* InGraph) override { return IRigVMAssetInterface::GetOrCreateController(InGraph); }
+
+	UFUNCTION(BlueprintCallable, Category = "RigVM Blueprint")
+	virtual TArray<FString> GeneratePythonCommands(const FString InNewBlueprintName) override { return IRigVMAssetInterface::GeneratePythonCommands(InNewBlueprintName); }
+#endif
+
+	virtual TArray<FRigVMExternalDependency> GetExternalDependenciesForCategory(const FName& InCategory) const override { return IRigVMAssetInterface::GetExternalDependenciesForCategory(InCategory); }
+
+
+	virtual bool ShouldBeMarkedDirtyUponTransaction() const override { return false; }
+
+	
+#if WITH_EDITORONLY_DATA
+	UPROPERTY()
+	TObjectPtr<URigVMEdGraph> FunctionLibraryEdGraph;
+#endif
+
+
+	UPROPERTY(EditAnywhere, Category = "User Interface")
+	FRigVMEdGraphDisplaySettings RigGraphDisplaySettings;
+
+	UPROPERTY(EditAnywhere, Category = "VM")
+	FRigVMRuntimeSettings VMRuntimeSettings;
+
+	UPROPERTY(EditAnywhere, BlueprintReadWrite, Category = "VM")
+	FRigVMCompileSettings VMCompileSettings;
+
+	UPROPERTY(EditAnywhere, Category = "Python Log Settings")
+	FRigVMPythonSettings PythonLogSettings;
+
+	UPROPERTY()
+	TMap<FString, FSoftObjectPath> UserDefinedStructGuidToPathName;
+
+	UPROPERTY()
+	TMap<FString, FSoftObjectPath> UserDefinedEnumToPathName;
+
+	UPROPERTY(transient)
+	TSet<TObjectPtr<UObject>> UserDefinedTypesInUse;
+
+protected:
+
+#if WITH_EDITORONLY_DATA
+	UPROPERTY()
+	TObjectPtr<URigVMGraph> Model_DEPRECATED;
+
+	UPROPERTY()
+	TObjectPtr<URigVMFunctionLibrary> FunctionLibrary_DEPRECATED;
+#endif
+
+	UPROPERTY()
+	FRigVMClient RigVMClient;
+
+#if WITH_EDITORONLY_DATA
+
+	UPROPERTY()
+	bool ReferencedObjectPathsStored;
+
+	UPROPERTY()
+	TArray<FSoftObjectPath> ReferencedObjectPaths;
+
+#endif
+
+	/** Asset searchable information about exposed public functions on this rig */
+	UPROPERTY(AssetRegistrySearchable)
+	TArray<FRigVMGraphFunctionHeader> PublicGraphFunctions;
+
+	/** Asset searchable information function references in this rig */
+	UPROPERTY(AssetRegistrySearchable)
+	TArray<FRigVMReferenceNodeData> FunctionReferenceNodeData;
+
+#if WITH_EDITORONLY_DATA
+
+	/** Variant information about this asset */
+	UPROPERTY(EditAnywhere, AssetRegistrySearchable, Category = "Variant")
+	FRigVMVariant AssetVariant;
+
+#endif
+
+public:
+
+	UFUNCTION(BlueprintCallable, Category = "VM")
+	virtual UClass* GetRigVMHostClass() const { return GeneratedClass; }
+
+	UFUNCTION(BlueprintCallable, Category = "VM")
+	virtual URigVMHost* CreateRigVMHost() { return IRigVMAssetInterface::CreateRigVMHostImpl(); }
+
+	UFUNCTION(BlueprintCallable, Category = "VM")
+	virtual URigVMHost* GetDebuggedRigVMHost() { return IRigVMAssetInterface::GetDebuggedRigVMHost(); }
+
+	UFUNCTION(BlueprintCallable, Category = "VM")
+	virtual TArray<UStruct*> GetAvailableRigVMStructs() const override { return IRigVMAssetInterface::GetAvailableRigVMStructs(); }
+
+#if WITH_EDITOR
+	UFUNCTION(BlueprintCallable, Category = "Variables")
+	virtual TArray<FRigVMGraphVariableDescription> GetMemberVariables() const { return GetAssetVariables(); }
+	
+	UFUNCTION(BlueprintCallable, Category = "Variables")
+	virtual FName AddMemberVariable(const FName& InName, const FString& InCPPType, bool bIsPublic = false, bool bIsReadOnly = false, FString InDefaultValue = TEXT("")) override;
+
+	UFUNCTION(BlueprintCallable, Category = "Variables")
+	virtual bool RemoveMemberVariable(const FName& InName) override;
+
+	UFUNCTION(BlueprintCallable, Category = "Variables")
+	virtual bool RenameMemberVariable(const FName& InOldName, const FName& InNewName) override;
+
+	UFUNCTION(BlueprintCallable, Category = "Variables")
+	virtual bool ChangeMemberVariableType(const FName& InName, const FString& InCPPType, bool bIsPublic = false, bool bIsReadOnly = false, FString InDefaultValue = TEXT("")) override;
+	virtual bool ChangeMemberVariableType(const FName& InName, const FEdGraphPinType& InType) override;
+
+	UFUNCTION(BlueprintPure, Category = "Variants", meta = (DisplayName = "GetAssetVariant", ScriptName = "GetAssetVariant"))
+	virtual FRigVMVariant GetAssetVariantBP() const { return AssetVariant; }
+
+	UFUNCTION(BlueprintPure, Category = "Variants")
+	virtual FRigVMVariantRef GetAssetVariantRef() const override { return IRigVMAssetInterface::GetAssetVariantRefImpl(); }
+
+	/** Resets the asset's guid to a new one and splits it from the former variant set */
+	UFUNCTION(BlueprintCallable, Category = "Variants")
+	virtual bool SplitAssetVariant() override { return IRigVMAssetInterface::SplitAssetVariantImpl(); }
+
+	/** Merges the asset's guid with a provided one to join the variant set */
+	UFUNCTION(BlueprintCallable, Category = "Variants")
+	virtual bool JoinAssetVariant(const FGuid& InGuid) override { return IRigVMAssetInterface::JoinAssetVariantImpl(InGuid); }
+
+	UFUNCTION(BlueprintPure, Category = "Variants")
+	virtual TArray<FRigVMVariantRef> GetMatchingVariants() const override { return IRigVMAssetInterface::GetMatchingVariantsImpl(); }
+#endif	// #if WITH_EDITOR
+
+	
+private:
+
+	/** The event names this rigvm blueprint contains */
+	UPROPERTY(AssetRegistrySearchable)
+	TArray<FName> SupportedEventNames;
+
+#if WITH_EDITOR
+	UFUNCTION(BlueprintCallable, Category = "RigVM Blueprint")
+	virtual void SuspendNotifications(bool bSuspendNotifs) { return IRigVMAssetInterface::SuspendNotifications(bSuspendNotifs); }
+#endif
+
+protected:
+
+#if WITH_EDITOR
+	/** Our currently running rig vm instance */
+	// Declaring these transient properties here instead of the IRigVMAssetInterface class in order to have them as UPROPERTY, to avoid having stale pointers when GC happens
+	UPROPERTY(transient)
+	TObjectPtr<URigVMHost> EditorHost = nullptr;
+	virtual TObjectPtr<URigVMHost>& GetEditorHost() override { return EditorHost; }
+
+	// Declaring these transient properties here instead of the IRigVMAssetInterface class in order to have them as UPROPERTY, to avoid having stale pointers when GC happens
+	UPROPERTY(transient)
+	TArray<TObjectPtr<URigVMNode>> RigVMBreakpointNodes;
+	virtual TArray<TObjectPtr<URigVMNode>>& GetRigVMBreakpointNodes() override { return RigVMBreakpointNodes; }
+	
+	static FName FindHostMemberVariableUniqueName(TSharedPtr<FKismetNameValidator> InNameValidator, const FString& InBaseName);
+	static int32 AddHostMemberVariable(URigVMBlueprint* InBlueprint, const FName& InVarName, FEdGraphPinType InVarType, bool bIsPublic, bool bIsReadOnly, FString InDefaultValue);
+	virtual FName AddHostMemberVariableFromExternal(FRigVMExternalVariable InVariableToCreate, FString InDefaultValue = FString()) override;
+public:
+
+	virtual void OnPreVariableChange(UObject* InObject);
+	virtual void OnPostVariableChange(UBlueprint* InBlueprint);
+	bool bUpdatingExternalVariables;
+
+	void OnBlueprintChanged(UBlueprint* InBlueprint);
+	void OnSetObjectBeingDebuggedReceived(UObject* InObject);
+#endif
+
+protected:
+
+	virtual void SetupDefaultObjectDuringCompilation(URigVMHost* InCDO) override { IRigVMAssetInterface::SetupDefaultObjectDuringCompilation(InCDO); }
+	virtual void PreCompile() override;
+
+	virtual FProperty* FindGeneratedPropertyByName(const FName& InName) const override;
+	virtual bool SetVariableTooltip(const FName& InName, const FText& InTooltip) override;
+	virtual FText GetVariableTooltip(const FName& InName) const override;
+	virtual bool SetVariableCategory(const FName& InName, const FString& InCategory) override;
+	virtual FString GetVariableCategory(const FName& InName) override;
+	virtual FString GetVariableMetadataValue(const FName& InName, const FName& InKey) override;
+	virtual bool SetVariableExposeOnSpawn(const FName& InName, const bool bInExposeOnSpawn) override;
+	virtual bool SetVariableExposeToCinematics(const FName& InName, const bool bInExposeToCinematics) override;
+	virtual bool SetVariablePrivate(const FName& InName, const bool bInPrivate) override;
+	virtual bool SetVariablePublic(const FName& InName, const bool bIsPublic) override;
+	//virtual bool ChangeAssetVariableType(const FName& InName, const FEdGraphPinType& InType) override;
+	virtual FString OnCopyVariable(const FName& InName) const override;
+	virtual bool OnPasteVariable(const FString& InText) override;
+	
+#if WITH_EDITOR
+
+private:
+
+	TArray<FBPVariableDescription> LastNewVariables;
+#endif
+
+	// friend class FRigVMTreePackageNode;
+	// friend class UE::RigVM::Editor::Tools::FFilterByAssetTag;
+	// friend class FRigVMEditorModule;
+};
+
+#undef UE_API
diff --git a/Engine/Plugins/Runtime/RigVM/Source/RigVMDeveloper/Public/RigVMBlueprintUtils.h b/Engine/Plugins/Runtime/RigVM/Source/RigVMDeveloper/Public/RigVMBlueprintUtils.h
--- a/Engine/Plugins/Runtime/RigVM/Source/RigVMDeveloper/Public/RigVMBlueprintUtils.h
+++ b/Engine/Plugins/Runtime/RigVM/Source/RigVMDeveloper/Public/RigVMBlueprintUtils.h
@@ -3,11 +3,13 @@
 #pragma once
 
 #include "CoreMinimal.h"
+#include "RigVMAsset.h"
 #include "Templates/SubclassOf.h"
 #include "EdGraph/EdGraphPin.h"
 
 #define UE_API RIGVMDEVELOPER_API
 
+class IRigVMAssetInterface;
 class UStruct;
 class UBlueprint;
 class URigVMEdGraphNode;
@@ -20,18 +22,35 @@ struct FRigVMBlueprintUtils
 static UE_API void ForAllRigVMStructs(TFunction<void(UScriptStruct*)> InFunction);
 
 /** Handle blueprint node reconstruction */
-static UE_API void HandleReconstructAllNodes(UBlueprint* InBlueprint);
+static UE_API void HandleReconstructAllBlueprintNodes(UBlueprint* InBlueprint);
+static UE_API void HandleReconstructAllBlueprintNodes(TRigVMAssetInterfacePtr InBlueprint);
 
 /** Handle blueprint node refresh */
-static UE_API void HandleRefreshAllNodes(UBlueprint* InBlueprint);
+static UE_API void HandleRefreshAllNodes(TRigVMAssetInterfacePtr InBlueprint);
+static UE_API void HandleRefreshAllBlueprintNodes(UBlueprint* InBlueprint);
 
 /** Handle blueprint deleted */
 static UE_API void HandleAssetDeleted(const FAssetData& InAssetData);
 
 /** remove the variable if not used by anybody else but ToBeDeleted*/
-static UE_API void RemoveMemberVariableIfNotUsed(UBlueprint* Blueprint, const FName VarName);
+static UE_API void RemoveMemberVariableIfNotUsed(TRigVMAssetInterfacePtr Blueprint, const FName VarName);
 
-static UE_API FName ValidateName(UBlueprint* InBlueprint, const FString& InName);
+/** Create a new EdGraph */
+static UE_API URigVMEdGraph* CreateNewGraph(UObject* ParentScope, const FName& GraphName, TSubclassOf<class URigVMEdGraph> GraphClass, TSubclassOf<class URigVMEdGraphSchema> SchemaClass);
+
+static UE_API FName ValidateName(TRigVMAssetInterfacePtr InBlueprint, const FString& InName);
+
+/** Searches all nodes in a Blueprint and checks for a matching Guid */
+static UE_API UEdGraphNode* GetNodeByGUID(const TRigVMAssetInterfacePtr InBlueprint, const FGuid& InNodeGuid);
+
+// Helper function to get the asset that ultimately owns a graph.
+static UE_API TRigVMAssetInterfacePtr FindAssetForGraph(const UEdGraph* Graph);
+
+// Helper function to get the asset that ultimately owns a graph.
+static UE_API TRigVMAssetInterfacePtr FindAssetForNode(const UEdGraphNode* Node);
+
+// Helper function to find a unique name for an asset variable
+static UE_API FName FindUniqueVariableName(const IRigVMAssetInterface* InBlueprint, const FString& InBaseName);
 };
 
 #undef UE_API
diff --git a/Engine/Plugins/Runtime/RigVM/Source/RigVMDeveloper/Public/RigVMModel/Nodes/RigVMFunctionReferenceNode.h b/Engine/Plugins/Runtime/RigVM/Source/RigVMDeveloper/Public/RigVMModel/Nodes/RigVMFunctionReferenceNode.h
--- a/Engine/Plugins/Runtime/RigVM/Source/RigVMDeveloper/Public/RigVMModel/Nodes/RigVMFunctionReferenceNode.h
+++ b/Engine/Plugins/Runtime/RigVM/Source/RigVMDeveloper/Public/RigVMModel/Nodes/RigVMFunctionReferenceNode.h
@@ -92,6 +92,7 @@ private:
 	friend class URigVMController;
 	friend class FRigVMParserAST;
 	friend class URigVMBlueprint;
+	friend class IRigVMAssetInterface;
 	friend struct FRigVMClient;
 	friend struct EngineTestRigVMFramework;
 };
diff --git a/Engine/Plugins/Runtime/RigVM/Source/RigVMDeveloper/Public/RigVMModel/RigVMClient.h b/Engine/Plugins/Runtime/RigVM/Source/RigVMDeveloper/Public/RigVMModel/RigVMClient.h
--- a/Engine/Plugins/Runtime/RigVM/Source/RigVMDeveloper/Public/RigVMModel/RigVMClient.h
+++ b/Engine/Plugins/Runtime/RigVM/Source/RigVMDeveloper/Public/RigVMModel/RigVMClient.h
@@ -62,10 +62,10 @@ public:
 	virtual const FRigVMClient* GetRigVMClient() const = 0;
 
 	// Returns the rigvm function host
-	virtual IRigVMGraphFunctionHost* GetRigVMGraphFunctionHost() = 0;
+	virtual TScriptInterface<IRigVMGraphFunctionHost> GetRigVMGraphFunctionHost() = 0;
 
 	// Returns the rigvm function host
-	virtual const IRigVMGraphFunctionHost* GetRigVMGraphFunctionHost() const = 0;
+	virtual TScriptInterface<const IRigVMGraphFunctionHost> GetRigVMGraphFunctionHost() const = 0;
 
 	// Returns the editor object corresponding with the supplied RigVM graph
 	virtual UObject* GetEditorObjectForRigVMGraph(const URigVMGraph* InVMGraph) const = 0;
@@ -395,6 +395,7 @@ private:
 
 	friend class UEngineTestClientHost;
 	friend class URigVMBlueprint;
+	friend class IRigVMAssetInterface;
 };
 
 #undef UE_API
diff --git a/Engine/Plugins/Runtime/RigVM/Source/RigVMDeveloper/Public/RigVMModel/RigVMController.h b/Engine/Plugins/Runtime/RigVM/Source/RigVMDeveloper/Public/RigVMModel/RigVMController.h
--- a/Engine/Plugins/Runtime/RigVM/Source/RigVMDeveloper/Public/RigVMModel/RigVMController.h
+++ b/Engine/Plugins/Runtime/RigVM/Source/RigVMDeveloper/Public/RigVMModel/RigVMController.h
@@ -1695,7 +1695,7 @@ private:
 	friend class URigVMActionStack;
 	friend struct FRigVMBaseAction;
 	friend class URigVMCompiler;
-	friend class URigVMBlueprint;
+	friend class IRigVMAssetInterface;
 	friend struct FRigVMControllerObjectFactory;
 	friend struct FRigVMSetPinDefaultValueAction;
 	friend struct FRigVMAddRerouteNodeAction;
diff --git a/Engine/Plugins/Runtime/RigVM/Source/RigVMDeveloper/Public/RigVMModel/RigVMExternalDependency.h b/Engine/Plugins/Runtime/RigVM/Source/RigVMDeveloper/Public/RigVMModel/RigVMExternalDependency.h
--- a/Engine/Plugins/Runtime/RigVM/Source/RigVMDeveloper/Public/RigVMModel/RigVMExternalDependency.h
+++ b/Engine/Plugins/Runtime/RigVM/Source/RigVMDeveloper/Public/RigVMModel/RigVMExternalDependency.h
@@ -89,6 +89,8 @@ protected:
 	UE_API virtual void CollectExternalDependenciesForCPPTypeObject(TArray<FRigVMExternalDependency>& OutDependencies, const FName& InCategory, const UObject* InObject) const;
 
 	static UE_API TArray<FName> DependencyCategories;
+
+	friend class IRigVMAssetInterface;
 };
 
 #undef UE_API
diff --git a/Engine/Plugins/Runtime/RigVM/Source/RigVMDeveloper/Public/RigVMModel/RigVMFunctionLibrary.h b/Engine/Plugins/Runtime/RigVM/Source/RigVMDeveloper/Public/RigVMModel/RigVMFunctionLibrary.h
--- a/Engine/Plugins/Runtime/RigVM/Source/RigVMDeveloper/Public/RigVMModel/RigVMFunctionLibrary.h
+++ b/Engine/Plugins/Runtime/RigVM/Source/RigVMDeveloper/Public/RigVMModel/RigVMFunctionLibrary.h
@@ -102,6 +102,7 @@ private:
 	friend class URigVMController;
 	friend class URigVMCompiler;
 	friend class URigVMBlueprint;
+	friend class IRigVMAssetInterface;
 	friend struct FRigVMClient;
 	friend class UControlRigBlueprint;
 };
diff --git a/Engine/Plugins/Runtime/RigVM/Source/RigVMDeveloper/Public/RigVMModel/RigVMVariableDescription.h b/Engine/Plugins/Runtime/RigVM/Source/RigVMDeveloper/Public/RigVMModel/RigVMVariableDescription.h
--- a/Engine/Plugins/Runtime/RigVM/Source/RigVMDeveloper/Public/RigVMModel/RigVMVariableDescription.h
+++ b/Engine/Plugins/Runtime/RigVM/Source/RigVMDeveloper/Public/RigVMModel/RigVMVariableDescription.h
@@ -49,6 +49,26 @@ public:
 	UPROPERTY(EditAnywhere, BlueprintReadOnly, Category = RigVMGraphVariableDescription)
 	FText Category;
 
+	// The tooltip of the variable
+	UPROPERTY(EditAnywhere, BlueprintReadOnly, Category = RigVMGraphVariableDescription)
+	FText Tooltip;
+
+	// Should this variable be exposed on spawn
+	UPROPERTY(EditAnywhere, BlueprintReadOnly, Category = RigVMGraphVariableDescription)
+	bool bExposedOnSpawn = false;
+
+	// Should this variable be exposed on spawn
+	UPROPERTY(EditAnywhere, BlueprintReadOnly, Category = RigVMGraphVariableDescription)
+	bool bExposeToCinematics = false;
+
+	// Is this variable public
+	UPROPERTY(EditAnywhere, BlueprintReadOnly, Category = RigVMGraphVariableDescription)
+	bool bPublic = false;
+
+	// Is this variable private
+	UPROPERTY(EditAnywhere, BlueprintReadOnly, Category = RigVMGraphVariableDescription)
+	bool bPrivate = true;
+
 	// Returns nullptr external variable matching this description
 	FRigVMExternalVariable ToExternalVariable() const
 	{
diff --git a/Engine/Plugins/Runtime/RigVM/Source/RigVMEditor/Private/EdGraph/NodeSpawners/RigVMEdGraphEnumNodeSpawner.cpp b/Engine/Plugins/Runtime/RigVM/Source/RigVMEditor/Private/EdGraph/NodeSpawners/RigVMEdGraphEnumNodeSpawner.cpp
--- a/Engine/Plugins/Runtime/RigVM/Source/RigVMEditor/Private/EdGraph/NodeSpawners/RigVMEdGraphEnumNodeSpawner.cpp
+++ b/Engine/Plugins/Runtime/RigVM/Source/RigVMEditor/Private/EdGraph/NodeSpawners/RigVMEdGraphEnumNodeSpawner.cpp
@@ -7,12 +7,11 @@
 #include "EdGraph/RigVMEdGraphSchema.h"
 #include "BlueprintNodeTemplateCache.h"
 #include "RigVMModel/RigVMController.h"
-#include "RigVMBlueprint.h"
+#include "RigVMBlueprintLegacy.h"
+#include "RigVMBlueprintUtils.h"
 
 #include "RigVMModel/Nodes/RigVMEnumNode.h"
 
-#include UE_INLINE_GENERATED_CPP_BY_NAME(RigVMEdGraphEnumNodeSpawner)
-
 #if WITH_EDITOR
 #include "Editor.h"
 #include "Kismet2/BlueprintEditorUtils.h"
@@ -20,38 +19,24 @@
 
 #define LOCTEXT_NAMESPACE "RigVMEdGraphEnumNodeSpawner"
 
-URigVMEdGraphEnumNodeSpawner* URigVMEdGraphEnumNodeSpawner::CreateForEnum(UEnum* InEnum, const FText& InMenuDesc, const FText& InCategory, const FText& InTooltip)
+URigVMEdGraphEnumNodeSpawner::URigVMEdGraphEnumNodeSpawner(UEnum* InEnum, const FText& InMenuDesc, const FText& InCategory, const FText& InTooltip)
 {
-	URigVMEdGraphEnumNodeSpawner* NodeSpawner = NewObject<URigVMEdGraphEnumNodeSpawner>(GetTransientPackage());
-	NodeSpawner->NodeClass = URigVMEdGraphNode::StaticClass();
-	NodeSpawner->Enum = InEnum;
-
-	FBlueprintActionUiSpec& MenuSignature = NodeSpawner->DefaultMenuSignature;
-	
-	MenuSignature.MenuName = InMenuDesc;
-	MenuSignature.Tooltip  = InTooltip;
-	MenuSignature.Category = InCategory;
-	MenuSignature.Keywords = FText::FromString(TEXT("Enum"));
-	MenuSignature.Icon = FSlateIcon(TEXT("RigVMEditorStyle"), TEXT("RigVM.Unit"));
-
-	return NodeSpawner;
+	NodeClass = URigVMEdGraphNode::StaticClass();
+	Enum = InEnum;
+
+	MenuName = InMenuDesc;
+	MenuTooltip  = InTooltip;
+	MenuCategory = InCategory;
+	MenuKeywords = FText::FromString(TEXT("Enum"));
+	MenuIcon = FSlateIcon(TEXT("RigVMEditorStyle"), TEXT("RigVM.Unit"));
 }
 
-FBlueprintNodeSignature URigVMEdGraphEnumNodeSpawner::GetSpawnerSignature() const
+FString URigVMEdGraphEnumNodeSpawner::GetSpawnerSignature() const
 {
-	return FBlueprintNodeSignature(FString("RigUnit=" + Enum->GetFName().ToString()));
-}
-
-FBlueprintActionUiSpec URigVMEdGraphEnumNodeSpawner::GetUiSpec(FBlueprintActionContext const& Context, FBindingSet const& Bindings) const
-{
-	UEdGraph* TargetGraph = (Context.Graphs.Num() > 0) ? Context.Graphs[0] : nullptr;
-	FBlueprintActionUiSpec MenuSignature = PrimeDefaultUiSpec(TargetGraph);
-
-	DynamicUiSignatureGetter.ExecuteIfBound(Context, Bindings, &MenuSignature);
-	return MenuSignature;
+	return FString("RigUnit=" + Enum->GetFName().ToString());
 }
 
-UEdGraphNode* URigVMEdGraphEnumNodeSpawner::Invoke(UEdGraph* ParentGraph, FBindingSet const& Bindings, FVector2D const Location) const
+URigVMEdGraphNode* URigVMEdGraphEnumNodeSpawner::Invoke(URigVMEdGraph* ParentGraph, FVector2D const Location) const
 {
 	URigVMEdGraphNode* NewNode = nullptr;
 
@@ -66,7 +51,7 @@ UEdGraphNode* URigVMEdGraphEnumNodeSpawner::Invoke(UEdGraph* ParentGraph, FBindi
 	// First create a backing member for our node
 	URigVMEdGraph* RigGraph = Cast<URigVMEdGraph>(ParentGraph);
 	if(RigGraph == nullptr) return nullptr;
-	URigVMBlueprint* RigBlueprint = Cast<URigVMBlueprint>(FBlueprintEditorUtils::FindBlueprintForGraph(ParentGraph));
+	TRigVMAssetInterfacePtr RigBlueprint = FRigVMBlueprintUtils::FindAssetForGraph(ParentGraph);
 	check(RigBlueprint);
 
 #if WITH_EDITOR
diff --git a/Engine/Plugins/Runtime/RigVM/Source/RigVMEditor/Private/EdGraph/NodeSpawners/RigVMEdGraphFunctionRefNodeSpawner.cpp b/Engine/Plugins/Runtime/RigVM/Source/RigVMEditor/Private/EdGraph/NodeSpawners/RigVMEdGraphFunctionRefNodeSpawner.cpp
--- a/Engine/Plugins/Runtime/RigVM/Source/RigVMEditor/Private/EdGraph/NodeSpawners/RigVMEdGraphFunctionRefNodeSpawner.cpp
+++ b/Engine/Plugins/Runtime/RigVM/Source/RigVMEditor/Private/EdGraph/NodeSpawners/RigVMEdGraphFunctionRefNodeSpawner.cpp
@@ -17,8 +17,6 @@
 #include "ScopedTransaction.h"
 #include "Editor/RigVMEditorTools.h"
 
-#include UE_INLINE_GENERATED_CPP_BY_NAME(RigVMEdGraphFunctionRefNodeSpawner)
-
 #if WITH_EDITOR
 #include "Editor.h"
 #include "SGraphActionMenu.h"
@@ -32,118 +30,108 @@
 // - CreateFromAssetData(const FAssetData& InAssetData, const FRigVMGraphFunctionHeader& InPublicFunction) --> Public function. Valid Header.
 // - CreateFromAssetData(const FAssetData& InAssetData, const FRigVMOldPublicFunctionData& InPublicFunction) --> Public function. Valid AssetData and Header.Name
 
-URigVMEdGraphFunctionRefNodeSpawner* URigVMEdGraphFunctionRefNodeSpawner::CreateFromFunction(URigVMLibraryNode* InFunction)
+URigVMEdGraphFunctionRefNodeSpawner::URigVMEdGraphFunctionRefNodeSpawner(URigVMLibraryNode* InFunction)
 {
 	check(InFunction);
 
-	URigVMEdGraphFunctionRefNodeSpawner* NodeSpawner = NewObject<URigVMEdGraphFunctionRefNodeSpawner>(GetTransientPackage());
-	NodeSpawner->ReferencedPublicFunctionHeader = InFunction->GetFunctionHeader();
-	NodeSpawner->NodeClass = URigVMEdGraphNode::StaticClass();
-	NodeSpawner->bIsLocalFunction = true;
-	NodeSpawner->FunctionPath = InFunction->GetPathName();
-
-	FBlueprintActionUiSpec& MenuSignature = NodeSpawner->DefaultMenuSignature;
+	ReferencedPublicFunctionHeader = InFunction->GetFunctionHeader();
+	NodeClass = URigVMEdGraphNode::StaticClass();
+	bIsLocalFunction = true;
+	FunctionPath = InFunction->GetPathName();
 
 	const FString Category = InFunction->GetNodeCategory();
 
-	MenuSignature.MenuName = FText::FromString(InFunction->GetName());
-	MenuSignature.Tooltip = InFunction->GetToolTipText();
+	MenuName = FText::FromString(InFunction->GetName());
+	MenuTooltip = InFunction->GetToolTipText();
 
 	static const FString LocalFunctionString = TEXT("Local Function");
-	if(MenuSignature.Tooltip.IsEmpty())
+	if(MenuTooltip.IsEmpty())
 	{
-		MenuSignature.Tooltip = FText::FromString(LocalFunctionString);
+		MenuTooltip = FText::FromString(LocalFunctionString);
 	}
 	else
 	{
 		static constexpr TCHAR Format[] = TEXT("%s\n\n%s");
-		MenuSignature.Tooltip = FText::FromString(FString::Printf(Format, *MenuSignature.Tooltip.ToString(), *LocalFunctionString));
+		MenuTooltip = FText::FromString(FString::Printf(Format, *MenuTooltip.ToString(), *LocalFunctionString));
 	}
 	
-	MenuSignature.Category = FText::FromString(Category);
-	MenuSignature.Keywords = FText::FromString(InFunction->GetNodeKeywords());
+	MenuCategory = FText::FromString(Category);
+	MenuKeywords = FText::FromString(InFunction->GetNodeKeywords());
 
 	// add at least one character, so that PrimeDefaultUiSpec() doesn't 
 	// attempt to query the template node
-	if (MenuSignature.Keywords.IsEmpty())
+	if (MenuKeywords.IsEmpty())
 	{
 		// want to set it to something so we won't end up back in this condition
-		MenuSignature.Keywords = FText::FromString(TEXT(" "));
+		MenuKeywords = FText::FromString(TEXT(" "));
 	}
 
-	MenuSignature.Icon = FSlateIcon(FAppStyle::GetAppStyleSetName(), "Kismet.AllClasses.FunctionIcon");
+	MenuIcon = FSlateIcon(FAppStyle::GetAppStyleSetName(), "Kismet.AllClasses.FunctionIcon");
 
 #if WITH_EDITOR
 	if (InFunction->IsMutable())
 	{
-		MenuSignature.IconTint = GetDefault<UGraphEditorSettings>()->FunctionCallNodeTitleColor;
+		MenuIconTint = GetDefault<UGraphEditorSettings>()->FunctionCallNodeTitleColor;
 	}
 	else
 	{
-		MenuSignature.IconTint = GetDefault<UGraphEditorSettings>()->PureFunctionCallNodeTitleColor;
+		MenuIconTint = GetDefault<UGraphEditorSettings>()->PureFunctionCallNodeTitleColor;
 	}
 #endif
-	
-	return NodeSpawner;
 }
 
-URigVMEdGraphFunctionRefNodeSpawner* URigVMEdGraphFunctionRefNodeSpawner::CreateFromAssetData(const FAssetData& InAssetData, const FRigVMGraphFunctionHeader& InPublicFunction)
+URigVMEdGraphFunctionRefNodeSpawner::URigVMEdGraphFunctionRefNodeSpawner(const FAssetData& InAssetData, const FRigVMGraphFunctionHeader& InPublicFunction)
 {
-	URigVMEdGraphFunctionRefNodeSpawner* NodeSpawner = NewObject<URigVMEdGraphFunctionRefNodeSpawner>(GetTransientPackage());
-	NodeSpawner->ReferencedPublicFunctionHeader = InPublicFunction;
-	NodeSpawner->NodeClass = URigVMEdGraphNode::StaticClass();
-	NodeSpawner->bIsLocalFunction = false;
-	NodeSpawner->AssetPath = InAssetData.ToSoftObjectPath(); 	
+	ReferencedPublicFunctionHeader = InPublicFunction;
+	NodeClass = URigVMEdGraphNode::StaticClass();
+	bIsLocalFunction = false;
+	AssetPath = InAssetData.ToSoftObjectPath(); 	
 
-	FBlueprintActionUiSpec& MenuSignature = NodeSpawner->DefaultMenuSignature;
-
-	MenuSignature.MenuName = FText::FromName(InPublicFunction.Name);
-	MenuSignature.Category = FText::FromString(InPublicFunction.Category);
-	MenuSignature.Keywords = FText::FromString(InPublicFunction.Keywords);
-	MenuSignature.Tooltip = InPublicFunction.GetTooltip();
+	MenuName = FText::FromName(InPublicFunction.Name);
+	MenuCategory = FText::FromString(InPublicFunction.Category);
+	MenuKeywords = FText::FromString(InPublicFunction.Keywords);
+	MenuTooltip = InPublicFunction.GetTooltip();
 
 	const FString PackagePathString = InAssetData.PackageName.ToString();
-	if(MenuSignature.Tooltip.IsEmpty())
+	if(MenuTooltip.IsEmpty())
 	{
-		MenuSignature.Tooltip = FText::FromString(PackagePathString);
+		MenuTooltip = FText::FromString(PackagePathString);
 	}
 	else
 	{
 		static constexpr TCHAR Format[] = TEXT("%s\n\n%s");
-		MenuSignature.Tooltip = FText::FromString(FString::Printf(Format, *MenuSignature.Tooltip.ToString(), *PackagePathString));
+		MenuTooltip = FText::FromString(FString::Printf(Format, *MenuTooltip.ToString(), *PackagePathString));
 	}
 
 	// add at least one character, so that PrimeDefaultUiSpec() doesn't 
 	// attempt to query the template node
-	if (MenuSignature.Keywords.IsEmpty())
+	if (MenuKeywords.IsEmpty())
 	{
 		// want to set it to something so we won't end up back in this condition
-		MenuSignature.Keywords = FText::FromString(TEXT(" "));
+		MenuKeywords = FText::FromString(TEXT(" "));
 	}
 
-	MenuSignature.Icon = FSlateIcon(FAppStyle::GetAppStyleSetName(), "Kismet.AllClasses.FunctionIcon");
+	MenuIcon = FSlateIcon(FAppStyle::GetAppStyleSetName(), "Kismet.AllClasses.FunctionIcon");
 
 #if WITH_EDITOR
 	if (InPublicFunction.IsMutable())
 	{
-		MenuSignature.IconTint = GetDefault<UGraphEditorSettings>()->FunctionCallNodeTitleColor;
+		MenuIconTint = GetDefault<UGraphEditorSettings>()->FunctionCallNodeTitleColor;
 	}
 	else
 	{
-		MenuSignature.IconTint = GetDefault<UGraphEditorSettings>()->PureFunctionCallNodeTitleColor;
+		MenuIconTint = GetDefault<UGraphEditorSettings>()->PureFunctionCallNodeTitleColor;
 	}
 #endif
 	
-	return NodeSpawner;
 }
 
-URigVMEdGraphFunctionRefNodeSpawner* URigVMEdGraphFunctionRefNodeSpawner::CreateFromAssetData(const FAssetData& InAssetData, const FRigVMOldPublicFunctionData& InPublicFunction)
+URigVMEdGraphFunctionRefNodeSpawner::URigVMEdGraphFunctionRefNodeSpawner(const FAssetData& InAssetData, const FRigVMOldPublicFunctionData& InPublicFunction)
 {
-	URigVMEdGraphFunctionRefNodeSpawner* NodeSpawner = NewObject<URigVMEdGraphFunctionRefNodeSpawner>(GetTransientPackage());
-	NodeSpawner->NodeClass = URigVMEdGraphNode::StaticClass();
-	NodeSpawner->bIsLocalFunction = false;
+	NodeClass = URigVMEdGraphNode::StaticClass();
+	bIsLocalFunction = false;
 
-	FRigVMGraphFunctionHeader& Header = NodeSpawner->ReferencedPublicFunctionHeader;
+	FRigVMGraphFunctionHeader& Header = ReferencedPublicFunctionHeader;
 	Header.Name = InPublicFunction.Name;
 	Header.Arguments.Reserve(InPublicFunction.Arguments.Num());
 	for (const FRigVMOldPublicFunctionArg& Arg : InPublicFunction.Arguments)
@@ -156,69 +144,59 @@ URigVMEdGraphFunctionRefNodeSpawner* URigVMEdGraphFunctionRefNodeSpawner::Create
 		NewArgument.CPPTypeObject = FSoftObjectPath(Arg.CPPTypeObjectPath.ToString());
 		Header.Arguments.Add(NewArgument);
 	}
-	NodeSpawner->AssetPath = InAssetData.ToSoftObjectPath();
-
-	FBlueprintActionUiSpec& MenuSignature = NodeSpawner->DefaultMenuSignature;
+	AssetPath = InAssetData.ToSoftObjectPath();
 
 	const FString Category = InPublicFunction.Category;
 
-	MenuSignature.MenuName = FText::FromName(InPublicFunction.Name);
-	MenuSignature.Category = FText::FromString(Category);
-	MenuSignature.Keywords = FText::FromString(InPublicFunction.Keywords);
+	MenuName = FText::FromName(InPublicFunction.Name);
+	MenuCategory = FText::FromString(Category);
+	MenuKeywords = FText::FromString(InPublicFunction.Keywords);
 
-	if(const URigVMBlueprint* ReferencedBlueprint = Cast<URigVMBlueprint>(InAssetData.FastGetAsset(false)))
+	if(const TRigVMAssetInterfacePtr ReferencedBlueprint = InAssetData.FastGetAsset(false))
 	{
 		if(const URigVMFunctionLibrary* FunctionLibrary = ReferencedBlueprint->GetLocalFunctionLibrary())
 		{
 			if(const URigVMLibraryNode* FunctionNode = FunctionLibrary->FindFunction(InPublicFunction.Name))
 			{
-				MenuSignature.Tooltip = FunctionNode->GetToolTipText();
+				MenuTooltip = FunctionNode->GetToolTipText();
 			}
 		}
 	}
 
 	const FString ObjectPathString = InAssetData.GetObjectPathString();
-	if(MenuSignature.Tooltip.IsEmpty())
+	if(MenuTooltip.IsEmpty())
 	{
-		MenuSignature.Tooltip = FText::FromString(ObjectPathString);
+		MenuTooltip = FText::FromString(ObjectPathString);
 	}
 	else
 	{
 		static constexpr TCHAR Format[] = TEXT("%s\n\n%s");
-		MenuSignature.Tooltip = FText::FromString(FString::Printf(Format, *MenuSignature.Tooltip.ToString(), *ObjectPathString));
+		MenuTooltip = FText::FromString(FString::Printf(Format, *MenuTooltip.ToString(), *ObjectPathString));
 	}
 
 	// add at least one character, so that PrimeDefaultUiSpec() doesn't 
 	// attempt to query the template node
-	if (MenuSignature.Keywords.IsEmpty())
+	if (MenuKeywords.IsEmpty())
 	{
 		// want to set it to something so we won't end up back in this condition
-		MenuSignature.Keywords = FText::FromString(TEXT(" "));
+		MenuKeywords = FText::FromString(TEXT(" "));
 	}
 
-	MenuSignature.Icon = FSlateIcon(FAppStyle::GetAppStyleSetName(), "Kismet.AllClasses.FunctionIcon");
+	MenuIcon = FSlateIcon(FAppStyle::GetAppStyleSetName(), "Kismet.AllClasses.FunctionIcon");
 
 #if WITH_EDITOR
 	if (InPublicFunction.IsMutable())
 	{
-		MenuSignature.IconTint = GetDefault<UGraphEditorSettings>()->FunctionCallNodeTitleColor;
+		MenuIconTint = GetDefault<UGraphEditorSettings>()->FunctionCallNodeTitleColor;
 	}
 	else
 	{
-		MenuSignature.IconTint = GetDefault<UGraphEditorSettings>()->PureFunctionCallNodeTitleColor;
+		MenuIconTint = GetDefault<UGraphEditorSettings>()->PureFunctionCallNodeTitleColor;
 	}
 #endif
-	
-	return NodeSpawner;
-}
-
-void URigVMEdGraphFunctionRefNodeSpawner::Prime()
-{
-	// we expect that you don't need a node template to construct menu entries
-	// from this, so we choose not to pre-cache one here
 }
 
-FBlueprintNodeSignature URigVMEdGraphFunctionRefNodeSpawner::GetSpawnerSignature() const
+FString URigVMEdGraphFunctionRefNodeSpawner::GetSpawnerSignature() const
 {
 	FString SignatureString = TEXT("Invalid RigFunction");
 	if (ReferencedPublicFunctionHeader.IsValid())
@@ -240,19 +218,10 @@ FBlueprintNodeSignature URigVMEdGraphFunctionRefNodeSpawner::GetSpawnerSignature
 		SignatureString += TEXT(" (local)");
 	}
 
-	return FBlueprintNodeSignature(SignatureString);
+	return SignatureString;
 }
 
-FBlueprintActionUiSpec URigVMEdGraphFunctionRefNodeSpawner::GetUiSpec(FBlueprintActionContext const& Context, FBindingSet const& Bindings) const
-{
-	UEdGraph* TargetGraph = (Context.Graphs.Num() > 0) ? Context.Graphs[0] : nullptr;
-	FBlueprintActionUiSpec MenuSignature = PrimeDefaultUiSpec(TargetGraph);
-
-	DynamicUiSignatureGetter.ExecuteIfBound(Context, Bindings, &MenuSignature);
-	return MenuSignature;
-}
-
-UEdGraphNode* URigVMEdGraphFunctionRefNodeSpawner::Invoke(UEdGraph* ParentGraph, FBindingSet const& Bindings, FVector2D const Location) const
+URigVMEdGraphNode* URigVMEdGraphFunctionRefNodeSpawner::Invoke(URigVMEdGraph* ParentGraph, FVector2D const Location) const
 {
 	URigVMEdGraphNode* NewNode = nullptr;
 
@@ -261,9 +230,13 @@ UEdGraphNode* URigVMEdGraphFunctionRefNodeSpawner::Invoke(UEdGraph* ParentGraph,
 	{
 		if (GetResolvedAssetPath().IsValid())
 		{
-			if (URigVMBlueprint* Blueprint = Cast<URigVMBlueprint>(GetResolvedAssetPath().TryLoad()))
+			if (UObject* BlueprintObj = GetResolvedAssetPath().TryLoad())
 			{
-				ReferencedPublicFunctionHeader = Blueprint->GetLocalFunctionLibrary()->FindFunction(ReferencedPublicFunctionHeader.Name)->GetFunctionHeader();			
+				if (BlueprintObj->Implements<URigVMAssetInterface>())
+				{
+					TRigVMAssetInterfacePtr Blueprint = BlueprintObj;
+					ReferencedPublicFunctionHeader = Blueprint->GetLocalFunctionLibrary()->FindFunction(ReferencedPublicFunctionHeader.Name)->GetFunctionHeader();
+				}
 			}
 		}
 	}
@@ -278,7 +251,7 @@ UEdGraphNode* URigVMEdGraphFunctionRefNodeSpawner::Invoke(UEdGraph* ParentGraph,
 		}
 #endif
 
-		UBlueprint* Blueprint = FBlueprintEditorUtils::FindBlueprintForGraph(ParentGraph);
+		TRigVMAssetInterfacePtr Blueprint = FRigVMBlueprintUtils::FindAssetForGraph(ParentGraph);
 		NewNode = SpawnNode(ParentGraph, Blueprint, ReferencedPublicFunctionHeader, Location);
 	}
 	else
@@ -296,16 +269,15 @@ UEdGraphNode* URigVMEdGraphFunctionRefNodeSpawner::Invoke(UEdGraph* ParentGraph,
 	return NewNode;
 }
 
-URigVMEdGraphNode* URigVMEdGraphFunctionRefNodeSpawner::SpawnNode(UEdGraph* ParentGraph, UBlueprint* Blueprint, FRigVMGraphFunctionHeader& InFunction, FVector2D const Location)
+URigVMEdGraphNode* URigVMEdGraphFunctionRefNodeSpawner::SpawnNode(UEdGraph* ParentGraph, TRigVMAssetInterfacePtr Blueprint, FRigVMGraphFunctionHeader& InFunction, FVector2D const Location)
 {
 	URigVMEdGraphNode* NewNode = nullptr;
-	URigVMBlueprint* RigBlueprint = Cast<URigVMBlueprint>(Blueprint);
 	URigVMEdGraph* RigGraph = Cast<URigVMEdGraph>(ParentGraph);
 
-	if (RigBlueprint != nullptr && RigGraph != nullptr)
+	if (Blueprint != nullptr && RigGraph != nullptr)
 	{
-		FName Name = FRigVMBlueprintUtils::ValidateName(RigBlueprint, InFunction.Name.ToString());
-		URigVMController* Controller = RigBlueprint->GetController(ParentGraph);
+		FName Name = FRigVMBlueprintUtils::ValidateName(Blueprint, InFunction.Name.ToString());
+		URigVMController* Controller = Blueprint->GetController(ParentGraph);
 
 		Controller->OpenUndoBracket(FString::Printf(TEXT("Add '%s' Node"), *Name.ToString()));
 
@@ -319,7 +291,7 @@ URigVMEdGraphNode* URigVMEdGraphFunctionRefNodeSpawner::SpawnNode(UEdGraph* Pare
 				Controller->ClearNodeSelection(true);
 				Controller->SelectNode(ModelNode, true, true);
 
-				URigVMEdGraphUnitNodeSpawner::HookupMutableNode(ModelNode, RigBlueprint);
+				URigVMEdGraphUnitNodeSpawner::HookupMutableNode(ModelNode, Blueprint);
 			}
 
 			for(URigVMNode* OtherModelNode : ModelNode->GetGraph()->GetNodes())
@@ -359,9 +331,9 @@ URigVMEdGraphNode* URigVMEdGraphFunctionRefNodeSpawner::SpawnNode(UEdGraph* Pare
 	return NewNode;
 }
 
-bool URigVMEdGraphFunctionRefNodeSpawner::IsTemplateNodeFilteredOut(FBlueprintActionFilter const& Filter) const
+bool URigVMEdGraphFunctionRefNodeSpawner::IsTemplateNodeFilteredOut(TArray<UObject*>& InAssets, TArray<UEdGraph*> InGraphs, TArray<UEdGraphPin*> InPins) const
 {
-	if(URigVMEdGraphNodeSpawner::IsTemplateNodeFilteredOut(Filter))
+	if(URigVMEdGraphNodeSpawner::IsTemplateNodeFilteredOut(InAssets, InGraphs, InPins))
 	{
 		return true;
 	}
@@ -382,17 +354,20 @@ bool URigVMEdGraphFunctionRefNodeSpawner::IsTemplateNodeFilteredOut(FBlueprintAc
 	{
 		if(ReferencedPublicFunctionHeader.IsValid())
 		{
-			for (UBlueprint* Blueprint : Filter.Context.Blueprints)
+			for (UObject* Blueprint : InAssets)
 			{
-				if(Cast<IRigVMGraphFunctionHost>(Blueprint->GeneratedClass) != ReferencedPublicFunctionHeader.GetFunctionHost())
+				if (IRigVMAssetInterface* Asset = Cast<IRigVMAssetInterface>(Blueprint))
 				{
-					return true;
+					if (Asset->GetRigVMClientHost()->GetRigVMGraphFunctionHost() != ReferencedPublicFunctionHeader.GetFunctionHost())
+					{
+						return true;
+					}
 				}
 			}
 		}
 	}
 	const FString ReferencedAssetObjectPathString = ReferencedPublicFunctionHeader.LibraryPointer.GetNodeSoftPath().GetAssetName();
-	for (UBlueprint* Blueprint : Filter.Context.Blueprints)
+	for (UObject* Blueprint : InAssets)
 	{
 		if(Blueprint->GetPathName() == ReferencedAssetObjectPathString)
 		{
diff --git a/Engine/Plugins/Runtime/RigVM/Source/RigVMEditor/Private/EdGraph/NodeSpawners/RigVMEdGraphInvokeEntryNodeSpawner.cpp b/Engine/Plugins/Runtime/RigVM/Source/RigVMEditor/Private/EdGraph/NodeSpawners/RigVMEdGraphInvokeEntryNodeSpawner.cpp
--- a/Engine/Plugins/Runtime/RigVM/Source/RigVMEditor/Private/EdGraph/NodeSpawners/RigVMEdGraphInvokeEntryNodeSpawner.cpp
+++ b/Engine/Plugins/Runtime/RigVM/Source/RigVMEditor/Private/EdGraph/NodeSpawners/RigVMEdGraphInvokeEntryNodeSpawner.cpp
@@ -11,56 +11,35 @@
 #include "EditorCategoryUtils.h"
 #include "BlueprintNodeTemplateCache.h"
 #include "RigVMModel/RigVMController.h"
-#include "RigVMBlueprint.h"
+#include "RigVMBlueprintLegacy.h"
+#include "RigVMBlueprintUtils.h"
 #include "RigVMModel/Nodes/RigVMInvokeEntryNode.h"
 
-#include UE_INLINE_GENERATED_CPP_BY_NAME(RigVMEdGraphInvokeEntryNodeSpawner)
-
 #if WITH_EDITOR
 #include "Editor.h"
 #endif
 
 #define LOCTEXT_NAMESPACE "RigVMEdGraphInvokeEntryNodeSpawner"
 
-URigVMEdGraphInvokeEntryNodeSpawner* URigVMEdGraphInvokeEntryNodeSpawner::CreateForEntry(URigVMBlueprint* InBlueprint, const FName& InEntryName, const FText& InMenuDesc, const FText& InCategory, const FText& InTooltip)
-{
-	URigVMEdGraphInvokeEntryNodeSpawner* NodeSpawner = NewObject<URigVMEdGraphInvokeEntryNodeSpawner>(GetTransientPackage());
-	NodeSpawner->Blueprint = InBlueprint;
-	NodeSpawner->EntryName = InEntryName;
-	NodeSpawner->NodeClass = URigVMEdGraphNode::StaticClass();
-
-	FBlueprintActionUiSpec& MenuSignature = NodeSpawner->DefaultMenuSignature;
-	
-	MenuSignature.MenuName = InMenuDesc;
-	MenuSignature.Tooltip  = InTooltip;
-	MenuSignature.Category = InCategory;
-	MenuSignature.Keywords = FText::FromString(TEXT("Event,Entry,Invoke,Run,Launch,Start"));
-	MenuSignature.Icon = FSlateIcon(FAppStyle::GetAppStyleSetName(), "GraphEditor.Event_16x");
-
-	return NodeSpawner;
-}
-
-void URigVMEdGraphInvokeEntryNodeSpawner::Prime()
+URigVMEdGraphInvokeEntryNodeSpawner::URigVMEdGraphInvokeEntryNodeSpawner(TRigVMAssetInterfacePtr InBlueprint, const FName& InEntryName, const FText& InMenuDesc, const FText& InCategory, const FText& InTooltip)
 {
-	// we expect that you don't need a node template to construct menu entries
-	// from this, so we choose not to pre-cache one here
+	Blueprint = InBlueprint;
+	EntryName = InEntryName;
+	NodeClass = URigVMEdGraphNode::StaticClass();
+
+	MenuName = InMenuDesc;
+	MenuTooltip  = InTooltip;
+	MenuCategory = InCategory;
+	MenuKeywords = FText::FromString(TEXT("Event,Entry,Invoke,Run,Launch,Start"));
+	MenuIcon = FSlateIcon(FAppStyle::GetAppStyleSetName(), "GraphEditor.Event_16x");
 }
 
-FBlueprintNodeSignature URigVMEdGraphInvokeEntryNodeSpawner::GetSpawnerSignature() const
+FString URigVMEdGraphInvokeEntryNodeSpawner::GetSpawnerSignature() const
 {
-	return FBlueprintNodeSignature(FString("InvokeEntryNode=" + EntryName.ToString()));
-}
-
-FBlueprintActionUiSpec URigVMEdGraphInvokeEntryNodeSpawner::GetUiSpec(FBlueprintActionContext const& Context, FBindingSet const& Bindings) const
-{
-	UEdGraph* TargetGraph = (Context.Graphs.Num() > 0) ? Context.Graphs[0] : nullptr;
-	FBlueprintActionUiSpec MenuSignature = PrimeDefaultUiSpec(TargetGraph);
-
-	DynamicUiSignatureGetter.ExecuteIfBound(Context, Bindings, &MenuSignature);
-	return MenuSignature;
+	return FString("InvokeEntryNode=" + EntryName.ToString());
 }
 
-UEdGraphNode* URigVMEdGraphInvokeEntryNodeSpawner::Invoke(UEdGraph* ParentGraph, FBindingSet const& Bindings, FVector2D const Location) const
+URigVMEdGraphNode* URigVMEdGraphInvokeEntryNodeSpawner::Invoke(URigVMEdGraph* ParentGraph, FVector2D const Location) const
 {
 	URigVMEdGraphNode* NewNode = nullptr;
 
@@ -69,7 +48,7 @@ UEdGraphNode* URigVMEdGraphInvokeEntryNodeSpawner::Invoke(UEdGraph* ParentGraph,
 	// First create a backing member for our node
 	URigVMEdGraph* RigGraph = Cast<URigVMEdGraph>(ParentGraph);
 	if(RigGraph == nullptr) return nullptr;
-	URigVMBlueprint* RigBlueprint = Cast<URigVMBlueprint>(FBlueprintEditorUtils::FindBlueprintForGraph(ParentGraph));
+	TRigVMAssetInterfacePtr RigBlueprint = FRigVMBlueprintUtils::FindAssetForGraph(ParentGraph);
 	check(RigBlueprint);
 
 #if WITH_EDITOR
@@ -129,16 +108,16 @@ UEdGraphNode* URigVMEdGraphInvokeEntryNodeSpawner::Invoke(UEdGraph* ParentGraph,
 	return NewNode;
 }
 
-bool URigVMEdGraphInvokeEntryNodeSpawner::IsTemplateNodeFilteredOut(FBlueprintActionFilter const& Filter) const
+bool URigVMEdGraphInvokeEntryNodeSpawner::IsTemplateNodeFilteredOut(TArray<UObject*>& InAssets, TArray<UEdGraph*> InGraphs, TArray<UEdGraphPin*> InPins) const
 {
-	if(URigVMEdGraphNodeSpawner::IsTemplateNodeFilteredOut(Filter))
+	if(URigVMEdGraphNodeSpawner::IsTemplateNodeFilteredOut(InAssets, InGraphs, InPins))
 	{
 		return true;
 	}
 
 	if (Blueprint.IsValid())
 	{
-		if (!Filter.Context.Blueprints.Contains(Blueprint.Get()))
+		if (!InAssets.Contains(Blueprint.GetObject()))
 		{
 			return true;
 		}
diff --git a/Engine/Plugins/Runtime/RigVM/Source/RigVMEditor/Private/EdGraph/NodeSpawners/RigVMEdGraphNodeSpawner.cpp b/Engine/Plugins/Runtime/RigVM/Source/RigVMEditor/Private/EdGraph/NodeSpawners/RigVMEdGraphNodeSpawner.cpp
--- a/Engine/Plugins/Runtime/RigVM/Source/RigVMEditor/Private/EdGraph/NodeSpawners/RigVMEdGraphNodeSpawner.cpp
+++ b/Engine/Plugins/Runtime/RigVM/Source/RigVMEditor/Private/EdGraph/NodeSpawners/RigVMEdGraphNodeSpawner.cpp
@@ -7,13 +7,13 @@
 
 #define LOCTEXT_NAMESPACE "RigVMEdGraphNodeSpawner"
 
-bool URigVMEdGraphNodeSpawner::IsTemplateNodeFilteredOut(FBlueprintActionFilter const& Filter) const
+bool URigVMEdGraphNodeSpawner::IsTemplateNodeFilteredOut(TArray<UObject*>& InAssets, TArray<UEdGraph*> InGraphs, TArray<UEdGraphPin*> InPins) const
 {
 	check(RelatedBlueprintClass);
 
-	for(const UBlueprint* Blueprint : Filter.Context.Blueprints)
+	for(const UObject* Asset : InAssets)
 	{
-		if(Blueprint->GetClass() != RelatedBlueprintClass)
+		if(Asset->GetClass() != RelatedBlueprintClass)
 		{
 			return true;
 		}
@@ -22,12 +22,12 @@ bool URigVMEdGraphNodeSpawner::IsTemplateNodeFilteredOut(FBlueprintActionFilter
 	return false;
 }
 
-void URigVMEdGraphNodeSpawner::SetRelatedBlueprintClass(TSubclassOf<URigVMBlueprint> InClass)
+void URigVMEdGraphNodeSpawner::SetRelatedBlueprintClass(UClass* InClass)
 {
 	RelatedBlueprintClass = InClass;
 }
 
-URigVMEdGraphNode* URigVMEdGraphNodeSpawner::SpawnTemplateNode(UEdGraph* InParentGraph, const TArray<FPinInfo>& InPins, const FName& InNodeName)
+URigVMEdGraphNode* URigVMEdGraphNodeSpawner::SpawnTemplateNode(URigVMEdGraph* InParentGraph, const TArray<FPinInfo>& InPins, const FName& InNodeName)
 {
 	URigVMEdGraphNode* NewNode = NewObject<URigVMEdGraphNode>(InParentGraph, InNodeName);
 	InParentGraph->AddNode(NewNode, false);
@@ -67,5 +67,41 @@ URigVMEdGraphNode* URigVMEdGraphNodeSpawner::SpawnTemplateNode(UEdGraph* InParen
 	return NewNode;
 }
 
+FBlueprintNodeSignature URigVMEdGraphNodeBlueprintSpawner::GetSpawnerSignature() const
+{
+	return RigVMSpawner->GetSpawnerSignature();
+}
+
+FBlueprintActionUiSpec URigVMEdGraphNodeBlueprintSpawner::GetUiSpec(FBlueprintActionContext const& Context, FBindingSet const& Bindings) const
+{
+	UEdGraph* TargetGraph = (Context.Graphs.Num() > 0) ? Context.Graphs[0] : nullptr;
+	FBlueprintActionUiSpec MenuSignature = PrimeDefaultUiSpec(TargetGraph);
+
+	DynamicUiSignatureGetter.ExecuteIfBound(Context, Bindings, &MenuSignature);
+	return MenuSignature;
+}
+
+UEdGraphNode* URigVMEdGraphNodeBlueprintSpawner::Invoke(UEdGraph* ParentGraph, FBindingSet const& Bindings, FVector2D const Location) const
+{
+	return RigVMSpawner->Invoke(Cast<URigVMEdGraph>(ParentGraph), Location);
+}
+
+bool URigVMEdGraphNodeBlueprintSpawner::IsTemplateNodeFilteredOut(FBlueprintActionFilter const& Filter) const
+{
+	TArray<UObject*> BlueprintObjs;
+	Algo::Transform(Filter.Context.Blueprints, BlueprintObjs, [](UBlueprint* Blueprint) { return Blueprint; });
+	return RigVMSpawner->IsTemplateNodeFilteredOut(BlueprintObjs, Filter.Context.Graphs, Filter.Context.Pins);
+}
+
+void URigVMEdGraphNodeBlueprintSpawner::Prime()
+{
+	// we expect that you don't need a node template to construct menu entries
+	// from this, so we choose not to pre-cache one here
+}
+
+
+
+
+
 #undef LOCTEXT_NAMESPACE
 
diff --git a/Engine/Plugins/Runtime/RigVM/Source/RigVMEditor/Private/EdGraph/NodeSpawners/RigVMEdGraphTemplateNodeSpawner.cpp b/Engine/Plugins/Runtime/RigVM/Source/RigVMEditor/Private/EdGraph/NodeSpawners/RigVMEdGraphTemplateNodeSpawner.cpp
--- a/Engine/Plugins/Runtime/RigVM/Source/RigVMEditor/Private/EdGraph/NodeSpawners/RigVMEdGraphTemplateNodeSpawner.cpp
+++ b/Engine/Plugins/Runtime/RigVM/Source/RigVMEditor/Private/EdGraph/NodeSpawners/RigVMEdGraphTemplateNodeSpawner.cpp
@@ -17,32 +17,27 @@
 #include "ScopedTransaction.h"
 #include "RigVMModel/Nodes/RigVMDispatchNode.h"
 
-#include UE_INLINE_GENERATED_CPP_BY_NAME(RigVMEdGraphTemplateNodeSpawner)
-
 #if WITH_EDITOR
 #include "Editor.h"
 #endif
 
 #define LOCTEXT_NAMESPACE "RigVMEdGraphTemplateNodeSpawner"
 
-URigVMEdGraphTemplateNodeSpawner* URigVMEdGraphTemplateNodeSpawner::CreateFromNotation(const FName& InNotation, const FText& InMenuDesc, const FText& InCategory, const FText& InTooltip)
+URigVMEdGraphTemplateNodeSpawner::URigVMEdGraphTemplateNodeSpawner(const FName& InNotation, const FText& InMenuDesc, const FText& InCategory, const FText& InTooltip)
 {
-	URigVMEdGraphTemplateNodeSpawner* NodeSpawner = NewObject<URigVMEdGraphTemplateNodeSpawner>(GetTransientPackage());
-	NodeSpawner->TemplateNotation = InNotation;
-	NodeSpawner->Template =	FRigVMRegistry::Get().FindTemplate(InNotation);
-	NodeSpawner->NodeClass = URigVMEdGraphNode::StaticClass();
-
-	FBlueprintActionUiSpec& MenuSignature = NodeSpawner->DefaultMenuSignature;
-	
-	MenuSignature.MenuName = InMenuDesc;
-	MenuSignature.Tooltip  = InTooltip;
-	MenuSignature.Category = InCategory;
-
-	if (const FRigVMTemplate* Template = FRigVMRegistry::Get().FindTemplate(InNotation))
+	TemplateNotation = InNotation;
+	Template =	FRigVMRegistry::Get().FindTemplate(InNotation);
+	NodeClass = URigVMEdGraphNode::StaticClass();
+
+	MenuName = InMenuDesc;
+	MenuTooltip  = InTooltip;
+	MenuCategory = InCategory;
+
+	if (const FRigVMTemplate* RigVMTemplate = FRigVMRegistry::Get().FindTemplate(InNotation))
 	{
 #if WITH_EDITOR
-		FString KeywordsMetadata = Template->GetKeywords();
-		MenuSignature.Keywords = FText::FromString(KeywordsMetadata);
+		FString KeywordsMetadata = RigVMTemplate->GetKeywords();
+		MenuKeywords = FText::FromString(KeywordsMetadata);
 #endif
 	}
 
@@ -50,40 +45,32 @@ URigVMEdGraphTemplateNodeSpawner* URigVMEdGraphTemplateNodeSpawner::CreateFromNo
 	// attempt to query the template node
 	//
 	// @TODO: maybe UPROPERTY() fields should have keyword metadata like functions
-	if (MenuSignature.Keywords.IsEmpty())
+	if (MenuKeywords.IsEmpty())
 	{
 		// want to set it to something so we won't end up back in this condition
-		MenuSignature.Keywords = FText::FromString(TEXT(" "));
+		MenuKeywords = FText::FromString(TEXT(" "));
 	}
 
 	// @TODO: should use details customization-like extensibility system to provide editor only data like this
-	MenuSignature.Icon = FSlateIcon(TEXT("RigVMEditorStyle"), TEXT("RigVM.Unit"));
-
-	return NodeSpawner;
-}
-
-void URigVMEdGraphTemplateNodeSpawner::Prime()
-{
-	// we expect that you don't need a node template to construct menu entries
-	// from this, so we choose not to pre-cache one here
+	MenuIcon = FSlateIcon(TEXT("RigVMEditorStyle"), TEXT("RigVM.Unit"));
 }
 
-FBlueprintNodeSignature URigVMEdGraphTemplateNodeSpawner::GetSpawnerSignature() const
+FString URigVMEdGraphTemplateNodeSpawner::GetSpawnerSignature() const
 {
 	const int32 NotationHash = (int32)GetTypeHash(TemplateNotation);
-	return FBlueprintNodeSignature(FString("RigVMTemplate_") + FString::FromInt(NotationHash));
+	return FString("RigVMTemplate_") + FString::FromInt(NotationHash);
 }
 
-bool URigVMEdGraphTemplateNodeSpawner::IsTemplateNodeFilteredOut(FBlueprintActionFilter const& Filter) const
+bool URigVMEdGraphTemplateNodeSpawner::IsTemplateNodeFilteredOut(TArray<UObject*>& InAssets, TArray<UEdGraph*> InGraphs, TArray<UEdGraphPin*> InPins) const
 {
-	if(URigVMEdGraphNodeSpawner::IsTemplateNodeFilteredOut(Filter))
+	if(URigVMEdGraphNodeSpawner::IsTemplateNodeFilteredOut(InAssets, InGraphs, InPins))
 	{
 		return true;
 	}
 
-	for (const UEdGraphPin* Pin : Filter.Context.Pins)
+	for (const UEdGraphPin* Pin : InPins)
 	{
-		for (UEdGraph* Graph : Filter.Context.Graphs)
+		for (UEdGraph* Graph : InGraphs)
 		{
 			if (!Graph->Nodes.Contains(Pin->GetOwningNode()))
 			{
@@ -135,16 +122,7 @@ bool URigVMEdGraphTemplateNodeSpawner::IsTemplateNodeFilteredOut(FBlueprintActio
 	return false;
 }
 
-FBlueprintActionUiSpec URigVMEdGraphTemplateNodeSpawner::GetUiSpec(FBlueprintActionContext const& Context, FBindingSet const& Bindings) const
-{
-	UEdGraph* TargetGraph = (Context.Graphs.Num() > 0) ? Context.Graphs[0] : nullptr;
-	FBlueprintActionUiSpec MenuSignature = PrimeDefaultUiSpec(TargetGraph);
-
-	DynamicUiSignatureGetter.ExecuteIfBound(Context, Bindings, &MenuSignature);
-	return MenuSignature;
-}
-
-UEdGraphNode* URigVMEdGraphTemplateNodeSpawner::Invoke(UEdGraph* ParentGraph, FBindingSet const& Bindings, FVector2D const Location) const
+URigVMEdGraphNode* URigVMEdGraphTemplateNodeSpawner::Invoke(URigVMEdGraph* ParentGraph, FVector2D const Location) const
 {
 	URigVMEdGraphNode* NewNode = nullptr;
 
@@ -202,19 +180,18 @@ UEdGraphNode* URigVMEdGraphTemplateNodeSpawner::Invoke(UEdGraph* ParentGraph, FB
 				NewNode->ModelNodePath = Template->GetNotation().ToString();
 			}
 			return NewNode;
-		}			
+		}
 
-		UBlueprint* Blueprint = FBlueprintEditorUtils::FindBlueprintForGraph(ParentGraph);
+		TRigVMAssetInterfacePtr Blueprint = FRigVMBlueprintUtils::FindAssetForGraph(ParentGraph);
 		NewNode = SpawnNode(ParentGraph, Blueprint, Template, Location);
 	}
 
 	return NewNode;
 }
 
-URigVMEdGraphNode* URigVMEdGraphTemplateNodeSpawner::SpawnNode(UEdGraph* ParentGraph, UBlueprint* Blueprint, const FRigVMTemplate* Template, FVector2D const Location)
+URigVMEdGraphNode* URigVMEdGraphTemplateNodeSpawner::SpawnNode(URigVMEdGraph* ParentGraph, TRigVMAssetInterfacePtr RigBlueprint, const FRigVMTemplate* Template, FVector2D const Location)
 {
 	URigVMEdGraphNode* NewNode = nullptr;
-	URigVMBlueprint* RigBlueprint = Cast<URigVMBlueprint>(Blueprint);
 	URigVMEdGraph* RigGraph = Cast<URigVMEdGraph>(ParentGraph);
 
 	if (RigBlueprint != nullptr && RigGraph != nullptr)
diff --git a/Engine/Plugins/Runtime/RigVM/Source/RigVMEditor/Private/EdGraph/NodeSpawners/RigVMEdGraphUnitNodeSpawner.cpp b/Engine/Plugins/Runtime/RigVM/Source/RigVMEditor/Private/EdGraph/NodeSpawners/RigVMEdGraphUnitNodeSpawner.cpp
--- a/Engine/Plugins/Runtime/RigVM/Source/RigVMEditor/Private/EdGraph/NodeSpawners/RigVMEdGraphUnitNodeSpawner.cpp
+++ b/Engine/Plugins/Runtime/RigVM/Source/RigVMEditor/Private/EdGraph/NodeSpawners/RigVMEdGraphUnitNodeSpawner.cpp
@@ -18,8 +18,6 @@
 #include "RigVMFunctions/Execution/RigVMFunction_UserDefinedEvent.h"
 #include "RigVMPropertyUtils.h"
 
-#include UE_INLINE_GENERATED_CPP_BY_NAME(RigVMEdGraphUnitNodeSpawner)
-
 #if WITH_EDITOR
 #include "Editor.h"
 #include "SGraphActionMenu.h"
@@ -27,18 +25,15 @@
 
 #define LOCTEXT_NAMESPACE "RigVMEdGraphUnitNodeSpawner"
 
-URigVMEdGraphUnitNodeSpawner* URigVMEdGraphUnitNodeSpawner::CreateFromStruct(UScriptStruct* InStruct, const FName& InMethodName, const FText& InMenuDesc, const FText& InCategory, const FText& InTooltip)
+URigVMEdGraphUnitNodeSpawner::URigVMEdGraphUnitNodeSpawner(UScriptStruct* InStruct, const FName& InMethodName, const FText& InMenuDesc, const FText& InCategory, const FText& InTooltip)
 {
-	URigVMEdGraphUnitNodeSpawner* NodeSpawner = NewObject<URigVMEdGraphUnitNodeSpawner>(GetTransientPackage());
-	NodeSpawner->StructTemplate = InStruct;
-	NodeSpawner->MethodName = InMethodName;
-	NodeSpawner->NodeClass = URigVMEdGraphNode::StaticClass();
+	StructTemplate = InStruct;
+	MethodName = InMethodName;
+	NodeClass = URigVMEdGraphNode::StaticClass();
 
-	FBlueprintActionUiSpec& MenuSignature = NodeSpawner->DefaultMenuSignature;
-	
-	MenuSignature.MenuName = InMenuDesc;
-	MenuSignature.Tooltip  = InTooltip;
-	MenuSignature.Category = InCategory;
+	MenuName = InMenuDesc;
+	MenuTooltip  = InTooltip;
+	MenuCategory = InCategory;
 
 	FString KeywordsMetadata, TemplateNameMetadata;
 	InStruct->GetStringMetaDataHierarchical(FRigVMStruct::KeywordsMetaName, &KeywordsMetadata);
@@ -53,16 +48,16 @@ URigVMEdGraphUnitNodeSpawner* URigVMEdGraphUnitNodeSpawner::CreateFromStruct(USc
 			KeywordsMetadata = KeywordsMetadata + TEXT(",") + TemplateNameMetadata;
 		}
 	}
-	MenuSignature.Keywords = FText::FromString(KeywordsMetadata);
+	MenuKeywords = FText::FromString(KeywordsMetadata);
 
 	// add at least one character, so that PrimeDefaultUiSpec() doesn't 
 	// attempt to query the template node
 	//
 	// @TODO: maybe UPROPERTY() fields should have keyword metadata like functions
-	if (MenuSignature.Keywords.IsEmpty())
+	if (MenuKeywords.IsEmpty())
 	{
 		// want to set it to something so we won't end up back in this condition
-		MenuSignature.Keywords = FText::FromString(TEXT(" "));
+		MenuKeywords = FText::FromString(TEXT(" "));
 	}
 
 	// @TODO: should use details customization-like extensibility system to provide editor only data like this
@@ -101,42 +96,25 @@ URigVMEdGraphUnitNodeSpawner* URigVMEdGraphUnitNodeSpawner::CreateFromStruct(USc
 				NameIndex++;
 				IconPath = Right;
 			}
-			MenuSignature.Icon = FSlateIcon(IconPathNames[0], IconPathNames[1], IconPathNames[2], IconPathNames[3]);
+			MenuIcon = FSlateIcon(IconPathNames[0], IconPathNames[1], IconPathNames[2], IconPathNames[3]);
 		}
 		else
 		{
-			MenuSignature.Icon = FSlateIcon(TEXT("RigVMEditorStyle"), TEXT("RigVM.Unit"));
+			MenuIcon = FSlateIcon(TEXT("RigVMEditorStyle"), TEXT("RigVM.Unit"));
 		}
 	}
 	else
 	{
-		MenuSignature.Icon = FSlateIcon(FAppStyle::GetAppStyleSetName(), "GraphEditor.Event_16x");
+		MenuIcon = FSlateIcon(FAppStyle::GetAppStyleSetName(), "GraphEditor.Event_16x");
 	}
-
-	return NodeSpawner;
-}
-
-void URigVMEdGraphUnitNodeSpawner::Prime()
-{
-	// we expect that you don't need a node template to construct menu entries
-	// from this, so we choose not to pre-cache one here
 }
 
-FBlueprintNodeSignature URigVMEdGraphUnitNodeSpawner::GetSpawnerSignature() const
+FString URigVMEdGraphUnitNodeSpawner::GetSpawnerSignature() const
 {
-	return FBlueprintNodeSignature(FString("RigUnit=" + StructTemplate->GetFName().ToString()));
-}
-
-FBlueprintActionUiSpec URigVMEdGraphUnitNodeSpawner::GetUiSpec(FBlueprintActionContext const& Context, FBindingSet const& Bindings) const
-{
-	UEdGraph* TargetGraph = (Context.Graphs.Num() > 0) ? Context.Graphs[0] : nullptr;
-	FBlueprintActionUiSpec MenuSignature = PrimeDefaultUiSpec(TargetGraph);
-
-	DynamicUiSignatureGetter.ExecuteIfBound(Context, Bindings, &MenuSignature);
-	return MenuSignature;
+	return FString("RigUnit=" + StructTemplate->GetFName().ToString());
 }
 
-UEdGraphNode* URigVMEdGraphUnitNodeSpawner::Invoke(UEdGraph* ParentGraph, FBindingSet const& Bindings, FVector2D const Location) const
+URigVMEdGraphNode* URigVMEdGraphUnitNodeSpawner::Invoke(URigVMEdGraph* ParentGraph, FVector2D const Location) const
 {
 	URigVMEdGraphNode* NewNode = nullptr;
 
@@ -149,17 +127,16 @@ UEdGraphNode* URigVMEdGraphUnitNodeSpawner::Invoke(UEdGraph* ParentGraph, FBindi
 		}
 #endif
 
-		UBlueprint* Blueprint = FBlueprintEditorUtils::FindBlueprintForGraph(ParentGraph);
+		TRigVMAssetInterfacePtr Blueprint = FRigVMBlueprintUtils::FindAssetForGraph(ParentGraph);
 		NewNode = SpawnNode(ParentGraph, Blueprint, StructTemplate, MethodName, Location);
 	}
 
 	return NewNode;
 }
 
-URigVMEdGraphNode* URigVMEdGraphUnitNodeSpawner::SpawnNode(UEdGraph* ParentGraph, UBlueprint* Blueprint, UScriptStruct* StructTemplate, const FName& InMethodName, FVector2D const Location)
+URigVMEdGraphNode* URigVMEdGraphUnitNodeSpawner::SpawnNode(URigVMEdGraph* ParentGraph, TRigVMAssetInterfacePtr RigBlueprint, UScriptStruct* StructTemplate, const FName& InMethodName, FVector2D const Location)
 {
 	URigVMEdGraphNode* NewNode = nullptr;
-	URigVMBlueprint* RigBlueprint = Cast<URigVMBlueprint>(Blueprint);
 	URigVMEdGraph* RigGraph = Cast<URigVMEdGraph>(ParentGraph);
 
 	if (RigBlueprint != nullptr && RigGraph != nullptr)
@@ -252,7 +229,8 @@ URigVMEdGraphNode* URigVMEdGraphUnitNodeSpawner::SpawnNode(UEdGraph* ParentGraph
 		Controller->OpenUndoBracket(FString::Printf(TEXT("Add '%s' Node"), *Name.ToString()));
 
 		FRigVMUnitNodeCreatedContext& UnitNodeCreatedContext = Controller->GetUnitNodeCreatedContext();
-		FRigVMUnitNodeCreatedContext::FScope ReasonScope(UnitNodeCreatedContext, ERigVMNodeCreatedReason::NodeSpawner, RigBlueprint);
+		TScriptInterface<IRigVMClientHost> ClientHost = RigBlueprint.GetObject();
+		FRigVMUnitNodeCreatedContext::FScope ReasonScope(UnitNodeCreatedContext, ERigVMNodeCreatedReason::NodeSpawner, ClientHost.GetInterface());
 
 		if (URigVMUnitNode* ModelNode = Controller->AddUnitNode(StructTemplate, InMethodName, Location, Name.ToString(), true, true))
 		{
@@ -372,7 +350,7 @@ URigVMEdGraphNode* URigVMEdGraphUnitNodeSpawner::SpawnNode(UEdGraph* ParentGraph
 	return NewNode;
 }
 
-void URigVMEdGraphUnitNodeSpawner::HookupMutableNode(URigVMNode* InModelNode, URigVMBlueprint* InRigBlueprint)
+void URigVMEdGraphUnitNodeSpawner::HookupMutableNode(URigVMNode* InModelNode, TRigVMAssetInterfacePtr InRigBlueprint)
 {
 	if(!InRigBlueprint->GetRigVMEditorSettings()->bAutoLinkMutableNodes)
 	{
diff --git a/Engine/Plugins/Runtime/RigVM/Source/RigVMEditor/Private/EdGraph/NodeSpawners/RigVMEdGraphVariableNodeSpawner.cpp b/Engine/Plugins/Runtime/RigVM/Source/RigVMEditor/Private/EdGraph/NodeSpawners/RigVMEdGraphVariableNodeSpawner.cpp
--- a/Engine/Plugins/Runtime/RigVM/Source/RigVMEditor/Private/EdGraph/NodeSpawners/RigVMEdGraphVariableNodeSpawner.cpp
+++ b/Engine/Plugins/Runtime/RigVM/Source/RigVMEditor/Private/EdGraph/NodeSpawners/RigVMEdGraphVariableNodeSpawner.cpp
@@ -12,64 +12,51 @@
 #include "K2Node_Variable.h"
 #include "BlueprintNodeTemplateCache.h"
 #include "RigVMModel/RigVMController.h"
-#include "RigVMBlueprint.h"
+#include "RigVMBlueprintLegacy.h"
+#include "RigVMBlueprintUtils.h"
 #include "RigVMModel/Nodes/RigVMVariableNode.h"
 
-#include UE_INLINE_GENERATED_CPP_BY_NAME(RigVMEdGraphVariableNodeSpawner)
-
 #if WITH_EDITOR
 #include "Editor.h"
 #endif
 
 #define LOCTEXT_NAMESPACE "RigVMEdGraphVariableNodeSpawner"
 
-URigVMEdGraphVariableNodeSpawner* URigVMEdGraphVariableNodeSpawner::CreateFromExternalVariable(URigVMBlueprint* InBlueprint, const FRigVMExternalVariable& InExternalVariable, bool bInIsGetter, const FText& InMenuDesc, const FText& InCategory, const FText& InTooltip)
+URigVMEdGraphVariableNodeSpawner::URigVMEdGraphVariableNodeSpawner(TRigVMAssetInterfacePtr InBlueprint, const FRigVMExternalVariable& InExternalVariable, bool bInIsGetter, const FText& InMenuDesc, const FText& InCategory, const FText& InTooltip)
 {
-	URigVMEdGraphVariableNodeSpawner* NodeSpawner = NewObject<URigVMEdGraphVariableNodeSpawner>(GetTransientPackage());
-	NodeSpawner->Blueprint = InBlueprint;
-	NodeSpawner->ExternalVariable = InExternalVariable;
-	NodeSpawner->bIsGetter = bInIsGetter;
-	NodeSpawner->NodeClass = URigVMEdGraphNode::StaticClass();
-
-	FBlueprintActionUiSpec& MenuSignature = NodeSpawner->DefaultMenuSignature;
-	
-	MenuSignature.MenuName = FText::FromString(FString::Printf(TEXT("%s %s"), bInIsGetter ? TEXT("Get") : TEXT("Set"), *InMenuDesc.ToString()));
-	MenuSignature.Tooltip  = InTooltip;
-	MenuSignature.Category = InCategory;
-	MenuSignature.Keywords = FText::FromString(TEXT("Variable"));
+	Blueprint = InBlueprint;
+	ExternalVariable = InExternalVariable;
+	bIsGetter = bInIsGetter;
+	NodeClass = URigVMEdGraphNode::StaticClass();
 
-	FEdGraphPinType PinType = RigVMTypeUtils::PinTypeFromExternalVariable(InExternalVariable);
-	MenuSignature.Icon = UK2Node_Variable::GetVarIconFromPinType(PinType, MenuSignature.IconTint);
+	MenuName = FText::FromString(FString::Printf(TEXT("%s %s"), bInIsGetter ? TEXT("Get") : TEXT("Set"), *InMenuDesc.ToString()));
+	MenuTooltip  = InTooltip;
+	MenuCategory = InCategory;
+	MenuKeywords = FText::FromString(TEXT("Variable"));
 
-	return NodeSpawner;
+	FEdGraphPinType PinType = RigVMTypeUtils::PinTypeFromExternalVariable(InExternalVariable);
+	MenuIcon = UK2Node_Variable::GetVarIconFromPinType(PinType, MenuIconTint);
 }
 
-URigVMEdGraphVariableNodeSpawner* URigVMEdGraphVariableNodeSpawner::CreateFromLocalVariable(
-	URigVMBlueprint* InBlueprint, URigVMGraph* InGraphOwner, const FRigVMGraphVariableDescription& InLocalVariable,
+URigVMEdGraphVariableNodeSpawner::URigVMEdGraphVariableNodeSpawner(
+	TRigVMAssetInterfacePtr InBlueprint, URigVMGraph* InGraphOwner, const FRigVMGraphVariableDescription& InLocalVariable,
 	bool bInIsGetter, const FText& InMenuDesc, const FText& InCategory, const FText& InTooltip)
+		: URigVMEdGraphVariableNodeSpawner(InBlueprint, InLocalVariable.ToExternalVariable(), bInIsGetter, InMenuDesc, InCategory, InTooltip)
 {
-	URigVMEdGraphVariableNodeSpawner* Spawner = CreateFromExternalVariable(InBlueprint, InLocalVariable.ToExternalVariable(), bInIsGetter, InMenuDesc, InCategory, InTooltip);
-	Spawner->bIsLocalVariable = true;
-	Spawner->GraphOwner = InGraphOwner;
-	return Spawner;
-}
-
-void URigVMEdGraphVariableNodeSpawner::Prime()
-{
-	// we expect that you don't need a node template to construct menu entries
-	// from this, so we choose not to pre-cache one here
+	bIsLocalVariable = true;
+	GraphOwner = InGraphOwner;
 }
 
-bool URigVMEdGraphVariableNodeSpawner::IsTemplateNodeFilteredOut(FBlueprintActionFilter const& Filter) const
+bool URigVMEdGraphVariableNodeSpawner::IsTemplateNodeFilteredOut(TArray<UObject*>& InAssets, TArray<UEdGraph*> InGraphs, TArray<UEdGraphPin*> InPins) const
 {
-	if(URigVMEdGraphNodeSpawner::IsTemplateNodeFilteredOut(Filter))
+	if(URigVMEdGraphNodeSpawner::IsTemplateNodeFilteredOut(InAssets, InGraphs, InPins))
 	{
 		return true;
 	}
 
 	if (Blueprint.IsValid())
 	{
-		if (!Filter.Context.Blueprints.Contains(Blueprint.Get()))
+		if (!InAssets.Contains(Blueprint.GetObject()))
 		{
 			return true;
 		}
@@ -77,9 +64,9 @@ bool URigVMEdGraphVariableNodeSpawner::IsTemplateNodeFilteredOut(FBlueprintActio
 		if (bIsLocalVariable)
 		{
 			bool bIsFiltered = true;
-			if (Filter.Context.Graphs.Num() == 1 && GraphOwner.IsValid())
+			if (InGraphs.Num() == 1 && GraphOwner.IsValid())
 			{
-				if (URigVMEdGraph* Graph = Cast<URigVMEdGraph>(Filter.Context.Graphs[0]))
+				if (URigVMEdGraph* Graph = Cast<URigVMEdGraph>(InGraphs[0]))
 				{
 					if (GraphOwner.Get() == Graph->GetModel())
 					{
@@ -97,21 +84,12 @@ bool URigVMEdGraphVariableNodeSpawner::IsTemplateNodeFilteredOut(FBlueprintActio
 	return false;
 }
 
-FBlueprintNodeSignature URigVMEdGraphVariableNodeSpawner::GetSpawnerSignature() const
+FString URigVMEdGraphVariableNodeSpawner::GetSpawnerSignature() const
 {
-	return FBlueprintNodeSignature(FString("ExternalVariable=" + ExternalVariable.Name.ToString()));
-}
-
-FBlueprintActionUiSpec URigVMEdGraphVariableNodeSpawner::GetUiSpec(FBlueprintActionContext const& Context, FBindingSet const& Bindings) const
-{
-	UEdGraph* TargetGraph = (Context.Graphs.Num() > 0) ? Context.Graphs[0] : nullptr;
-	FBlueprintActionUiSpec MenuSignature = PrimeDefaultUiSpec(TargetGraph);
-
-	DynamicUiSignatureGetter.ExecuteIfBound(Context, Bindings, &MenuSignature);
-	return MenuSignature;
+	return FString("ExternalVariable=" + ExternalVariable.Name.ToString());
 }
 
-UEdGraphNode* URigVMEdGraphVariableNodeSpawner::Invoke(UEdGraph* ParentGraph, FBindingSet const& Bindings, FVector2D const Location) const
+URigVMEdGraphNode* URigVMEdGraphVariableNodeSpawner::Invoke(URigVMEdGraph* ParentGraph, FVector2D const Location) const
 {
 	URigVMEdGraphNode* NewNode = nullptr;
 
@@ -120,7 +98,7 @@ UEdGraphNode* URigVMEdGraphVariableNodeSpawner::Invoke(UEdGraph* ParentGraph, FB
 	// First create a backing member for our node
 	URigVMEdGraph* RigGraph = Cast<URigVMEdGraph>(ParentGraph);
 	if(RigGraph == nullptr) return nullptr;
-	URigVMBlueprint* RigBlueprint = Cast<URigVMBlueprint>(FBlueprintEditorUtils::FindBlueprintForGraph(ParentGraph));
+	TRigVMAssetInterfacePtr RigBlueprint = FRigVMBlueprintUtils::FindAssetForGraph(ParentGraph);
 	check(RigBlueprint);
 
 	FName MemberName = NAME_None;
diff --git a/Engine/Plugins/Runtime/RigVM/Source/RigVMEditor/Private/EdGraph/RigVMEdGraphConnectionDrawingPolicy.cpp b/Engine/Plugins/Runtime/RigVM/Source/RigVMEditor/Private/EdGraph/RigVMEdGraphConnectionDrawingPolicy.cpp
--- a/Engine/Plugins/Runtime/RigVM/Source/RigVMEditor/Private/EdGraph/RigVMEdGraphConnectionDrawingPolicy.cpp
+++ b/Engine/Plugins/Runtime/RigVM/Source/RigVMEditor/Private/EdGraph/RigVMEdGraphConnectionDrawingPolicy.cpp
@@ -4,7 +4,8 @@
 #include "EdGraph/RigVMEdGraph.h"
 #include "EdGraph/RigVMEdGraphNode.h"
 #include "RigVMHost.h"
-#include "RigVMBlueprint.h"
+#include "RigVMAsset.h"
+#include "RigVMBlueprintUtils.h"
 #include "RigVMModel/RigVMController.h"
 #include "Kismet2/BlueprintEditorUtils.h"
 
@@ -322,7 +323,7 @@ void FRigVMEdGraphConnectionDrawingPolicy::DetermineWiringStyle(UEdGraphPin* Out
 
 			if (OutputInstructionIndex != INDEX_NONE && InputInstructionIndex != INDEX_NONE)
 			{
-				if (URigVMBlueprint* RigBlueprint = Cast<URigVMBlueprint>(FBlueprintEditorUtils::FindBlueprintForNode(OutputNode)))
+				if (TRigVMAssetInterfacePtr RigBlueprint = FRigVMBlueprintUtils::FindAssetForNode(OutputNode))
 				{
 					if (URigVMHost* Host = Cast<URigVMHost>(RigBlueprint->GetObjectBeingDebugged()))
 					{
diff --git a/Engine/Plugins/Runtime/RigVM/Source/RigVMEditor/Private/EdGraph/RigVMEdGraphPanelNodeFactory.cpp b/Engine/Plugins/Runtime/RigVM/Source/RigVMEditor/Private/EdGraph/RigVMEdGraphPanelNodeFactory.cpp
--- a/Engine/Plugins/Runtime/RigVM/Source/RigVMEditor/Private/EdGraph/RigVMEdGraphPanelNodeFactory.cpp
+++ b/Engine/Plugins/Runtime/RigVM/Source/RigVMEditor/Private/EdGraph/RigVMEdGraphPanelNodeFactory.cpp
@@ -13,7 +13,7 @@
 
 FName FRigVMEdGraphPanelNodeFactory::GetFactoryName() const
 {
-	return URigVMBlueprint::RigVMPanelNodeFactoryName;
+	return IRigVMAssetInterface::RigVMPanelNodeFactoryName;
 }
 
 TSharedPtr<SGraphNode> FRigVMEdGraphPanelNodeFactory::CreateNode(UEdGraphNode* Node) const
@@ -21,7 +21,7 @@ TSharedPtr<SGraphNode> FRigVMEdGraphPanelNodeFactory::CreateNode(UEdGraphNode* N
 	// we need to check if this is the right factory for the implementation
 	if(const URigVMEdGraph* EdGraph = Node->GetTypedOuter<URigVMEdGraph>())
 	{
-		if(const URigVMBlueprint* Blueprint = EdGraph->GetBlueprintDefaultObject())
+		if(const TRigVMAssetInterfacePtr Blueprint = EdGraph->GetAssetDefaultObject())
 		{
 			if(Blueprint->GetPanelNodeFactoryName() != GetFactoryName())
 			{
diff --git a/Engine/Plugins/Runtime/RigVM/Source/RigVMEditor/Private/EdGraph/RigVMEdGraphPanelPinFactory.cpp b/Engine/Plugins/Runtime/RigVM/Source/RigVMEditor/Private/EdGraph/RigVMEdGraphPanelPinFactory.cpp
--- a/Engine/Plugins/Runtime/RigVM/Source/RigVMEditor/Private/EdGraph/RigVMEdGraphPanelPinFactory.cpp
+++ b/Engine/Plugins/Runtime/RigVM/Source/RigVMEditor/Private/EdGraph/RigVMEdGraphPanelPinFactory.cpp
@@ -24,7 +24,7 @@
 
 FName FRigVMEdGraphPanelPinFactory::GetFactoryName() const
 {
-	return URigVMBlueprint::RigVMPanelPinFactoryName;
+	return IRigVMAssetInterface::RigVMPanelPinFactoryName;
 }
 
 TSharedPtr<SGraphPin> FRigVMEdGraphPanelPinFactory::CreatePin(UEdGraphPin* InPin) const
@@ -39,7 +39,7 @@ TSharedPtr<SGraphPin> FRigVMEdGraphPanelPinFactory::CreatePin(UEdGraphPin* InPin
 	{
 		if(const URigVMEdGraph* EdGraph = EdGraphNode->GetTypedOuter<URigVMEdGraph>())
 		{
-			if(const URigVMBlueprint* Blueprint = EdGraph->GetBlueprintDefaultObject())
+			if(const TRigVMAssetInterfacePtr Blueprint = EdGraph->GetAssetDefaultObject())
 			{
 				if(Blueprint->GetPanelPinFactoryName() != GetFactoryName())
 				{
@@ -90,11 +90,11 @@ TSharedPtr<SGraphPin> FRigVMEdGraphPanelPinFactory::CreatePin_Internal(UEdGraphP
 			{
 				if (ModelPin->IsBoundToVariable())
 				{
-					if (URigVMBlueprint* Blueprint = RigGraph->GetTypedOuter<URigVMBlueprint>())
+					if (TRigVMAssetInterfacePtr Blueprint = IRigVMAssetInterface::GetInterfaceOuter(RigGraph))
 					{
 						return SNew(SRigVMGraphPinVariableBinding, InPin)
 							.ModelPins({ModelPin})
-							.Blueprint(Blueprint);
+							.Asset(Blueprint);
 					}
 				}
 
diff --git a/Engine/Plugins/Runtime/RigVM/Source/RigVMEditor/Private/Editor/RigVMActionMenuItem.cpp b/Engine/Plugins/Runtime/RigVM/Source/RigVMEditor/Private/Editor/RigVMActionMenuItem.cpp
--- a/Engine/Plugins/Runtime/RigVM/Source/RigVMEditor/Private/Editor/RigVMActionMenuItem.cpp
+++ b/Engine/Plugins/Runtime/RigVM/Source/RigVMEditor/Private/Editor/RigVMActionMenuItem.cpp
@@ -14,6 +14,7 @@
 #include "IDocumentationPage.h"
 #include "Internationalization/Internationalization.h"
 #include "K2Node.h"
+#include "RigVMBlueprintUtils.h"
 #include "Kismet2/BlueprintEditorUtils.h"
 #include "Kismet2/KismetEditorUtilities.h"
 #include "Math/UnrealMathSSE.h"
@@ -82,13 +83,17 @@ namespace FBlueprintMenuActionItemImpl
 //------------------------------------------------------------------------------
 static void FBlueprintMenuActionItemImpl::DirtyBlueprintFromNewNode(UEdGraphNode* SpawnedNode)
 {
+	if (!SpawnedNode)
+	{
+		return; 
+	}
 	UEdGraph const* const NodeGraph = SpawnedNode->GetGraph();
 	check(NodeGraph != nullptr);
 	
-	UBlueprint* Blueprint = FBlueprintEditorUtils::FindBlueprintForGraphChecked(NodeGraph);
+	TRigVMAssetInterfacePtr Blueprint = FRigVMBlueprintUtils::FindAssetForGraph(NodeGraph);
 	check(Blueprint != nullptr);
 
-	if (SpawnedNode->GetSchema()->MarkBlueprintDirtyFromNewNode(Blueprint, SpawnedNode))
+	if (!Blueprint || !SpawnedNode)
 	{
 		return;
 	}
@@ -99,11 +104,11 @@ static void FBlueprintMenuActionItemImpl::DirtyBlueprintFromNewNode(UEdGraphNode
 	// non-k2 nodes structurally modify the blueprint)
 	if ((K2Node == nullptr) || K2Node->NodeCausesStructuralBlueprintChange())
 	{
-		FBlueprintEditorUtils::MarkBlueprintAsStructurallyModified(Blueprint);
+		Blueprint->MarkAssetAsStructurallyModified();
 	}
 	else
 	{
-		FBlueprintEditorUtils::MarkBlueprintAsModified(Blueprint);
+		Blueprint->MarkAssetAsModified();
 	}
 }
 
diff --git a/Engine/Plugins/Runtime/RigVM/Source/RigVMEditor/Private/Editor/RigVMCommentNodeDetailCustomization.cpp b/Engine/Plugins/Runtime/RigVM/Source/RigVMEditor/Private/Editor/RigVMCommentNodeDetailCustomization.cpp
--- a/Engine/Plugins/Runtime/RigVM/Source/RigVMEditor/Private/Editor/RigVMCommentNodeDetailCustomization.cpp
+++ b/Engine/Plugins/Runtime/RigVM/Source/RigVMEditor/Private/Editor/RigVMCommentNodeDetailCustomization.cpp
@@ -6,11 +6,12 @@
 #include "IDetailChildrenBuilder.h"
 #include "Widgets/Input/SCheckBox.h"
 #include "Widgets/Input/STextComboBox.h"
-#include "RigVMBlueprint.h"
+#include "RigVMAsset.h"
 #include "RigVMCore/RigVM.h"
 #include "Widgets/Colors/SColorBlock.h"
 #include "Widgets/Colors/SColorPicker.h"
 #include "UnrealEngine.h"
+#include "EdGraph/RigVMEdGraphSchema.h"
 #include "Widgets/Input/SNumericEntryBox.h"
 #include "Widgets/Input/SEditableTextBox.h"
 
@@ -30,7 +31,7 @@ void FRigVMCommentNodeDetailCustomization::CustomizeDetails(IDetailLayoutBuilder
 
 	if(ObjectsBeingCustomized[0].IsValid())
 	{
-		BlueprintBeingCustomized = ObjectsBeingCustomized[0]->GetTypedOuter<URigVMBlueprint>();
+		BlueprintBeingCustomized = ObjectsBeingCustomized[0]->GetImplementingOuter<IRigVMAssetInterface>();
 	}
 
 	IDetailCategoryBuilder& Category = DetailBuilder.EditCategory(TEXT("Comment Node"));
