RigVM: Adding IRigVMAssetInterface as parent of RigVMBlueprint. First step to separation of RigVM and Blueprint.

diff --git a/Engine/Plugins/Runtime/RigVM/Source/RigVMEditor/Private/Editor/RigVMEditor.cpp b/Engine/Plugins/Runtime/RigVM/Source/RigVMEditor/Private/Editor/RigVMEditor.cpp
--- a/Engine/Plugins/Runtime/RigVM/Source/RigVMEditor/Private/Editor/RigVMEditor.cpp
+++ b/Engine/Plugins/Runtime/RigVM/Source/RigVMEditor/Private/Editor/RigVMEditor.cpp
@@ -8,7 +8,7 @@
 #include "EditorModeManager.h"
 #include "EulerTransform.h"
 #include "InstancedPropertyBagStructureDataProvider.h"
-#include "RigVMBlueprint.h"
+#include "RigVMAsset.h"
 #include "RigVMBlueprintUtils.h"
 #include "RigVMEditorCommands.h"
 #include "RigVMPythonUtils.h"
@@ -72,7 +72,7 @@ FRigVMEditorBase::FRigVMEditorBase()
 
 void FRigVMEditorBase::UnbindEditor()
 {
-	URigVMBlueprint* RigVMBlueprint = GetRigVMBlueprint();
+	TRigVMAssetInterfacePtr RigVMBlueprint = GetRigVMAsset();
 	RigVMEditorClosedDelegate.Broadcast(this, RigVMBlueprint);
 
 	ClearDetailObject();
@@ -92,7 +92,7 @@ void FRigVMEditorBase::UnbindEditor()
 		// clear editor related data from the debugged control rig instance 
 		RigVMBlueprint->SetObjectBeingDebugged(nullptr);
 
-		URigVMBlueprint::sCurrentlyOpenedRigVMBlueprints.Remove(RigVMBlueprint);
+		IRigVMAssetInterface::sCurrentlyOpenedRigVMBlueprints.Remove(RigVMBlueprint);
 
 		RigVMBlueprint->OnRefreshEditor().RemoveAll(this);
 		RigVMBlueprint->OnVariableDropped().RemoveAll(this);
@@ -111,7 +111,7 @@ void FRigVMEditorBase::UnbindEditor()
 		RigVMBlueprint->OnRequestInspectObject().RemoveAll(this);
 		RigVMBlueprint->OnRequestInspectMemoryStorage().RemoveAll(this);
 
-		for(UEdGraph* Graph : RigVMBlueprint->UbergraphPages)
+		for(UEdGraph* Graph : RigVMBlueprint->GetUberGraphs())
 		{
 			if (URigVMEdGraph* RigGraph = Cast<URigVMEdGraph>(Graph))
 			{
@@ -141,7 +141,7 @@ void FRigVMEditorBase::UnbindEditor()
 		// the callback finishes, so the open editor action will fail. In that case, let's make sure we
 		// detect a reopen requested, and open the editor again on the next tick.
 		bRequestedReopen = false;
-		FSoftObjectPath AssetToReopen = RigVMBlueprint;
+		FSoftObjectPath AssetToReopen = RigVMBlueprint->GetObject();
 		GEditor->GetTimerManager()->SetTimerForNextTick([AssetToReopen]()
 		{
 			GEditor->GetEditorSubsystem<UAssetEditorSubsystem>()->OpenEditorForAsset(AssetToReopen);
@@ -173,7 +173,7 @@ FRigVMEditorBase* FRigVMEditorBase::GetFromAssetEditorInstance(IAssetEditorInsta
 	return Editor;
 }
 
-void FRigVMEditorBase::InitRigVMEditor(const EToolkitMode::Type Mode, const TSharedPtr< class IToolkitHost >& InitToolkitHost, class URigVMBlueprint* InRigVMBlueprint)
+void FRigVMEditorBase::InitRigVMEditor(const EToolkitMode::Type Mode, const TSharedPtr< class IToolkitHost >& InitToolkitHost, TRigVMAssetInterfacePtr InRigVMBlueprint)
 {
 	DECLARE_SCOPE_HIERARCHICAL_COUNTER_FUNC()
 
@@ -197,7 +197,7 @@ void FRigVMEditorBase::InitRigVMEditor(const EToolkitMode::Type Mode, const TSha
 
 	// Build up a list of objects being edited in this asset editor
 	TArray<UObject*> ObjectsBeingEdited;
-	ObjectsBeingEdited.Add(InRigVMBlueprint);
+	ObjectsBeingEdited.Add(InRigVMBlueprint->GetObject());
 
 	// Initialize the asset editor and spawn tabs
 	const bool bCreateDefaultStandaloneMenu = true;
@@ -208,7 +208,7 @@ void FRigVMEditorBase::InitRigVMEditor(const EToolkitMode::Type Mode, const TSha
 
 	CreateDefaultCommands();
 
-	TArray<UBlueprint*> Blueprints;
+	TArray<TRigVMAssetInterfacePtr> Blueprints;
 	Blueprints.Add(InRigVMBlueprint);
 	InRigVMBlueprint->InitializeModelIfRequired();
 
@@ -216,7 +216,7 @@ void FRigVMEditorBase::InitRigVMEditor(const EToolkitMode::Type Mode, const TSha
 
 	// Refresh the class actions
 	{
-		UClass* ActionKey = InRigVMBlueprint->GetClass();
+		UClass* ActionKey = InRigVMBlueprint->GetObject()->GetClass();
 		FBlueprintActionDatabase& ActionDatabase = FBlueprintActionDatabase::Get();
 		FBlueprintActionDatabase::FActionRegistry const& ActionRegistry = ActionDatabase.GetAllActions();
 		ActionDatabase.ClearAssetActions(ActionKey);
@@ -230,7 +230,7 @@ void FRigVMEditorBase::InitRigVMEditor(const EToolkitMode::Type Mode, const TSha
 
 	{
 		TArray<UEdGraph*> EdGraphs;
-		InRigVMBlueprint->GetAllGraphs(EdGraphs);
+		InRigVMBlueprint->GetAllEdGraphs(EdGraphs);
 
 		for (UEdGraph* Graph : EdGraphs)
 		{
@@ -240,12 +240,12 @@ void FRigVMEditorBase::InitRigVMEditor(const EToolkitMode::Type Mode, const TSha
 				continue;
 			}
 
-			RigVMEdGraph->InitializeFromBlueprint(InRigVMBlueprint);
+			RigVMEdGraph->InitializeFromAsset(InRigVMBlueprint);
 		}
 
 	}
 
-	URigVMBlueprint::sCurrentlyOpenedRigVMBlueprints.AddUnique(InRigVMBlueprint);
+	IRigVMAssetInterface::sCurrentlyOpenedRigVMBlueprints.AddUnique(InRigVMBlueprint);
 
 	InRigVMBlueprint->OnModified().AddSP(StaticCastSharedRef<FRigVMEditorBase>(SharedRef()), &FRigVMEditorBase::HandleModifiedEvent);
 	InRigVMBlueprint->OnVMCompiled().AddSP(StaticCastSharedRef<FRigVMEditorBase>(SharedRef()), &FRigVMEditorBase::HandleVMCompiledEvent);
@@ -292,7 +292,7 @@ void FRigVMEditorBase::InitRigVMEditor(const EToolkitMode::Type Mode, const TSha
 	if (ShouldOpenGraphByDefault() && (Blueprints.Num() > 0))
 	{
 		bool bBroughtGraphToFront = false;
-		for(UEdGraph* Graph : Blueprints[0]->UbergraphPages)
+		for(UEdGraph* Graph : Blueprints[0]->GetUberGraphs())
 		{
 			if (URigVMEdGraph* RigGraph = Cast<URigVMEdGraph>(Graph))
 			{
@@ -319,7 +319,7 @@ void FRigVMEditorBase::InitRigVMEditor(const EToolkitMode::Type Mode, const TSha
 			{
 				// remember all ed graphs which were visible as tabs
 				TArray<UEdGraph*> EdGraphs;
-				InRigVMBlueprint->GetAllGraphs(EdGraphs);
+				InRigVMBlueprint->GetAllEdGraphs(EdGraphs);
 
 				for (UEdGraph* EdGraph : EdGraphs)
 				{
@@ -348,7 +348,7 @@ void FRigVMEditorBase::InitRigVMEditor(const EToolkitMode::Type Mode, const TSha
 					InRigVMBlueprint->GetController(Graph)->ClearNodeSelection(false);
 				}
 
-				if (UPackage* Package = InRigVMBlueprint->GetOutermost())
+				if (UPackage* Package = InRigVMBlueprint->GetObject()->GetOutermost())
 				{
 					Package->SetDirtyFlag(InRigVMBlueprint->IsMarkedDirtyDuringLoad());
 				}
@@ -391,7 +391,7 @@ void FRigVMEditorBase::InitRigVMEditor(const EToolkitMode::Type Mode, const TSha
 
 	if (Blueprints.Num() > 0)
 	{
-		if(Blueprints[0]->Status == BS_Error)
+		if(Blueprints[0]->GetAssetStatus() == RVMA_Error)
 		{
 			Compile();
 		}
@@ -404,7 +404,7 @@ void FRigVMEditorBase::InitRigVMEditor(const EToolkitMode::Type Mode, const TSha
 		{
 			// no need to do anything if the the CR is not opened anymore
 		   // (i.e. destructor has been called before that task actually got a chance to start)
-		   if (!URigVMBlueprint::sCurrentlyOpenedRigVMBlueprints.Contains(InRigVMBlueprint))
+		   if (!IRigVMAssetInterface::sCurrentlyOpenedRigVMBlueprints.Contains(InRigVMBlueprint))
 		   {
 			   return;		
 		   }
@@ -428,7 +428,7 @@ void FRigVMEditorBase::InitRigVMEditor(const EToolkitMode::Type Mode, const TSha
 	UpdateStaleWatchedPins();
 
 #if WITH_EDITOR
-	FString BlueprintName = InRigVMBlueprint->GetPathName();
+	FString BlueprintName = InRigVMBlueprint->GetObject()->GetPathName();
 	RigVMPythonUtils::PrintPythonContext(BlueprintName);
 #endif
 
@@ -498,7 +498,7 @@ void FRigVMEditorBase::InitRigVMEditor(const EToolkitMode::Type Mode, const TSha
 
 void FRigVMEditorBase::HandleAssetRequestedOpen(UObject* InObject)
 {
-	if (InObject == GetRigVMBlueprint())
+	if (InObject == GetRigVMAsset()->GetObject())
 	{
 		bRequestedReopen = true;
 	}
@@ -506,7 +506,7 @@ void FRigVMEditorBase::HandleAssetRequestedOpen(UObject* InObject)
 
 void FRigVMEditorBase::HandleAssetRequestClose(UObject* InObject, EAssetEditorCloseReason InReason)
 {
-	if (InObject == GetRigVMBlueprint())
+	if (InObject == GetRigVMAsset()->GetObject())
 	{
 		bRequestedReopen = false;
 	}
@@ -517,14 +517,19 @@ const FName FRigVMEditorBase::GetEditorModeName() const
 	return FRigVMEditorModes::RigVMEditorMode;
 }
 
-UBlueprint* FRigVMEditorBase::GetBlueprintObj() const
+URigVMBlueprint* FRigVMEditorBase::GetRigVMBlueprint() const
+{
+	return Cast<URigVMBlueprint>(GetRigVMAsset().GetObject());
+}
+
+TRigVMAssetInterfacePtr FRigVMEditorBase::GetRigVMAsset() const
 {
 	const TArray<UObject*>& EditingObjs = GetEditingBlueprints();
 	for (UObject* Obj : EditingObjs)
 	{
-		if (Obj->IsA<URigVMBlueprint>()) 
+		if (Obj->Implements<URigVMAssetInterface>()) 
 		{
-			return (UBlueprint*)Obj;
+			return Obj;
 		}
 	}
 	return nullptr;
@@ -532,9 +537,9 @@ UBlueprint* FRigVMEditorBase::GetBlueprintObj() const
 
 TSubclassOf<UEdGraphSchema> FRigVMEditorBase::GetDefaultSchemaClass() const
 {
-	if(const URigVMBlueprint* RigVMBlueprint = GetRigVMBlueprint())
+	if(const TRigVMAssetInterfacePtr RigVMBlueprint = GetRigVMAsset())
 	{
-		return RigVMBlueprint->GetRigVMEdGraphSchemaClass();
+		return RigVMBlueprint->GetRigVMClientHost()->GetRigVMEdGraphSchemaClass();
 	}
 	return URigVMEdGraphSchema::StaticClass();
 }
@@ -556,7 +561,7 @@ void FRigVMEditorBase::OnGraphEditorFocused(const TSharedRef<SGraphEditor>& InGr
 void FRigVMEditorBase::Tick(float DeltaTime)
 {
 	// tick the  rigvm host
-	if (const URigVMBlueprint* Blueprint = GetRigVMBlueprint())
+	if (const TRigVMAssetInterfacePtr Blueprint = GetRigVMAsset())
 	{
 		if (LastDebuggedHost != GetCustomDebugObjectLabel(Blueprint->GetObjectBeingDebugged()))
 		{
@@ -567,7 +572,7 @@ void FRigVMEditorBase::Tick(float DeltaTime)
 			{
 				if (DebugObject.NameOverride == LastDebuggedHost)
 				{
-					GetBlueprintObj()->SetObjectBeingDebugged(DebugObject.Object);
+					GetRigVMAsset()->SetObjectBeingDebugged(DebugObject.Object);
 					break;
 				}
 			}
@@ -600,7 +605,7 @@ FText FRigVMEditorBase::GetBaseToolkitName() const
 
 FText FRigVMEditorBase::GetToolkitToolTipText() const
 {
-	return FAssetEditorToolkit::GetToolTipTextForObject(GetBlueprintObj());
+	return FAssetEditorToolkit::GetToolTipTextForObject(GetRigVMAsset()->GetObject());
 }
 
 FString FRigVMEditorBase::GetWorldCentricTabPrefix() const
@@ -615,7 +620,7 @@ FLinearColor FRigVMEditorBase::GetWorldCentricTabColorScale() const
 
 void FRigVMEditorBase::InitToolMenuContextImpl(FToolMenuContext& MenuContext)
 {
-	if (URigVMBlueprint* RigVMBlueprint = GetRigVMBlueprint())
+	if (TRigVMAssetInterfacePtr RigVMBlueprint = GetRigVMAsset())
 	{
 		URigVMGraph* Model = nullptr;
 		URigVMNode* Node = nullptr;
@@ -648,7 +653,7 @@ void FRigVMEditorBase::InitToolMenuContextImpl(FToolMenuContext& MenuContext)
 
 bool FRigVMEditorBase::TransactionObjectAffectsBlueprintImpl(UObject* InTransactedObject)
 {
-	URigVMBlueprint* RigVMBlueprint = GetRigVMBlueprint();
+	TRigVMAssetInterfacePtr RigVMBlueprint = GetRigVMAsset();
 	if (RigVMBlueprint == nullptr)
 	{
 		return false;
@@ -712,7 +717,7 @@ void FRigVMEditorBase::DeleteSelectedNodes()
 {
 	DECLARE_SCOPE_HIERARCHICAL_COUNTER_FUNC()
 
-	URigVMBlueprint* RigVMBlueprint = GetRigVMBlueprint();
+	TRigVMAssetInterfacePtr RigVMBlueprint = GetRigVMAsset();
 	if (RigVMBlueprint == nullptr)
 	{
 		return;
@@ -730,7 +735,7 @@ void FRigVMEditorBase::DeleteSelectedNodes()
 		{
 			if (Node->CanUserDeleteNode())
 			{
-				AnalyticsTrackNodeEvent(GetBlueprintObj(), Node, true);
+				AnalyticsTrackNodeEvent(GetRigVMAsset().GetInterface(), Node, true);
 				if (const URigVMEdGraphNode* RigVMEdGraphNode = Cast<URigVMEdGraphNode>(Node))
 				{
 					bRelinkPins = bRelinkPins || FSlateApplication::Get().GetModifierKeys().IsShiftDown();
@@ -860,7 +865,7 @@ void FRigVMEditorBase::AddNewFunctionVariant(const UEdGraph* InOriginalFunction)
 		{
 			if (URigVMFunctionLibrary* FunctionLibrary = Cast<URigVMFunctionLibrary>(RigVMGraph->GetParentGraph()))
 			{
-				URigVMBlueprint* RigVMBlueprint = GetRigVMBlueprint();
+				TRigVMAssetInterfacePtr RigVMBlueprint = GetRigVMAsset();
 				URigVMController* Controller = RigVMBlueprint->GetController(FunctionLibrary);
 				if (const URigVMLibraryNode* VariantNode = Controller->CreateFunctionVariant(RigVMGraph->GetOuter()->GetFName(), NAME_None, true, true))
 				{
@@ -893,13 +898,13 @@ void FRigVMEditorBase::PostTransaction(bool bSuccess, const FTransaction* Transa
 		return;
 	}
 
-	const URigVMBlueprint* Blueprint = GetRigVMBlueprint();
+	const TRigVMAssetInterfacePtr Blueprint = GetRigVMAsset();
 	if(Blueprint == nullptr)
 	{
 			return;
 	}
 
-	const UPackage* Package = Blueprint->GetPackage();
+	const UPackage* Package = Blueprint->GetObject()->GetPackage();
 
 	TArray<UObject*> TransactedObjects;
 	Transaction->GetTransactionObjects(TransactedObjects);
@@ -998,7 +1003,7 @@ void FRigVMEditorBase::PasteNodes()
 
 	URigVMController* FocusedController = GetFocusedController();
 
-	const bool bPastePerformed = UE::RigVM::Editor::Tools::PasteNodes(PasteLocation, TextToImport, FocusedController, GetFocusedModel(), GetRigVMBlueprint()->GetLocalFunctionLibrary(), GetRigVMBlueprint()->GetRigVMGraphFunctionHost());
+	const bool bPastePerformed = UE::RigVM::Editor::Tools::PasteNodes(PasteLocation, TextToImport, FocusedController, GetFocusedModel(), GetRigVMAsset()->GetLocalFunctionLibrary(), GetRigVMAsset()->GetRigVMClientHost()->GetRigVMGraphFunctionHost().GetInterface());
 	if (bPastePerformed)
 	{
 		FocusedController->CloseUndoBracket();
@@ -1010,18 +1015,13 @@ void FRigVMEditorBase::PasteNodes()
 
 }
 
-URigVMBlueprint* FRigVMEditorBase::GetRigVMBlueprint() const
-{
-	return Cast<URigVMBlueprint>(GetBlueprintObj());
-}
-
 URigVMHost* FRigVMEditorBase::GetRigVMHost() const
 {
-	if (URigVMBlueprint* RigVMBlueprint = GetRigVMBlueprint())
+	if (TRigVMAssetInterfacePtr RigVMBlueprint = GetRigVMAsset())
 	{
-		if(RigVMBlueprint->EditorHost && IsValid(RigVMBlueprint->EditorHost))
+		if(RigVMBlueprint->GetEditorHost() && IsValid(RigVMBlueprint->GetEditorHost()))
 		{
-			return RigVMBlueprint->EditorHost;
+			return RigVMBlueprint->GetEditorHost();
 		}
 	}
 	return nullptr;
@@ -1029,7 +1029,7 @@ URigVMHost* FRigVMEditorBase::GetRigVMHost() const
 
 UObject* FRigVMEditorBase::GetOuterForHost() const
 {
-	return GetRigVMBlueprint();
+	return GetRigVMAsset()->GetObject();
 }
 
 UClass* FRigVMEditorBase::GetDetailWrapperClass() const
@@ -1048,7 +1048,7 @@ bool FRigVMEditorBase::SelectLocalVariable(const UEdGraph* Graph, const FName& V
 				if (Variable.Name == VariableName)
 				{
 					URigVMDetailsViewWrapperObject* WrapperObject = URigVMDetailsViewWrapperObject::MakeInstance(
-						GetDetailWrapperClass(), GetBlueprintObj(), Variable.StaticStruct(), (uint8*)&Variable, RigVMGraph);
+						GetDetailWrapperClass(), GetRigVMAsset()->GetObject(), Variable.StaticStruct(), (uint8*)&Variable, RigVMGraph);
 					WrapperObject->GetWrappedPropertyChangedChainEvent().AddSP(StaticCastSharedRef<FRigVMEditorBase>(SharedRef()), &FRigVMEditorBase::OnWrappedPropertyChangedChainEvent);
 					WrapperObject->AddToRoot();
 
@@ -1064,7 +1064,7 @@ bool FRigVMEditorBase::SelectLocalVariable(const UEdGraph* Graph, const FName& V
 
 void FRigVMEditorBase::CreateDefaultCommandsImpl()
 {
-	if (GetBlueprintObj())
+	if (GetRigVMAsset())
 	{
 		// FBlueprintEditor::CreateDefaultCommands();
 	}
@@ -1086,14 +1086,14 @@ void FRigVMEditorBase::Compile()
 	{
 		DECLARE_SCOPE_HIERARCHICAL_COUNTER_FUNC()
 
-		URigVMBlueprint* RigVMBlueprint = GetRigVMBlueprint();
+		TRigVMAssetInterfacePtr RigVMBlueprint = GetRigVMAsset();
 		if (RigVMBlueprint == nullptr)
 		{
 			return;
 		}
 
 		// force to disable the supended notif brackets
-		RigVMBlueprint->GetRigVMClient()->bSuspendModelNotificationsForOthers = false;
+		RigVMBlueprint->GetRigVMClientHost()->GetRigVMClient()->bSuspendModelNotificationsForOthers = false;
 		RigVMBlueprint->bSuspendModelNotificationsForSelf = false;
 
 		RigVMBlueprint->GetCompileLog().Messages.Reset();
@@ -1123,11 +1123,9 @@ void FRigVMEditorBase::Compile()
 			RigVMLog.Reset();
 			RigVMHost->SetLog(&RigVMLog);
 
-			URigVMBlueprintGeneratedClass* GeneratedClass = Cast<URigVMBlueprintGeneratedClass>(RigVMHost->GetClass());
-			if (GeneratedClass)
+			if (URigVM* VM = RigVMHost->GetVM())
 			{
-				URigVMHost* CDO = Cast<URigVMHost>(GeneratedClass->GetDefaultObject(true /* create if needed */));
-				FRigVMInstructionArray Instructions = CDO->GetVM()->GetInstructions();
+				FRigVMInstructionArray Instructions = VM->GetInstructions();
 
 				if (Instructions.Num() <= 1) // just the "done" operator
 				{
@@ -1154,7 +1152,7 @@ void FRigVMEditorBase::Compile()
 
 		// invalidate all node titles
 		TArray<UEdGraph*> EdGraphs;
-		RigVMBlueprint->GetAllGraphs(EdGraphs);
+		RigVMBlueprint->GetAllEdGraphs(EdGraphs);
 		for (UEdGraph* EdGraph : EdGraphs)
 		{
 			URigVMEdGraph* RigVMEdGraph = Cast<URigVMEdGraph>(EdGraph);
@@ -1174,9 +1172,9 @@ void FRigVMEditorBase::Compile()
 
 		// store the defaults from the CDO back on the new variables list
 		bool bAnyVariableValueChanged = false;
-		for(FBPVariableDescription& NewVariable : RigVMBlueprint->NewVariables)
+		for(FRigVMGraphVariableDescription& NewVariable : RigVMBlueprint->GetAssetVariables())
 		{
-			bAnyVariableValueChanged |= UpdateDefaultValueForVariable(NewVariable, true);
+			bAnyVariableValueChanged |= UpdateDefaultValueForVariable(NewVariable, false);
 		}
 		if (bAnyVariableValueChanged)
 		{
@@ -1201,7 +1199,7 @@ void FRigVMEditorBase::Compile()
 
 void FRigVMEditorBase::SaveAsset_Execute()
 {
-	LastDebuggedHost = GetCustomDebugObjectLabel(GetBlueprintObj()->GetObjectBeingDebugged());
+	LastDebuggedHost = GetCustomDebugObjectLabel(GetRigVMAsset()->GetObjectBeingDebugged());
 	SaveAsset_Execute_Impl();
 
 	UpdateRigVMHost();
@@ -1209,7 +1207,7 @@ void FRigVMEditorBase::SaveAsset_Execute()
 
 void FRigVMEditorBase::SaveAssetAs_Execute()
 {
-	LastDebuggedHost = GetCustomDebugObjectLabel(GetBlueprintObj()->GetObjectBeingDebugged());
+	LastDebuggedHost = GetCustomDebugObjectLabel(GetRigVMAsset()->GetObjectBeingDebugged());
 	SaveAssetAs_Execute_Impl();
 
 	UpdateRigVMHost();
@@ -1222,7 +1220,7 @@ bool FRigVMEditorBase::IsEditable(UEdGraph* InGraph) const
 		return false;
 	}
 	
-	if(URigVMBlueprint* RigVMBlueprint = GetRigVMBlueprint())
+	if(TRigVMAssetInterfacePtr RigVMBlueprint = GetRigVMAsset())
 	{
 		// aggregate graphs are always read only
 		if(const URigVMGraph* Model = RigVMBlueprint->GetModel(InGraph))
@@ -1273,7 +1271,7 @@ void FRigVMEditorBase::OnSelectedNodesChangedImpl(const TSet<UObject*>& NewSelec
 
 	TGuardValue<bool> SelectGuard(RigVMEdGraph->bIsSelecting, true);
 
-	URigVMBlueprint* RigVMBlueprint = GetRigVMBlueprint();
+	TRigVMAssetInterfacePtr RigVMBlueprint = GetRigVMAsset();
 	if (RigVMBlueprint)
 	{
 		TArray<FName> NodeNamesToSelect;
@@ -1292,7 +1290,7 @@ void FRigVMEditorBase::OnSelectedNodesChangedImpl(const TSet<UObject*>& NewSelec
 	}
 }
 
-void FRigVMEditorBase::OnBlueprintChangedImpl(UBlueprint* InBlueprint, bool bIsJustBeingCompiled)
+void FRigVMEditorBase::OnBlueprintChangedImpl(IRigVMAssetInterface* InBlueprint, bool bIsJustBeingCompiled)
 {
 	DECLARE_SCOPE_HIERARCHICAL_COUNTER_FUNC()
 
@@ -1303,7 +1301,7 @@ void FRigVMEditorBase::OnBlueprintChangedImpl(UBlueprint* InBlueprint, bool bIsJ
 
 	OnBlueprintChangedInnerImpl(InBlueprint, bIsJustBeingCompiled);
 
-	if(InBlueprint == GetBlueprintObj())
+	if(InBlueprint == GetRigVMAsset().GetInterface())
 	{
 		if(bIsJustBeingCompiled)
 		{
@@ -1318,7 +1316,7 @@ void FRigVMEditorBase::OnBlueprintChangedImpl(UBlueprint* InBlueprint, bool bIsJ
 				{
 					if (DebugObject.NameOverride == LastDebuggedHost)
 					{
-						GetBlueprintObj()->SetObjectBeingDebugged(DebugObject.Object);
+						GetRigVMAsset()->SetObjectBeingDebugged(DebugObject.Object);
 						LastDebuggedHost.Empty();
 						break;
 					}
@@ -1363,7 +1361,7 @@ void FRigVMEditorBase::OnNodeTitleCommitted(const FText& NewText, ETextCommit::T
 
 	if (const UEdGraphNode_Comment* CommentBeingChanged = Cast<UEdGraphNode_Comment>(NodeBeingChanged))
 	{
-		if (GetRigVMBlueprint())
+		if (GetRigVMAsset())
 		{
 			GetFocusedController()->SetCommentTextByName(CommentBeingChanged->GetFName(), NewText.ToString(), CommentBeingChanged->FontSize, CommentBeingChanged->bCommentBubbleVisible, CommentBeingChanged->bColorCommentBubble, true, true);
 		}
@@ -1402,92 +1400,94 @@ void FRigVMEditorBase::BindCommands()
 	DECLARE_SCOPE_HIERARCHICAL_COUNTER_FUNC()
 
 	TSharedRef<FRigVMEditorBase> ThisRef = StaticCastSharedRef<FRigVMEditorBase>(SharedRef());
+	TSharedRef<FUICommandList> ToolkitCommands = GetToolkitCommands();
+	FRigVMEditorCommands& EditorCommands = FRigVMEditorCommands::Get();
 
-	GetToolkitCommands()->MapAction(
-		FRigVMEditorCommands::Get().AutoCompileGraph,
+	ToolkitCommands->MapAction(
+		EditorCommands.AutoCompileGraph,
 		FExecuteAction::CreateSP(ThisRef, &FRigVMEditorBase::ToggleAutoCompileGraph), 
 		FIsActionChecked::CreateSP(ThisRef, &FRigVMEditorBase::CanAutoCompileGraph),
 		FIsActionChecked::CreateSP(ThisRef, &FRigVMEditorBase::IsAutoCompileGraphOn));
 
-	GetToolkitCommands()->MapAction(
-		FRigVMEditorCommands::Get().ToggleEventQueue,
+	ToolkitCommands->MapAction(
+		EditorCommands.ToggleEventQueue,
 		FExecuteAction::CreateSP(ThisRef, &FRigVMEditorBase::ToggleEventQueue),
 		FCanExecuteAction());
 
-	GetToolkitCommands()->MapAction(
-		FRigVMEditorCommands::Get().ToggleExecutionMode,
+	ToolkitCommands->MapAction(
+		EditorCommands.ToggleExecutionMode,
 		FExecuteAction::CreateSP(ThisRef, &FRigVMEditorBase::ToggleExecutionMode),
 		FCanExecuteAction());
 
-	GetToolkitCommands()->MapAction(
-		FRigVMEditorCommands::Get().ReleaseMode,
+	ToolkitCommands->MapAction(
+		EditorCommands.ReleaseMode,
 		FExecuteAction::CreateSP(ThisRef, &FRigVMEditorBase::SetExecutionMode, ERigVMEditorExecutionModeType_Release),
 		FCanExecuteAction());
 
-	GetToolkitCommands()->MapAction(
-		FRigVMEditorCommands::Get().DebugMode,
+	ToolkitCommands->MapAction(
+		EditorCommands.DebugMode,
 		FExecuteAction::CreateSP(ThisRef, &FRigVMEditorBase::SetExecutionMode, ERigVMEditorExecutionModeType_Debug),
 		FCanExecuteAction());
 
-	GetToolkitCommands()->MapAction(
-		FRigVMEditorCommands::Get().ResumeExecution,
+	ToolkitCommands->MapAction(
+		EditorCommands.ResumeExecution,
 		FExecuteAction::CreateSP(ThisRef, &FRigVMEditorBase::HandleBreakpointActionRequested, ERigVMBreakpointAction::Resume),			
 		FIsActionChecked::CreateSP(ThisRef, &FRigVMEditorBase::IsHaltedAtBreakpoint),
 		FIsActionChecked(),
 		FIsActionButtonVisible::CreateSP(ThisRef, &FRigVMEditorBase::IsHaltedAtBreakpoint));
 
-	GetToolkitCommands()->MapAction(
-		FRigVMEditorCommands::Get().ShowCurrentStatement,
+	ToolkitCommands->MapAction(
+		EditorCommands.ShowCurrentStatement,
 		FExecuteAction::CreateSP(ThisRef, &FRigVMEditorBase::HandleShowCurrentStatement),
 		FIsActionChecked::CreateSP(ThisRef, &FRigVMEditorBase::IsHaltedAtBreakpoint),
 		FIsActionChecked(),
 		FIsActionButtonVisible::CreateSP(ThisRef, &FRigVMEditorBase::IsHaltedAtBreakpoint));
 
-	GetToolkitCommands()->MapAction(
-		FRigVMEditorCommands::Get().StepOver,
+	ToolkitCommands->MapAction(
+		EditorCommands.StepOver,
 		FExecuteAction::CreateSP(ThisRef, &FRigVMEditorBase::HandleBreakpointActionRequested, ERigVMBreakpointAction::StepOver),
 		FIsActionChecked::CreateSP(ThisRef, &FRigVMEditorBase::IsHaltedAtBreakpoint),
 		FIsActionChecked(),
 		FIsActionButtonVisible::CreateSP(ThisRef, &FRigVMEditorBase::IsHaltedAtBreakpoint));
 
-	GetToolkitCommands()->MapAction(
-		FRigVMEditorCommands::Get().StepInto,
+	ToolkitCommands->MapAction(
+		EditorCommands.StepInto,
 		FExecuteAction::CreateSP(ThisRef, &FRigVMEditorBase::HandleBreakpointActionRequested, ERigVMBreakpointAction::StepInto),
 		FIsActionChecked::CreateSP(ThisRef, &FRigVMEditorBase::IsHaltedAtBreakpoint),
 		FIsActionChecked(),
 		FIsActionButtonVisible::CreateSP(ThisRef, &FRigVMEditorBase::IsHaltedAtBreakpoint));
 
-	GetToolkitCommands()->MapAction(
-		FRigVMEditorCommands::Get().StepOut,
+	ToolkitCommands->MapAction(
+		EditorCommands.StepOut,
 		FExecuteAction::CreateSP(ThisRef, &FRigVMEditorBase::HandleBreakpointActionRequested, ERigVMBreakpointAction::StepOut),
 		FIsActionChecked::CreateSP(ThisRef, &FRigVMEditorBase::IsHaltedAtBreakpoint),
 		FIsActionChecked(),
 		FIsActionButtonVisible::CreateSP(ThisRef, &FRigVMEditorBase::IsHaltedAtBreakpoint));
 
-	GetToolkitCommands()->MapAction(
-		FRigVMEditorCommands::Get().FrameSelection,
+	ToolkitCommands->MapAction(
+		EditorCommands.FrameSelection,
 		FExecuteAction::CreateSP(ThisRef, &FRigVMEditorBase::FrameSelection),
 		FCanExecuteAction());
 
-	GetToolkitCommands()->MapAction(
-		FRigVMEditorCommands::Get().SwapFunctionWithinAsset,
+	ToolkitCommands->MapAction(
+		EditorCommands.SwapFunctionWithinAsset,
 		FExecuteAction::CreateSP(ThisRef, &FRigVMEditorBase::SwapFunctionWithinAsset),
 		FCanExecuteAction());
 
-	GetToolkitCommands()->MapAction(
-		FRigVMEditorCommands::Get().SwapFunctionAcrossProject,
+	ToolkitCommands->MapAction(
+		EditorCommands.SwapFunctionAcrossProject,
 		FExecuteAction::CreateSP(ThisRef, &FRigVMEditorBase::SwapFunctionAcrossProject),
 		FCanExecuteAction());
 
-	GetToolkitCommands()->MapAction(
-		FRigVMEditorCommands::Get().SwapAssetReferences,
+	ToolkitCommands->MapAction(
+		EditorCommands.SwapAssetReferences,
 		FExecuteAction::CreateSP(ThisRef, &FRigVMEditorBase::SwapAssetReferences),
 		FCanExecuteAction());
 }
 
 void FRigVMEditorBase::ToggleAutoCompileGraph()
 {
-	if (URigVMBlueprint* RigVMBlueprint = GetRigVMBlueprint())
+	if (TRigVMAssetInterfacePtr RigVMBlueprint = GetRigVMAsset())
 	{
 		RigVMBlueprint->SetAutoVMRecompile(!RigVMBlueprint->GetAutoVMRecompile());
 		if (RigVMBlueprint->GetAutoVMRecompile())
@@ -1499,7 +1499,7 @@ void FRigVMEditorBase::ToggleAutoCompileGraph()
 
 bool FRigVMEditorBase::IsAutoCompileGraphOn() const
 {
-	if (URigVMBlueprint* RigVMBlueprint = GetRigVMBlueprint())
+	if (TRigVMAssetInterfacePtr RigVMBlueprint = GetRigVMAsset())
 	{
 		return RigVMBlueprint->GetAutoVMRecompile();
 	}
@@ -1595,7 +1595,7 @@ void FRigVMEditorBase::OnNewDocumentClicked(ECreatedDocumentType GraphType)
 {
 	if (GraphType == FRigVMNewEditor::CGT_NewFunctionGraph)
 	{
-		if (URigVMBlueprint* RigVMBlueprint = GetRigVMBlueprint())
+		if (TRigVMAssetInterfacePtr RigVMBlueprint = GetRigVMAsset())
 		{
 			if (URigVMController* Controller = RigVMBlueprint->GetOrCreateController(RigVMBlueprint->GetLocalFunctionLibrary()))
 			{
@@ -1612,9 +1612,9 @@ void FRigVMEditorBase::OnNewDocumentClicked(ECreatedDocumentType GraphType)
 	}
 	else if(GraphType == FRigVMNewEditor::CGT_NewEventGraph)
 	{
-		if (URigVMBlueprint* RigVMBlueprint = GetRigVMBlueprint())
+		if (TRigVMAssetInterfacePtr RigVMBlueprint = GetRigVMAsset())
 		{
-			const UClass* EdGraphSchemaClass = RigVMBlueprint->GetRigVMEdGraphSchemaClass();
+			const UClass* EdGraphSchemaClass = RigVMBlueprint->GetRigVMClientHost()->GetRigVMEdGraphSchemaClass();
 			const URigVMEdGraphSchema* SchemaCDO = CastChecked<URigVMEdGraphSchema>(EdGraphSchemaClass->GetDefaultObject());
 
 			if(URigVMGraph* Model = RigVMBlueprint->AddModel(SchemaCDO->GetRootGraphName().ToString()))
@@ -1658,33 +1658,6 @@ bool FRigVMEditorBase::IsSectionVisibleImpl(RigVMNodeSectionID::Type InSectionID
 	return false;
 }
 
-bool FRigVMEditorBase::AreEventGraphsAllowed() const
-{
-	if(const URigVMBlueprint* RigVMBlueprint = GetRigVMBlueprint())
-	{
-		return RigVMBlueprint->SupportsEventGraphs();
-	}
-	return false;
-}
-
-bool FRigVMEditorBase::AreMacrosAllowed() const
-{
-	if(const URigVMBlueprint* RigVMBlueprint = GetRigVMBlueprint())
-	{
-		return RigVMBlueprint->SupportsMacros();
-	}
-	return false;
-}
-
-bool FRigVMEditorBase::AreDelegatesAllowed() const
-{
-	if(const URigVMBlueprint* RigVMBlueprint = GetRigVMBlueprint())
-	{
-		return RigVMBlueprint->SupportsDelegates();
-	}
-	return false;
-}
-
 bool FRigVMEditorBase::NewDocument_IsVisibleForTypeImpl(ECreatedDocumentType GraphType) const
 {
 	switch(GraphType)
@@ -1706,7 +1679,7 @@ FGraphAppearanceInfo FRigVMEditorBase::GetGraphAppearance(UEdGraph* InGraph) con
 {
 	FGraphAppearanceInfo AppearanceInfo = GetGraphAppearanceImpl(InGraph);
 
-	if (const URigVMBlueprint* RigVMBlueprint = GetRigVMBlueprint())
+	if (const TRigVMAssetInterfacePtr RigVMBlueprint = GetRigVMAsset())
 	{
 		AppearanceInfo.CornerText = LOCTEXT("AppearanceCornerText_RigVMEditor", "RigVM");
 
@@ -1726,7 +1699,7 @@ FGraphAppearanceInfo FRigVMEditorBase::GetGraphAppearance(UEdGraph* InGraph) con
 			if(RigVMHost->VMRuntimeSettings.bEnableProfiling)
 			{
 				static constexpr TCHAR Format[] = TEXT("Total %.02f µs");
-				AppearanceInfo.WarningText = FText::FromString(FString::Printf(Format, (float)RigVMBlueprint->RigGraphDisplaySettings.TotalMicroSeconds));
+				AppearanceInfo.WarningText = FText::FromString(FString::Printf(Format, (float)RigVMBlueprint->GetRigGraphDisplaySettings().TotalMicroSeconds));
 			}
 		}
 	}
@@ -1738,7 +1711,7 @@ void FRigVMEditorBase::HandleModifiedEvent(ERigVMGraphNotifType InNotifType, URi
 {
 	DECLARE_SCOPE_HIERARCHICAL_COUNTER_FUNC()
 
-	URigVMBlueprint* RigVMBlueprint = GetRigVMBlueprint();
+	TRigVMAssetInterfacePtr RigVMBlueprint = GetRigVMAsset();
 	if (RigVMBlueprint == nullptr)
 	{
 		return;
@@ -1827,7 +1800,7 @@ void FRigVMEditorBase::HandleModifiedEvent(ERigVMGraphNotifType InNotifType, URi
 				{
 					if(UnitNode->IsEvent())
 					{
-						FBlueprintEditorUtils::MarkBlueprintAsStructurallyModified(GetBlueprintObj());
+						GetRigVMAsset()->MarkAssetAsStructurallyModified();
 						CacheNameLists();
 					}
 				}
@@ -1873,7 +1846,7 @@ void FRigVMEditorBase::HandleModifiedEvent(ERigVMGraphNotifType InNotifType, URi
 			{
 				if(UnitNode->IsEvent())
 				{
-					FBlueprintEditorUtils::MarkBlueprintAsStructurallyModified(GetBlueprintObj());
+					GetRigVMAsset()->MarkAssetAsStructurallyModified();
 					CacheNameLists();
 				}
 			}
@@ -1900,8 +1873,9 @@ public:
 
 void FRigVMEditorBase::HandleVMCompiledEvent(UObject* InCompiledObject, URigVM* InVM, FRigVMExtendedExecuteContext& InContext)
 {
-	if(URigVMBlueprint* RigVMBlueprint = Cast<URigVMBlueprint>(InCompiledObject))
+	if (InCompiledObject->Implements<URigVMAssetInterface>())
 	{
+		TRigVMAssetInterfacePtr RigVMBlueprint = InCompiledObject;
 		GetCompilerResultsListing()->ClearMessages();
 		GetCompilerResultsListing()->AddMessages(RigVMBlueprint->GetCompileLog().Messages);
 		RigVMBlueprint->GetCompileLog().Messages.Reset();
@@ -1944,7 +1918,7 @@ void FRigVMEditorBase::HandleVMCompiledEvent(UObject* InCompiledObject, URigVM*
 
 void FRigVMEditorBase::HandleVMExecutedEvent(URigVMHost* InHost, const FName& InEventName)
 {
-	if (URigVMBlueprint* RigVMBlueprint = GetRigVMBlueprint())
+	if (TRigVMAssetInterfacePtr RigVMBlueprint = GetRigVMAsset())
 	{
 		URigVMHost* DebuggedHost = Cast<URigVMHost>(RigVMBlueprint->GetObjectBeingDebugged());
 		if (DebuggedHost == nullptr)
@@ -1952,7 +1926,7 @@ void FRigVMEditorBase::HandleVMExecutedEvent(URigVMHost* InHost, const FName& In
 			DebuggedHost = GetRigVMHost();
 		}
 
-		if(RigVMBlueprint->RigGraphDisplaySettings.NodeRunLimit > 1)
+		if(RigVMBlueprint->GetRigGraphDisplaySettings().NodeRunLimit > 1)
 		{
 			if(DebuggedHost)
 			{
@@ -1964,7 +1938,7 @@ void FRigVMEditorBase::HandleVMExecutedEvent(URigVMHost* InHost, const FName& In
 					for(int32 InstructionIndex = 0; InstructionIndex < ByteCode.GetNumInstructions(); InstructionIndex++)
 					{
 						const int32 Count = VM->GetInstructionVisitedCount(DebuggedHost->GetRigVMExtendedExecuteContext(), InstructionIndex);
-						if(Count > RigVMBlueprint->RigGraphDisplaySettings.NodeRunLimit)
+						if(Count > RigVMBlueprint->GetRigGraphDisplaySettings().NodeRunLimit)
 						{
 							bFoundLimitWarnings = true;
 
@@ -1974,7 +1948,7 @@ void FRigVMEditorBase::HandleVMExecutedEvent(URigVMHost* InHost, const FName& In
 								const FString Message = FString::Printf(
                                     TEXT("Instruction has hit the NodeRunLimit\n(ran %d times, limit is %d)\n\nYou can increase the limit in the class settings."),
                                     Count,
-                                    RigVMBlueprint->RigGraphDisplaySettings.NodeRunLimit
+                                    RigVMBlueprint->GetRigGraphDisplaySettings().NodeRunLimit
                                 );
 
 								if(DebuggedHost->GetLog())
@@ -2021,33 +1995,33 @@ void FRigVMEditorBase::HandleVMExecutedEvent(URigVMHost* InHost, const FName& In
 			}
 		}
 
-		if(RigVMBlueprint->VMRuntimeSettings.bEnableProfiling)
+		if(RigVMBlueprint->GetVMRuntimeSettings().bEnableProfiling)
 		{
 			if(DebuggedHost)
 			{
-				RigVMBlueprint->RigGraphDisplaySettings.SetTotalMicroSeconds(DebuggedHost->GetProfilingInfo().GetLastExecutionMicroSeconds());
+				RigVMBlueprint->GetRigGraphDisplaySettings().SetTotalMicroSeconds(DebuggedHost->GetProfilingInfo().GetLastExecutionMicroSeconds());
 			}
 
-			if(RigVMBlueprint->RigGraphDisplaySettings.bAutoDetermineRange)
+			if(RigVMBlueprint->GetRigGraphDisplaySettings().bAutoDetermineRange)
 			{
-				if(RigVMBlueprint->RigGraphDisplaySettings.LastMaxMicroSeconds < 0.0)
+				if(RigVMBlueprint->GetRigGraphDisplaySettings().LastMaxMicroSeconds < 0.0)
 				{
-					RigVMBlueprint->RigGraphDisplaySettings.SetLastMinMicroSeconds(RigVMBlueprint->RigGraphDisplaySettings.MinMicroSeconds); 
-					RigVMBlueprint->RigGraphDisplaySettings.SetLastMaxMicroSeconds(RigVMBlueprint->RigGraphDisplaySettings.MaxMicroSeconds);
+					RigVMBlueprint->GetRigGraphDisplaySettings().SetLastMinMicroSeconds(RigVMBlueprint->GetRigGraphDisplaySettings().MinMicroSeconds); 
+					RigVMBlueprint->GetRigGraphDisplaySettings().SetLastMaxMicroSeconds(RigVMBlueprint->GetRigGraphDisplaySettings().MaxMicroSeconds);
 				}
-				else if(RigVMBlueprint->RigGraphDisplaySettings.MaxMicroSeconds >= 0.0)
+				else if(RigVMBlueprint->GetRigGraphDisplaySettings().MaxMicroSeconds >= 0.0)
 				{
-					RigVMBlueprint->RigGraphDisplaySettings.SetLastMinMicroSeconds(RigVMBlueprint->RigGraphDisplaySettings.MinMicroSeconds); 
-					RigVMBlueprint->RigGraphDisplaySettings.SetLastMaxMicroSeconds(RigVMBlueprint->RigGraphDisplaySettings.MaxMicroSeconds); 
+					RigVMBlueprint->GetRigGraphDisplaySettings().SetLastMinMicroSeconds(RigVMBlueprint->GetRigGraphDisplaySettings().MinMicroSeconds); 
+					RigVMBlueprint->GetRigGraphDisplaySettings().SetLastMaxMicroSeconds(RigVMBlueprint->GetRigGraphDisplaySettings().MaxMicroSeconds); 
 				}
 
-				RigVMBlueprint->RigGraphDisplaySettings.MinMicroSeconds = DBL_MAX; 
-				RigVMBlueprint->RigGraphDisplaySettings.MaxMicroSeconds = (double)INDEX_NONE;
+				RigVMBlueprint->GetRigGraphDisplaySettings().MinMicroSeconds = DBL_MAX; 
+				RigVMBlueprint->GetRigGraphDisplaySettings().MaxMicroSeconds = (double)INDEX_NONE;
 			}
 			else
 			{
-				RigVMBlueprint->RigGraphDisplaySettings.SetLastMinMicroSeconds(RigVMBlueprint->RigGraphDisplaySettings.MinMicroSeconds); 
-				RigVMBlueprint->RigGraphDisplaySettings.SetLastMaxMicroSeconds(RigVMBlueprint->RigGraphDisplaySettings.MaxMicroSeconds);
+				RigVMBlueprint->GetRigGraphDisplaySettings().SetLastMinMicroSeconds(RigVMBlueprint->GetRigGraphDisplaySettings().MinMicroSeconds); 
+				RigVMBlueprint->GetRigGraphDisplaySettings().SetLastMaxMicroSeconds(RigVMBlueprint->GetRigGraphDisplaySettings().MaxMicroSeconds);
 			}
 		}
 	}
@@ -2066,7 +2040,7 @@ void FRigVMEditorBase::HandleVMExecutionHalted(const int32 InstructionIndex, UOb
 	{
 		SetHaltedNode(InNode);
 		
-		if (URigVMBlueprint* RigVMBlueprint = GetRigVMBlueprint())
+		if (TRigVMAssetInterfacePtr RigVMBlueprint = GetRigVMAsset())
 		{
 			if (RigVMBlueprint->GetAllModels().Contains(InNode->GetGraph()))
 			{
@@ -2116,9 +2090,9 @@ void FRigVMEditorBase::NotifyPreChange(FProperty* PropertyAboutToChange)
 {
 	NotifyPreChangeImpl(PropertyAboutToChange);
 
-	if (URigVMBlueprint* RigVMBlueprint = GetRigVMBlueprint())
+	if (TRigVMAssetInterfacePtr RigVMBlueprint = GetRigVMAsset())
 	{
-		RigVMBlueprint->Modify();
+		RigVMBlueprint->GetObject()->Modify();
 	}
 }
 
@@ -2130,21 +2104,21 @@ void FRigVMEditorBase::NotifyPostChange(const FPropertyChangedEvent& PropertyCha
 	// OnFinishedChangingProperties is called only for top level property changes.
 	// changes on a lower level property like transform under a user defined struct
 	// only go through this.
-	URigVMBlueprint* RigVMBlueprint = GetRigVMBlueprint();
+	TRigVMAssetInterfacePtr RigVMBlueprint = GetRigVMAsset();
 	if(GetRigVMHost() && RigVMBlueprint)
 	{
-		bool bUseCDO = false; 
+		bool bUseDebuggedObject = true; 
 		if(PropertyChangedEvent.GetNumObjectsBeingEdited() == 1)
 		{
-			bUseCDO = PropertyChangedEvent.GetObjectBeingEdited(0)->HasAnyFlags(RF_ClassDefaultObject);
+			bUseDebuggedObject = !PropertyChangedEvent.GetObjectBeingEdited(0)->HasAnyFlags(RF_ClassDefaultObject);
 		}
 		
 		const FName VarName = PropertyChangedEvent.MemberProperty->GetFName();
-		for(FBPVariableDescription& NewVariable : RigVMBlueprint->NewVariables)
+		for(FRigVMGraphVariableDescription& NewVariable : RigVMBlueprint->GetAssetVariables())
 		{
-			if(NewVariable.VarName == VarName)
+			if(NewVariable.Name == VarName)
 			{
-				UpdateDefaultValueForVariable(NewVariable, bUseCDO);
+				UpdateDefaultValueForVariable(NewVariable, bUseDebuggedObject);
 				break;
 			}
 		}
@@ -2155,18 +2129,18 @@ void FRigVMEditorBase::OnFinishedChangingProperties(const FPropertyChangedEvent&
 {
 	DECLARE_SCOPE_HIERARCHICAL_COUNTER_FUNC()
 
-	URigVMBlueprint* RigVMBlueprint = GetRigVMBlueprint();
+	TRigVMAssetInterfacePtr RigVMBlueprint = GetRigVMAsset();
 
 	if (RigVMBlueprint)
 	{
-		if (PropertyChangedEvent.MemberProperty->GetNameCPP() == GET_MEMBER_NAME_STRING_CHECKED(URigVMBlueprint, VMCompileSettings))
+		if (PropertyChangedEvent.MemberProperty->GetNameCPP() == TEXT("VMCompileSettings"))
 		{
 			RigVMBlueprint->RecompileVM();
 		}
 
-		else if (PropertyChangedEvent.MemberProperty->GetNameCPP() == GET_MEMBER_NAME_STRING_CHECKED(URigVMBlueprint, VMRuntimeSettings))
+		else if (PropertyChangedEvent.MemberProperty->GetNameCPP() == TEXT("VMRuntimeSettings"))
 		{
-			RigVMBlueprint->VMRuntimeSettings.Validate();
+			RigVMBlueprint->GetVMRuntimeSettings().Validate();
 			RigVMBlueprint->PropagateRuntimeSettingsFromBPToInstances();
 		}
 	}
@@ -2174,13 +2148,13 @@ void FRigVMEditorBase::OnFinishedChangingProperties(const FPropertyChangedEvent&
 
 void FRigVMEditorBase::OnPropertyChanged(UObject* InObject, FPropertyChangedEvent& InEvent)
 {
-	URigVMBlueprint* RigVMBlueprint = GetRigVMBlueprint();
+	TRigVMAssetInterfacePtr RigVMBlueprint = GetRigVMAsset();
 
-	if (RigVMBlueprint && InObject == RigVMBlueprint)
+	if (RigVMBlueprint && InObject == RigVMBlueprint->GetObject())
 	{
 		// if the models have changed - we may need to close a document
-		if(InEvent.MemberProperty == RigVMBlueprint->GetClass()->FindPropertyByName(GET_MEMBER_NAME_CHECKED(URigVMBlueprint, RigVMClient)) ||
-			InEvent.MemberProperty == RigVMBlueprint->GetClass()->FindPropertyByName(GET_MEMBER_NAME_CHECKED(URigVMBlueprint, UbergraphPages)))
+		if(InEvent.MemberProperty == RigVMBlueprint->GetObject()->GetClass()->FindPropertyByName(TEXT("RigVMClient")) ||
+			InEvent.MemberProperty == RigVMBlueprint->GetObject()->GetClass()->FindPropertyByName(TEXT("UbergraphPages")))
 		{
 			GetDocumentManager()->CleanInvalidTabs();
 		}
@@ -2195,7 +2169,7 @@ void FRigVMEditorBase::OnWrappedPropertyChangedChainEvent(URigVMDetailsViewWrapp
 
 	TGuardValue<bool> SuspendDetailsPanelRefresh(bSuspendDetailsPanelRefresh, true);
 
-	URigVMBlueprint* RigVMBlueprint = GetRigVMBlueprint();
+	TRigVMAssetInterfacePtr RigVMBlueprint = GetRigVMAsset();
 
 	FString PropertyPath = InPropertyPath;
 	if(UScriptStruct* WrappedStruct = InWrapperObject->GetWrappedStruct())
@@ -2223,7 +2197,7 @@ void FRigVMEditorBase::OnWrappedPropertyChangedChainEvent(URigVMDetailsViewWrapp
 				}
 
 				ForceEditorRefresh();
-				GetRigVMBlueprint()->RequestAutoVMRecompilation();
+				GetRigVMAsset()->RequestAutoVMRecompilation();
 			}
 			else if (PropertyPath == TEXT("CPPType") || PropertyPath == TEXT("CPPTypeObject"))
 			{			
@@ -2235,7 +2209,7 @@ void FRigVMEditorBase::OnWrappedPropertyChangedChainEvent(URigVMDetailsViewWrapp
 						break;
 					}
 				}
-				GetRigVMBlueprint()->RequestAutoVMRecompilation();
+				GetRigVMAsset()->RequestAutoVMRecompilation();
 			}
 			else if (PropertyPath == TEXT("DefaultValue"))
 			{
@@ -2261,7 +2235,7 @@ void FRigVMEditorBase::OnWrappedPropertyChangedChainEvent(URigVMDetailsViewWrapp
 		const FName RootPinName = InPropertyChangedChainEvent.PropertyChain.GetHead()->GetValue()->GetFName();
 		const FString RootPinNameString = RootPinName.ToString();
 		FString PinPath = URigVMPin::JoinPinPath(Node->GetName(), RootPinNameString);
-		URigVMController* Controller = GetRigVMBlueprint()->GetController(Node->GetGraph());
+		URigVMController* Controller = GetRigVMAsset()->GetController(Node->GetGraph());
 		check(Controller);
 
 		const FProperty* Property = WrapperObjects[0]->GetClass()->FindPropertyByName(RootPinName);
@@ -2393,7 +2367,7 @@ void FRigVMEditorBase::OnRequestLocalizeFunctionDialog(FRigVMGraphFunctionIdenti
 	UE::RigVM::Editor::Tools::OnRequestLocalizeFunctionDialog(InFunction, InTargetController, InTargetFunctionHost, bForce);
 }
 
-FRigVMController_BulkEditResult FRigVMEditorBase::OnRequestBulkEditDialog(URigVMBlueprint* InBlueprint,
+FRigVMController_BulkEditResult FRigVMEditorBase::OnRequestBulkEditDialog(TRigVMAssetInterfacePtr InBlueprint,
 	URigVMController* InController, URigVMLibraryNode* InFunction, ERigVMControllerBulkEditType InEditType)
 {
 	if (bAllowBulkEdits)
@@ -2517,7 +2491,7 @@ TRigVMTypeIndex FRigVMEditorBase::OnRequestPinTypeSelectionDialog(const TArray<T
 										static const FLazyName ArrayTypeIcon(TEXT("Kismet.VariableList.ArrayTypeIcon"));
 
 										const FEdGraphPinType PinType = RigVMTypeUtils::PinTypeFromTypeIndex(TypeIndex);
-										const URigVMEdGraphSchema* Schema = CastChecked<URigVMEdGraphSchema>(GetRigVMBlueprint()->GetRigVMEdGraphSchemaClass()->GetDefaultObject());
+										const URigVMEdGraphSchema* Schema = CastChecked<URigVMEdGraphSchema>(GetRigVMAsset()->GetRigVMClientHost()->GetRigVMEdGraphSchemaClass()->GetDefaultObject());
 										const FLinearColor Color = Schema->GetPinTypeColor(PinType);
 										
 										return SNew(STableRow<TSharedPtr<FString>>, Owner)
@@ -2562,7 +2536,7 @@ TRigVMTypeIndex FRigVMEditorBase::OnRequestPinTypeSelectionDialog(const TArray<T
 
 void FRigVMEditorBase::HandleJumpToHyperlink(const UObject* InSubject)
 {
-	URigVMBlueprint* RigBlueprint = GetRigVMBlueprint();
+	TRigVMAssetInterfacePtr RigBlueprint = GetRigVMAsset();
 	if(RigBlueprint == nullptr)
 	{
 		return;
@@ -2603,13 +2577,13 @@ void FRigVMEditorBase::HandleJumpToHyperlink(const UObject* InSubject)
 
 	if (GraphToJumpTo && NodeToJumpTo)
 	{
-		if(URigVMBlueprint* OtherBlueprint = NodeToJumpTo->GetTypedOuter<URigVMBlueprint>())
+		if(IRigVMAssetInterface* OtherAsset = NodeToJumpTo->GetImplementingOuter<IRigVMAssetInterface>())
 		{
-			if(OtherBlueprint != RigBlueprint)
+			if(OtherAsset != RigBlueprint.GetInterface())
 			{
-				if (GEditor->GetEditorSubsystem<UAssetEditorSubsystem>()->OpenEditorForAsset(OtherBlueprint))
+				if (GEditor->GetEditorSubsystem<UAssetEditorSubsystem>()->OpenEditorForAsset(OtherAsset->GetObject()))
 				{
-					IAssetEditorInstance* OtherEditor = GEditor->GetEditorSubsystem<UAssetEditorSubsystem>()->FindEditorForAsset(OtherBlueprint, /*bFocusIfOpen =*/true);
+					IAssetEditorInstance* OtherEditor = GEditor->GetEditorSubsystem<UAssetEditorSubsystem>()->FindEditorForAsset(OtherAsset->GetObject(), /*bFocusIfOpen =*/true);
 					if (FRigVMEditorBase* OtherRigVMEditor = FRigVMEditorBase::GetFromAssetEditorInstance(OtherEditor))
 					{
 						OtherRigVMEditor->HandleJumpToHyperlink(NodeToJumpTo);
@@ -2618,7 +2592,7 @@ void FRigVMEditorBase::HandleJumpToHyperlink(const UObject* InSubject)
 				}
 			}
 					
-			if(URigVMEdGraph* EdGraph = Cast<URigVMEdGraph>(OtherBlueprint->GetEdGraph(NodeToJumpTo->GetGraph())))
+			if(URigVMEdGraph* EdGraph = Cast<URigVMEdGraph>(OtherAsset->GetEdGraph(NodeToJumpTo->GetGraph())))
 			{
 				if(URigVMEdGraphNode* EdGraphNode = Cast<URigVMEdGraphNode>(EdGraph->FindNodeForModelNodeName(NodeToJumpTo->GetFName())))
 				{
@@ -2642,29 +2616,18 @@ void FRigVMEditorBase::HandleJumpToHyperlink(const UObject* InSubject)
 	}
 }
 
-bool FRigVMEditorBase::UpdateDefaultValueForVariable(FBPVariableDescription& InVariable, bool bUseCDO)
+bool FRigVMEditorBase::UpdateDefaultValueForVariable(FRigVMGraphVariableDescription& InVariable, bool bUseDebuggedObject)
 {
 	bool bAnyValueChanged = false;
-	if (URigVMBlueprint* RigVMBlueprint = GetRigVMBlueprint())
+	if (TRigVMAssetInterfacePtr RigVMBlueprint = GetRigVMAsset())
 	{
-		UClass* GeneratedClass = RigVMBlueprint->GeneratedClass;
-		UObject* ObjectContainer = bUseCDO ? GeneratedClass->GetDefaultObject() : RigVMBlueprint->GetObjectBeingDebugged();
-		if(ObjectContainer)
+		FString NewDefaultValue = RigVMBlueprint->GetVariableDefaultValue(InVariable.Name, bUseDebuggedObject);
+		if (InVariable.DefaultValue != NewDefaultValue)
 		{
-			FProperty* TargetProperty = FindFProperty<FProperty>(GeneratedClass, InVariable.VarName);
-
-			if (TargetProperty)
-			{
-				FString NewDefaultValue;
-				const uint8* Container = (const uint8*)ObjectContainer;
-				FBlueprintEditorUtils::PropertyValueToString(TargetProperty, Container, NewDefaultValue, nullptr);
-				if (InVariable.DefaultValue != NewDefaultValue)
-				{
-					InVariable.DefaultValue = NewDefaultValue;
-					bAnyValueChanged = true;
-				}
-			}
+			InVariable.DefaultValue = NewDefaultValue;
+			bAnyValueChanged = true;
 		}
+		
 	}
 	return bAnyValueChanged;
 }
@@ -2673,8 +2636,8 @@ void FRigVMEditorBase::UpdateRigVMHost()
 {
 	DECLARE_SCOPE_HIERARCHICAL_COUNTER_FUNC()
 
-	URigVMBlueprint* Blueprint = GetRigVMBlueprint();
-	if(UClass* Class = Blueprint->GeneratedClass)
+	TRigVMAssetInterfacePtr Blueprint = GetRigVMAsset();
+	if(UClass* Class = Blueprint->GetRigVMGeneratedClass())
 	{
 		if (URigVMHost* CurrentHost = GetRigVMHost())
 		{
@@ -2697,7 +2660,7 @@ void FRigVMEditorBase::UpdateRigVMHost()
 		URigVMHost* RigVMHost = GetRigVMHost();
 		if (RigVMHost == nullptr)
 		{
-			RigVMHost = NewObject<URigVMHost>(GetOuterForHost(), Class);
+			RigVMHost = Blueprint->CreateRigVMHostSuper(GetOuterForHost());
 			SetHost(RigVMHost);
 			
 			// this is editing time rig
@@ -2705,6 +2668,7 @@ void FRigVMEditorBase::UpdateRigVMHost()
 
 			RigVMHost->Initialize(true);
  		}
+		check(IsValid(RigVMHost));
 
 #if WITH_EDITOR
 		RigVMHost->SetIsInDebugMode(ExecutionMode == ERigVMEditorExecutionModeType_Debug);
@@ -2713,13 +2677,13 @@ void FRigVMEditorBase::UpdateRigVMHost()
 		CacheNameLists();
 
 		// Make sure the object being debugged is the preview instance
-		GetBlueprintObj()->SetObjectBeingDebugged(RigVMHost);
+		GetRigVMAsset()->SetObjectBeingDebugged(RigVMHost);
 
 		if(!bIsCompilingThroughUI)
 		{
-			Blueprint->SetFlags(RF_Transient);
+			Blueprint->GetObject()->SetFlags(RF_Transient);
 			Blueprint->RecompileVM();
-			Blueprint->ClearFlags(RF_Transient);
+			Blueprint->GetObject()->ClearFlags(RF_Transient);
 		}
 
 		RigVMHost->OnInitialized_AnyThread().AddSP(StaticCastSharedRef<FRigVMEditorBase>(SharedRef()), &FRigVMEditorBase::HandleVMExecutedEvent);
@@ -2731,10 +2695,10 @@ void FRigVMEditorBase::UpdateRigVMHost()
 
 void FRigVMEditorBase::CacheNameLists()
 {
-	if (URigVMBlueprint* RigVMBlueprint = GetRigVMBlueprint())
+	if (TRigVMAssetInterfacePtr RigVMBlueprint = GetRigVMAsset())
 	{
 		TArray<UEdGraph*> EdGraphs;
-		RigVMBlueprint->GetAllGraphs(EdGraphs);
+		RigVMBlueprint->GetAllEdGraphs(EdGraphs);
 
 		for (UEdGraph* Graph : EdGraphs)
 		{
@@ -2757,7 +2721,7 @@ void FRigVMEditorBase::OnCreateComment()
 		{
 			if (URigVMEdGraph* RigVMEdGraph = Cast<URigVMEdGraph>(Graph))
 			{
-				if (URigVMBlueprint* Blueprint = GetRigVMBlueprint())
+				if (TRigVMAssetInterfacePtr Blueprint = GetRigVMAsset())
 				{
 					if (URigVMController* Controller = Blueprint->GetController(RigVMEdGraph))
 					{
@@ -2889,7 +2853,7 @@ void FRigVMEditorBase::SetDetailObjects(const TArray<UObject*>& InObjects, bool
 			EditClassDefaults_Clicked();
 			return;
 		}
-		else if (InObjects[0] == GetBlueprintObj())
+		else if (InObjects[0] == GetRigVMAsset()->GetObject())
 		{
 			EditGlobalOptions_Clicked();
 			return;
@@ -2913,7 +2877,7 @@ void FRigVMEditorBase::SetDetailObjects(const TArray<UObject*>& InObjects, bool
 		{
 			if(!LibraryNode->IsA<URigVMFunctionReferenceNode>())
 			{
-				if (UEdGraph* EdGraph = GetRigVMBlueprint()->GetEdGraph(LibraryNode->GetContainedGraph()))
+				if (UEdGraph* EdGraph = GetRigVMAsset()->GetEdGraph(LibraryNode->GetContainedGraph()))
 				{
 					FilteredObjects.AddUnique(EdGraph);
 					ModelNodes.Remove(LibraryNode);
@@ -2923,7 +2887,7 @@ void FRigVMEditorBase::SetDetailObjects(const TArray<UObject*>& InObjects, bool
 		}
 		else if (Cast<URigVMFunctionEntryNode>(InObject) || Cast<URigVMFunctionReturnNode>(InObject))
 		{
-			if (UEdGraph* EdGraph = GetRigVMBlueprint()->GetEdGraph(CastChecked<URigVMNode>(InObject)->GetGraph()))
+			if (UEdGraph* EdGraph = GetRigVMAsset()->GetEdGraph(CastChecked<URigVMNode>(InObject)->GetGraph()))
 			{
 				FilteredObjects.AddUnique(EdGraph);
 				ModelNodes.Remove(Cast<URigVMNode>(InObject));
@@ -2935,7 +2899,7 @@ void FRigVMEditorBase::SetDetailObjects(const TArray<UObject*>& InObjects, bool
 			TSharedRef<FAssetEditorToolkit> SharedApp = GetHostingApp().ToSharedRef();
 			if (SharedApp->IsBlueprintEditor())
 			{
-				if (URigVMEdGraph* EdGraph = Cast<URigVMEdGraph>(GetRigVMBlueprint()->GetEdGraph(CastChecked<URigVMNode>(InObject)->GetGraph())))
+				if (URigVMEdGraph* EdGraph = Cast<URigVMEdGraph>(GetRigVMAsset()->GetEdGraph(CastChecked<URigVMNode>(InObject)->GetGraph())))
 				{
 					if(UEdGraphNode* EdGraphNode = EdGraph->FindNodeForModelNodeName(CommentNode->GetFName()))
 					{
@@ -2960,7 +2924,7 @@ void FRigVMEditorBase::SetDetailObjects(const TArray<UObject*>& InObjects, bool
 			(void)CDOWrapper->GetClassForNodes(ModelNodes, false);
 
 			// create the wrapper object
-			if(URigVMDetailsViewWrapperObject* WrapperObject = URigVMDetailsViewWrapperObject::MakeInstance(GetDetailWrapperClass(), GetBlueprintObj(), ModelNodes, ModelNode))
+			if(URigVMDetailsViewWrapperObject* WrapperObject = URigVMDetailsViewWrapperObject::MakeInstance(GetDetailWrapperClass(), GetRigVMAsset()->GetObject(), ModelNodes, ModelNode))
 			{
 				WrapperObject->GetWrappedPropertyChangedChainEvent().AddSP(StaticCastSharedRef<FRigVMEditorBase>(SharedRef()), &FRigVMEditorBase::OnWrappedPropertyChangedChainEvent);
 				WrapperObject->AddToRoot();
@@ -3523,15 +3487,15 @@ void FRigVMEditorBase::ClearDetailsViewWrapperObjects()
 
 void FRigVMEditorBase::SetHost(URigVMHost* InHost)
 {
-	if (URigVMBlueprint* RigVMBlueprint = GetRigVMBlueprint())
+	if (TRigVMAssetInterfacePtr RigVMBlueprint = GetRigVMAsset())
 	{
-		if (IsValid(RigVMBlueprint->EditorHost) && RigVMBlueprint->EditorHost->GetOuter() == GetOuterForHost())
+		if (IsValid(RigVMBlueprint->GetEditorHost()) && RigVMBlueprint->GetEditorHost()->GetOuter() == GetOuterForHost())
 		{
-			RigVMBlueprint->EditorHost->Rename(nullptr, GetTransientPackage(), REN_DoNotDirty | REN_DontCreateRedirectors | REN_NonTransactional);
-			RigVMBlueprint->EditorHost->MarkAsGarbage();
+			RigVMBlueprint->GetEditorHost()->Rename(nullptr, GetTransientPackage(), REN_DoNotDirty | REN_DontCreateRedirectors | REN_NonTransactional);
+			RigVMBlueprint->GetEditorHost()->MarkAsGarbage();
 		}
-		RigVMBlueprint->EditorHost = InHost;
-		if(RigVMBlueprint->EditorHost && IsValid(RigVMBlueprint->EditorHost))
+		RigVMBlueprint->GetEditorHost() = InHost;
+		if(RigVMBlueprint->GetEditorHost() && IsValid(RigVMBlueprint->GetEditorHost()))
 		{
 			OnPreviewHostUpdated().Broadcast(this);
 		}
@@ -3540,7 +3504,7 @@ void FRigVMEditorBase::SetHost(URigVMHost* InHost)
 
 URigVMGraph* FRigVMEditorBase::GetFocusedModel() const
 {
-	const URigVMBlueprint* Blueprint = GetRigVMBlueprint();
+	const TRigVMAssetInterfacePtr Blueprint = GetRigVMAsset();
 	if (Blueprint == nullptr)
 	{
 		return nullptr;
@@ -3552,7 +3516,7 @@ URigVMGraph* FRigVMEditorBase::GetFocusedModel() const
 
 URigVMController* FRigVMEditorBase::GetFocusedController() const
 {
-	URigVMBlueprint* Blueprint = GetRigVMBlueprint();
+	TRigVMAssetInterfacePtr Blueprint = GetRigVMAsset();
 	if (Blueprint == nullptr)
 	{
 		return nullptr;
@@ -3713,41 +3677,9 @@ void FRigVMEditorBase::FillToolbar(FToolBarBuilder& ToolbarBuilder, bool bEndSec
 	}
 }
 
-void FRigVMEditorBase::HandleHideItem()
-{
-	DECLARE_SCOPE_HIERARCHICAL_COUNTER_FUNC()
-
-	URigVMBlueprint* RigVMBlueprint = GetRigVMBlueprint();
-	if(RigVMBlueprint == nullptr)
-	{
-		return;
-	}
-
-	TSet<UObject*> SelectedNodes = GetSelectedNodes();
-	if(SelectedNodes.Num() > 0)
-	{
-		FScopedTransaction Transaction(LOCTEXT("HideRigItem", "Hide rig item"));
-
-		RigVMBlueprint->Modify();
-
-		for(UObject* SelectedNodeObject : SelectedNodes)
-		{
-			if(URigVMEdGraphNode* SelectedNode = Cast<URigVMEdGraphNode>(SelectedNodeObject))
-			{
-				FBlueprintEditorUtils::RemoveNode(RigVMBlueprint, SelectedNode, true);
-			}
-		}
-	}
-}
-
-bool FRigVMEditorBase::CanHideItem() const
-{
-	return GetSelectedNodes().Num() > 0;
-}
-
 void FRigVMEditorBase::UpdateStaleWatchedPins()
 {
-	URigVMBlueprint* RigVMBlueprint = GetRigVMBlueprint();
+	TRigVMAssetInterfacePtr RigVMBlueprint = GetRigVMAsset();
 	if (RigVMBlueprint == nullptr)
 	{
 		return;
@@ -3757,8 +3689,7 @@ void FRigVMEditorBase::UpdateStaleWatchedPins()
 	uint16 WatchCount;
 
 	// Find all unique pins being watched
-	FKismetDebugUtilities::ForeachPinWatch(
-		RigVMBlueprint,
+	RigVMBlueprint->ForeachPinWatch(
 		[&AllPins, &WatchCount](UEdGraphPin* Pin)
 		{
 			++WatchCount;
@@ -3787,11 +3718,11 @@ void FRigVMEditorBase::UpdateStaleWatchedPins()
 	// Refresh watched pins with unique pins (throw away null or duplicate watches)
 	if (WatchCount != AllPins.Num())
 	{
-		RigVMBlueprint->Status = BS_Dirty;
+		RigVMBlueprint->SetAssetStatus(RVMA_Dirty);
 	}
 
-	FKismetDebugUtilities::ClearPinWatches(RigVMBlueprint);
-
+	RigVMBlueprint->ClearPinWatches();
+	
 	TArray<URigVMGraph*> Models = RigVMBlueprint->GetAllModels();
 	for (URigVMGraph* Model : Models)
 	{
@@ -3809,20 +3740,20 @@ void FRigVMEditorBase::UpdateStaleWatchedPins()
 	}
 	for (UEdGraphPin* Pin : AllPins)
 	{
-		FKismetDebugUtilities::AddPinWatch(RigVMBlueprint, FBlueprintWatchedPin(Pin));
+		RigVMBlueprint->AddPinWatch(Pin);
 		UEdGraph* EdGraph = Pin->GetOwningNode()->GetGraph();
 		RigVMBlueprint->GetController(EdGraph)->SetPinIsWatched(Pin->GetName(), true, false);
 	}
 }
 
-void FRigVMEditorBase::HandleRefreshEditorFromBlueprint(URigVMBlueprint* InBlueprint)
+void FRigVMEditorBase::HandleRefreshEditorFromBlueprint(TRigVMAssetInterfacePtr InBlueprint)
 {
 	Compile();
 }
 
 void FRigVMEditorBase::HandleVariableDroppedFromBlueprint(UObject* InSubject, FProperty* InVariableToDrop, const FVector2D& InDropPosition, const FVector2D& InScreenPosition)
 {
-	URigVMBlueprint* Blueprint = Cast<URigVMBlueprint>(GetBlueprintObj());
+	TRigVMAssetInterfacePtr Blueprint = GetRigVMAsset();
 	if (Blueprint == nullptr)
 	{
 		return;
@@ -3900,9 +3831,9 @@ void FRigVMEditorBase::OnGraphNodeClicked(URigVMEdGraphNode* InNode)
 	}
 }
 
-void FRigVMEditorBase::OnNodeDoubleClicked(URigVMBlueprint* InBlueprint, URigVMNode* InNode)
+void FRigVMEditorBase::OnNodeDoubleClicked(TRigVMAssetInterfacePtr InBlueprint, URigVMNode* InNode)
 {
-	ensure(GetRigVMBlueprint() == InBlueprint);
+	ensure(GetRigVMAsset() == InBlueprint);
 
 	if (URigVMLibraryNode* LibraryNode = Cast<URigVMLibraryNode>(InNode))
 	{
@@ -3949,7 +3880,7 @@ void FRigVMEditorBase::HandleShowCurrentStatement()
 {
 	if(HaltedAtNode)
 	{
-		if(URigVMBlueprint* Blueprint = Cast<URigVMBlueprint>(GetBlueprintObj()))
+		if(TRigVMAssetInterfacePtr Blueprint = GetRigVMAsset())
 		{
 			if(URigVMEdGraph* EdGraph = Cast<URigVMEdGraph>(Blueprint->GetEdGraph(HaltedAtNode->GetGraph())))
 			{
@@ -3964,7 +3895,7 @@ void FRigVMEditorBase::HandleShowCurrentStatement()
 
 void FRigVMEditorBase::HandleBreakpointActionRequested(const ERigVMBreakpointAction BreakpointAction)
 {
-	if(URigVMHost* DebuggedHost = Cast<URigVMHost>(GetBlueprintObj()->GetObjectBeingDebugged()))
+	if(URigVMHost* DebuggedHost = Cast<URigVMHost>(GetRigVMAsset()->GetObjectBeingDebugged()))
 	{
 		DebuggedHost->ExecuteBreakpointAction(BreakpointAction);
 	}
@@ -3989,7 +3920,7 @@ void FRigVMEditorBase::FrameSelection()
 
 void FRigVMEditorBase::SwapFunctionWithinAsset()
 {
-	const FAssetData Asset = UE::RigVM::Editor::Tools::FindAssetFromAnyPath(GetRigVMBlueprint()->GetPathName(), true);
+	const FAssetData Asset = UE::RigVM::Editor::Tools::FindAssetFromAnyPath(GetRigVMAsset()->GetObject()->GetPathName(), true);
 	SwapFunctionForAssets({Asset}, true);
 }
 
@@ -3997,7 +3928,8 @@ void FRigVMEditorBase::SwapFunctionAcrossProject()
 {
 	const IAssetRegistry& AssetRegistry = FModuleManager::LoadModuleChecked<FAssetRegistryModule>(TEXT("AssetRegistry")).Get();
 	TArray<FAssetData> AllAssets;
-	AssetRegistry.GetAssetsByClass(GetRigVMBlueprint()->GetClass()->GetClassPathName(), AllAssets, true);
+	// TODO: sara-sc we might have a problem when trying to swap functions between URigVMBlueprint and URigVMAsset classes
+	AssetRegistry.GetAssetsByClass(GetRigVMAsset()->GetObject()->GetClass()->GetClassPathName(), AllAssets, true);
 	SwapFunctionForAssets(AllAssets, false);
 }
 
@@ -4019,7 +3951,7 @@ void FRigVMEditorBase::SwapFunctionForAssets(const TArray<FAssetData>& InAssets,
 
 void FRigVMEditorBase::SwapAssetReferences()
 {
-	const FAssetData Asset = UE::RigVM::Editor::Tools::FindAssetFromAnyPath(GetRigVMBlueprint()->GetPathName(), true);
+	const FAssetData Asset = UE::RigVM::Editor::Tools::FindAssetFromAnyPath(GetRigVMAsset()->GetObject()->GetPathName(), true);
 	
 	SRigVMSwapAssetReferencesWidget::FArguments WidgetArgs;
 	WidgetArgs
@@ -4045,12 +3977,12 @@ void FRigVMEditorBase::OnGraphNodeDropToPerform(TSharedPtr<FDragDropOperation> I
 
 		if (GetFocusedGraphEditor().IsValid())
 		{
-			const URigVMEdGraphSchema* Schema = CastChecked<URigVMEdGraphSchema>(GetRigVMBlueprint()->GetRigVMEdGraphSchemaClass()->GetDefaultObject());
+			const URigVMEdGraphSchema* Schema = CastChecked<URigVMEdGraphSchema>(GetRigVMAsset()->GetRigVMClientHost()->GetRigVMEdGraphSchemaClass()->GetDefaultObject());
 			if (ExplorerOp->GetElement().Type == ERigVMExplorerElementType::Function)
 			{
-				if (URigVMBlueprint* TargetRigBlueprint = Cast<URigVMBlueprint>(FBlueprintEditorUtils::FindBlueprintForGraph(InGraph)))
+				if (TRigVMAssetInterfacePtr TargetRigBlueprint = FRigVMBlueprintUtils::FindAssetForGraph(InGraph))
 				{
-					if (URigVMFunctionLibrary* Library = ExplorerOp->GetBlueprint()->GetLocalFunctionLibrary())
+					if (URigVMFunctionLibrary* Library = ExplorerOp->GetRigVMAsset()->GetLocalFunctionLibrary())
 					{
 						if (URigVMLibraryNode* LibraryNode = Library->FindFunction(*ExplorerOp->GetElement().Name))
 						{
@@ -4061,15 +3993,15 @@ void FRigVMEditorBase::OnGraphNodeDropToPerform(TSharedPtr<FDragDropOperation> I
 			}
 			else if (ExplorerOp->GetElement().Type == ERigVMExplorerElementType::Variable)
 			{
-				if (URigVMBlueprint* TargetRigBlueprint = Cast<URigVMBlueprint>(FBlueprintEditorUtils::FindBlueprintForGraph(InGraph)))
+				if (TRigVMAssetInterfacePtr TargetRigBlueprint = FRigVMBlueprintUtils::FindAssetForGraph(InGraph))
 				{
-					FProperty* Property = TargetRigBlueprint->SkeletonGeneratedClass->FindPropertyByName(*ExplorerOp->GetElement().Name);
+					FProperty* Property = TargetRigBlueprint->FindGeneratedPropertyByName(*ExplorerOp->GetElement().Name);
 					Schema->RequestVariableDropOnPanel(InGraph, Property, InNodePosition, InScreenPosition);
 				}
 			}
 			else if (ExplorerOp->GetElement().Type == ERigVMExplorerElementType::LocalVariable)
 			{
-				if (URigVMBlueprint* TargetRigBlueprint = Cast<URigVMBlueprint>(FBlueprintEditorUtils::FindBlueprintForGraph(InGraph)))
+				if (TRigVMAssetInterfacePtr TargetRigBlueprint = FRigVMBlueprintUtils::FindAssetForGraph(InGraph))
 				{
 					if (URigVMGraph* Graph = TargetRigBlueprint->GetFocusedModel())
 					{
@@ -4077,7 +4009,7 @@ void FRigVMEditorBase::OnGraphNodeDropToPerform(TSharedPtr<FDragDropOperation> I
 						{
 							if (LocalVariable.Name == ExplorerOp->GetElement().Name)
 							{
-								if (URigVMController* Controller = TargetRigBlueprint->GetRigVMClient()->GetController(Graph))
+								if (URigVMController* Controller = TargetRigBlueprint->GetRigVMClientHost()->GetRigVMClient()->GetController(Graph))
 								{
 									FMenuBuilder MenuBuilder(true, nullptr);
 									const FText VariableNameText = FText::FromName( LocalVariable.Name );
@@ -4156,7 +4088,7 @@ void FRigVMEditorBase::UpdateGraphCompilerErrors()
 {
 	DECLARE_SCOPE_HIERARCHICAL_COUNTER_FUNC()
 
-	URigVMBlueprint* Blueprint = Cast<URigVMBlueprint>(GetBlueprintObj());
+	TRigVMAssetInterfacePtr Blueprint = GetRigVMAsset();
 	URigVMHost* RigVMHost = GetRigVMHost();
 	if (Blueprint && RigVMHost && RigVMHost->GetVM())
 	{
@@ -4169,7 +4101,7 @@ void FRigVMEditorBase::UpdateGraphCompilerErrors()
 		const FRigVMByteCode& ByteCode = VM->GetByteCode();
 
 		TArray<UEdGraph*> EdGraphs;
-		Blueprint->GetAllGraphs(EdGraphs);
+		Blueprint->GetAllEdGraphs(EdGraphs);
 
 		for (UEdGraph* Graph : EdGraphs)
 		{
@@ -4243,8 +4175,8 @@ void FRigVMEditorBase::UpdateGraphCompilerErrors()
 
 			if (bFoundError)
 			{
-				Blueprint->Status = BS_Error;
-				(void)Blueprint->MarkPackageDirty();
+				Blueprint->SetAssetStatus(RVMA_Error);
+				(void)Blueprint->GetObject()->MarkPackageDirty();
 			}
 
 			RigVMLog.RemoveRedundantEntries();
@@ -4316,7 +4248,7 @@ void FRigVMEditorBase::SetExecutionMode(const ERigVMEditorExecutionModeType InEx
 	}
 
 	ExecutionMode = InExecutionMode;
-	GetRigVMBlueprint()->SetDebugMode(InExecutionMode == ERigVMEditorExecutionModeType_Debug);
+	GetRigVMAsset()->SetDebugMode(InExecutionMode == ERigVMEditorExecutionModeType_Debug);
 	Compile();
 	
 	if (URigVMHost* CurrentHost = GetRigVMHost())
@@ -4359,7 +4291,7 @@ FSlateIcon FRigVMEditorBase::GetExecutionModeIcon() const
 
 void FRigVMEditorBase::GetDebugObjects(TArray<FRigVMCustomDebugObject>& DebugList) const
 {
-	URigVMBlueprint* RigVMBlueprint = GetRigVMBlueprint();
+	TRigVMAssetInterfacePtr RigVMBlueprint = GetRigVMAsset();
 	if (RigVMBlueprint == nullptr)
 	{
 		return;
@@ -4376,97 +4308,89 @@ void FRigVMEditorBase::GetDebugObjects(TArray<FRigVMCustomDebugObject>& DebugLis
 		}
 	}
 
-	URigVMBlueprintGeneratedClass* GeneratedClass = RigVMBlueprint->GetRigVMBlueprintGeneratedClass();
-	if (GeneratedClass)
+	struct Local
 	{
-		struct Local
+		static bool IsPendingKillOrUnreachableRecursive(UObject* InObject)
 		{
-			static bool IsPendingKillOrUnreachableRecursive(UObject* InObject)
+			if (InObject != nullptr)
 			{
-				if (InObject != nullptr)
+				if (!IsValidChecked(InObject) || InObject->IsUnreachable())
 				{
-					if (!IsValidChecked(InObject) || InObject->IsUnreachable())
-					{
-						return true;
-					}
-					return IsPendingKillOrUnreachableRecursive(InObject->GetOuter());
+					return true;
 				}
+				return IsPendingKillOrUnreachableRecursive(InObject->GetOuter());
+			}
+			return false;
+		}
+
+		static bool OuterNameContainsRecursive(UObject* InObject, const FString& InStringToSearch)
+		{
+			if (InObject == nullptr)
+			{
 				return false;
 			}
 
-			static bool OuterNameContainsRecursive(UObject* InObject, const FString& InStringToSearch)
+			UObject* InObjectOuter = InObject->GetOuter();
+			if (InObjectOuter == nullptr)
 			{
-				if (InObject == nullptr)
+				return false;
+			}
+
+			if (InObjectOuter->GetName().Contains(InStringToSearch))
+			{
+				return true;
+			}
+
+			return OuterNameContainsRecursive(InObjectOuter, InStringToSearch);
+		}
+	};
+
+	TArray<UObject*> ArchetypeInstances = RigVMBlueprint->GetArchetypeInstances(false, false);
+
+	// run in two passes - find the PIE related objects first
+	for (int32 Pass = 0; Pass < 2 ; Pass++)
+	{
+		for (UObject* Instance : ArchetypeInstances)
+		{
+			URigVMHost* InstancedHost = Cast<URigVMHost>(Instance);
+			if (InstancedHost && IsValid(InstancedHost) && InstancedHost != GetRigVMHost())
+			{
+				if (InstancedHost->GetOuter() == nullptr)
 				{
-					return false;
+					continue;
 				}
 
-				UObject* InObjectOuter = InObject->GetOuter();
-				if (InObjectOuter == nullptr)
+				UWorld* World = InstancedHost->GetWorld();
+				if (World == nullptr)
 				{
-					return false;
+					continue;
 				}
 
-				if (InObjectOuter->GetName().Contains(InStringToSearch))
+				// during pass 0 only do PIE instances,
+				// and in pass 1 only do non PIE instances
+				if((Pass == 1) == (World->IsPlayInEditor()))
 				{
-					return true;
+					continue;
 				}
 
-				return OuterNameContainsRecursive(InObjectOuter, InStringToSearch);
-			}
-		};
-
-		if (UObject* DefaultObject = GeneratedClass->GetDefaultObject(false))
-		{
-			TArray<UObject*> ArchetypeInstances;
-			DefaultObject->GetArchetypeInstances(ArchetypeInstances);
-
-			// run in two passes - find the PIE related objects first
-			for (int32 Pass = 0; Pass < 2 ; Pass++)
-			{
-				for (UObject* Instance : ArchetypeInstances)
+				// ensure to only allow preview actors in preview worlds
+				if (World->IsPreviewWorld())
 				{
-					URigVMHost* InstancedHost = Cast<URigVMHost>(Instance);
-					if (InstancedHost && IsValid(InstancedHost) && InstancedHost != GetRigVMHost())
+					if (!Local::OuterNameContainsRecursive(InstancedHost, TEXT("Preview")))
 					{
-						if (InstancedHost->GetOuter() == nullptr)
-						{
-							continue;
-						}
-
-						UWorld* World = InstancedHost->GetWorld();
-						if (World == nullptr)
-						{
-							continue;
-						}
-
-						// during pass 0 only do PIE instances,
-						// and in pass 1 only do non PIE instances
-						if((Pass == 1) == (World->IsPlayInEditor()))
-						{
-							continue;
-						}
-
-						// ensure to only allow preview actors in preview worlds
-						if (World->IsPreviewWorld())
-						{
-							if (!Local::OuterNameContainsRecursive(InstancedHost, TEXT("Preview")))
-							{
-								continue;
-							}
-						}
-
-						if (Local::IsPendingKillOrUnreachableRecursive(InstancedHost))
-						{
-							continue;
-						}
-
-						FRigVMCustomDebugObject DebugObject;
-						DebugObject.Object = InstancedHost;
-						DebugObject.NameOverride = GetCustomDebugObjectLabel(InstancedHost);
-						DebugList.Add(DebugObject);
+						continue;
 					}
 				}
+
+				if (Local::IsPendingKillOrUnreachableRecursive(InstancedHost))
+				{
+					continue;
+				}
+
+				FRigVMCustomDebugObject DebugObject;
+				DebugObject.Object = InstancedHost;
+				DebugObject.NameOverride = GetCustomDebugObjectLabel(InstancedHost);
+				DebugList.Add(DebugObject);
 			}
 		}
 	}
@@ -4474,7 +4398,7 @@ void FRigVMEditorBase::GetDebugObjects(TArray<FRigVMCustomDebugObject>& DebugLis
 
 void FRigVMEditorBase::HandleSetObjectBeingDebugged(UObject* InObject)
 {
-	if(URigVMHost* PreviouslyDebuggedHost = Cast<URigVMHost>(GetBlueprintObj()->GetObjectBeingDebugged()))
+	if(URigVMHost* PreviouslyDebuggedHost = Cast<URigVMHost>(GetRigVMAsset()->GetObjectBeingDebugged()))
 	{
 		if(!URigVMHost::IsGarbageOrDestroyed(PreviouslyDebuggedHost))
 		{
@@ -4489,20 +4413,19 @@ void FRigVMEditorBase::HandleSetObjectBeingDebugged(UObject* InObject)
 	if (DebuggedHost == nullptr)
 	{
 		// fall back to our default control rig (which still can be nullptr)
-		if (GetRigVMBlueprint() != nullptr && !bIsSettingObjectBeingDebugged)
+		if (GetRigVMAsset() != nullptr && !bIsSettingObjectBeingDebugged)
 		{
 			TGuardValue<bool> GuardSettingObjectBeingDebugged(bIsSettingObjectBeingDebugged, true);
-			GetBlueprintObj()->SetObjectBeingDebugged(GetRigVMHost());
+			GetRigVMAsset()->SetObjectBeingDebugged(GetRigVMHost());
 			return;
 		}
 	}
 
-	if (URigVMBlueprint* RigBlueprint = GetRigVMBlueprint())
+	if (TRigVMAssetInterfacePtr RigBlueprint = GetRigVMAsset())
 	{
-		if (URigVMBlueprintGeneratedClass* GeneratedClass = RigBlueprint->GetRigVMBlueprintGeneratedClass())
+		if (URigVM* VM = RigBlueprint->GetVM(true))
 		{
-			URigVMHost* CDO = Cast<URigVMHost>(GeneratedClass->GetDefaultObject(true /* create if needed */));
-			if (CDO->GetVM()->GetInstructions().Num() <= 1 /* only exit */)
+			if (VM->GetInstructions().Num() <= 1 /* only exit */)
 			{
 				RigBlueprint->RecompileVM();
 				RigBlueprint->RequestRigVMInit();
@@ -4546,16 +4469,17 @@ FString FRigVMEditorBase::GetCustomDebugObjectLabel(UObject* ObjectBeingDebugged
 				static const FString PIEPrefix = TEXT("PIE");
 				WorldLabel = PIEPrefix;
 			}
-			return FString::Printf(TEXT("%s: %s in %s"), *WorldLabel, *GetBlueprintObj()->GetName(), *ParentActor->GetActorLabel());
+			return FString::Printf(TEXT("%s: %s in %s"), *WorldLabel, *GetRigVMAsset()->GetObject()->GetName(), *ParentActor->GetActorLabel());
 		}
 	}
+	
 
-	return GetBlueprintObj()->GetName();
+	return GetRigVMAsset()->GetObject()->GetName();
 }
 
 void FRigVMEditorBase::OnPIEStopped(bool bSimulation)
 {
-	if(URigVMBlueprint* Blueprint = GetRigVMBlueprint())
+	if(TRigVMAssetInterfacePtr Blueprint = GetRigVMAsset())
 	{
 		Blueprint->SetObjectBeingDebugged(GetRigVMHost());
 	}
@@ -4563,7 +4487,7 @@ void FRigVMEditorBase::OnPIEStopped(bool bSimulation)
 
 bool FRigVMEditorBase::IsEditingSingleBlueprint() const
 {
-	return GetBlueprintObj() != nullptr;
+	return GetRigVMAsset() != nullptr;
 }
 
 #undef LOCTEXT_NAMESPACE
\ No newline at end of file
diff --git a/Engine/Plugins/Runtime/RigVM/Source/RigVMEditor/Private/Editor/RigVMEditorCommands.h b/Engine/Plugins/Runtime/RigVM/Source/RigVMEditor/Private/Editor/RigVMEditorCommands.h
--- a/Engine/Plugins/Runtime/RigVM/Source/RigVMEditor/Private/Editor/RigVMEditorCommands.h
+++ b/Engine/Plugins/Runtime/RigVM/Source/RigVMEditor/Private/Editor/RigVMEditorCommands.h
@@ -11,8 +11,8 @@ class FRigVMEditorCommands : public TCommands<FRigVMEditorCommands>
 public:
 	FRigVMEditorCommands() : TCommands<FRigVMEditorCommands>
 	(
-		"RigVMBlueprint",
-		NSLOCTEXT("Contexts", "RigVM", "RigVM Blueprint"),
+		"RigVMAsset",
+		NSLOCTEXT("Contexts", "RigVM", "RigVM Asset"),
 		NAME_None, // "MainFrame" // @todo Fix this crash
 		FRigVMEditorStyle::Get().GetStyleSetName() // Icon Style Set
 	)
diff --git a/Engine/Plugins/Runtime/RigVM/Source/RigVMEditor/Private/Editor/RigVMEditorMenuContext.cpp b/Engine/Plugins/Runtime/RigVM/Source/RigVMEditor/Private/Editor/RigVMEditorMenuContext.cpp
--- a/Engine/Plugins/Runtime/RigVM/Source/RigVMEditor/Private/Editor/RigVMEditorMenuContext.cpp
+++ b/Engine/Plugins/Runtime/RigVM/Source/RigVMEditor/Private/Editor/RigVMEditorMenuContext.cpp
@@ -2,7 +2,8 @@
 
 #include "Editor/RigVMEditorMenuContext.h"
 
-#include "RigVMBlueprint.h"
+#include "RigVMAsset.h"
+#include "RigVMBlueprintLegacy.h"
 #include "Editor/RigVMNewEditor.h"
 #include "Framework/Application/SlateApplication.h"
 
@@ -15,10 +16,16 @@ void URigVMEditorMenuContext::Init(TWeakPtr<IRigVMEditor> InRigVMEditor, const F
 }
 
 URigVMBlueprint* URigVMEditorMenuContext::GetRigVMBlueprint() const
+{
+	TRigVMAssetInterfacePtr Interface = GetRigVMAsset();
+	return Cast<URigVMBlueprint>(Interface->GetObject());
+}
+
+TRigVMAssetInterfacePtr URigVMEditorMenuContext::GetRigVMAsset() const
 {
 	if (const TSharedPtr<IRigVMEditor> Editor = WeakRigVMEditor.Pin())
 	{
-		return Editor->GetRigVMBlueprint();
+		return Editor->GetRigVMAsset();
 	}
 	
 	return nullptr;
@@ -26,9 +33,9 @@ URigVMBlueprint* URigVMEditorMenuContext::GetRigVMBlueprint() const
 
 URigVMHost* URigVMEditorMenuContext::GetRigVMHost() const
 {
-	if (URigVMBlueprint* RigBlueprint = GetRigVMBlueprint())
+	if (TRigVMAssetInterfacePtr RigBlueprint = GetRigVMAsset())
 	{
-		if (URigVMHost* Host = Cast<URigVMHost>(RigBlueprint->GetObjectBeingDebugged()))
+		if (URigVMHost* Host = RigBlueprint->GetDebuggedRigVMHost())
 		{
 			return Host;
 		}
diff --git a/Engine/Plugins/Runtime/RigVM/Source/RigVMEditor/Private/Editor/RigVMEditorTools.cpp b/Engine/Plugins/Runtime/RigVM/Source/RigVMEditor/Private/Editor/RigVMEditorTools.cpp
--- a/Engine/Plugins/Runtime/RigVM/Source/RigVMEditor/Private/Editor/RigVMEditorTools.cpp
+++ b/Engine/Plugins/Runtime/RigVM/Source/RigVMEditor/Private/Editor/RigVMEditorTools.cpp
@@ -168,8 +168,8 @@ bool FFilterByAssetTag::PassesFilter(const FContentBrowserItem& InItem) const
 	FAssetData AssetData;
 	if (InItem.Legacy_TryGetAssetData(AssetData))
 	{
-		static const FName AssetVariantPropertyName = GET_MEMBER_NAME_CHECKED(URigVMBlueprint, AssetVariant);
-		const FProperty* AssetVariantProperty = CastField<FProperty>(URigVMBlueprint::StaticClass()->FindPropertyByName(AssetVariantPropertyName));
+		static const FName AssetVariantPropertyName = TEXT("AssetVariant");
+		const FProperty* AssetVariantProperty = AssetData.GetClass()->FindPropertyByName(AssetVariantPropertyName);
 		const FString VariantStr = AssetData.GetTagValueRef<FString>(AssetVariantPropertyName);
 		if(!VariantStr.IsEmpty())
 		{
diff --git a/Engine/Plugins/Runtime/RigVM/Source/RigVMEditor/Private/Editor/RigVMFindReferences.cpp b/Engine/Plugins/Runtime/RigVM/Source/RigVMEditor/Private/Editor/RigVMFindReferences.cpp
--- a/Engine/Plugins/Runtime/RigVM/Source/RigVMEditor/Private/Editor/RigVMFindReferences.cpp
+++ b/Engine/Plugins/Runtime/RigVM/Source/RigVMEditor/Private/Editor/RigVMFindReferences.cpp
@@ -3,7 +3,8 @@
 #include "Editor/RigVMFindReferences.h"
 
 #include "K2Node_Variable.h"
-#include "RigVMBlueprint.h"
+#include "RigVMBlueprintLegacy.h"
+#include "RigVMBlueprintUtils.h"
 #include "EdGraph/RigVMEdGraphSchema.h"
 #include "Editor/RigVMNewEditor.h"
 #include "Framework/Commands/GenericCommands.h"
@@ -76,11 +77,21 @@ void RigVMFindReferencesHelpers::ExpandAllChildren(FRigVMSearchResult InTreeNode
 // FBlueprintSearchResult
 
 FRigVMFindResult::FRigVMFindResult(TWeakObjectPtr<URigVMBlueprint> InBlueprint)
-	: WeakBlueprint(InBlueprint)
+	: WeakBlueprint(InBlueprint->GetObject())
 {
 }
 
 FRigVMFindResult::FRigVMFindResult(TWeakObjectPtr<URigVMBlueprint> InBlueprint, const FText& InDisplayText )
+	: WeakBlueprint(InBlueprint->GetObject()), DisplayText(InDisplayText)
+{
+}
+
+FRigVMFindResult::FRigVMFindResult(TWeakInterfacePtr<IRigVMAssetInterface> InBlueprint)
+	: WeakBlueprint(InBlueprint)
+{
+}
+
+FRigVMFindResult::FRigVMFindResult(TWeakInterfacePtr<IRigVMAssetInterface> InBlueprint, const FText& InDisplayText )
 	: WeakBlueprint(InBlueprint), DisplayText(InDisplayText)
 {
 }
@@ -96,9 +107,9 @@ FReply FRigVMFindResult::OnClick()
 	return FReply::Handled();
 }
 
-UObject* FRigVMFindResult::GetObject(URigVMBlueprint* InBlueprint) const
+UObject* FRigVMFindResult::GetObject(TRigVMAssetInterfacePtr InBlueprint) const
 {
-	return GetBlueprint();
+	return GetRigVMAsset()->GetObject();
 }
 
 FText FRigVMFindResult::GetCategory() const
@@ -118,13 +129,23 @@ TSharedRef<SWidget> FRigVMFindResult::CreateIcon() const
 
 URigVMBlueprint* FRigVMFindResult::GetBlueprint() const
 {
-	if (TStrongObjectPtr<URigVMBlueprint> Blueprint = WeakBlueprint.Pin())
+	return Cast<URigVMBlueprint>(GetRigVMAsset().GetObject());
+}
+
+TRigVMAssetInterfacePtr FRigVMFindResult::GetRigVMAsset()
+{
+	if (TStrongObjectPtr<UObject> Blueprint = WeakBlueprint.GetWeakObjectPtr().Pin())
 	{
-		return Blueprint.Get();
+		return TRigVMAssetInterfacePtr(Blueprint.Get());
 	}
 	return nullptr;
 }
 
+const TRigVMAssetInterfacePtr FRigVMFindResult::GetRigVMAsset() const
+{
+	return const_cast<FRigVMFindResult*>(this)->GetRigVMAsset();
+}
+
 FText FRigVMFindResult::GetDisplayString() const
 {
 	return DisplayText;
@@ -134,6 +155,12 @@ FText FRigVMFindResult::GetDisplayString() const
 // FRigVMFindReferencesGraphNode
 
 FRigVMFindReferencesGraphNode::FRigVMFindReferencesGraphNode(TWeakObjectPtr<URigVMBlueprint> InBlueprint)
+	: FRigVMFindResult(TRigVMAssetInterfacePtr(InBlueprint.Get())), Glyph(FAppStyle::GetAppStyleSetName(), "")
+	, Class(nullptr)
+{
+}
+
+FRigVMFindReferencesGraphNode::FRigVMFindReferencesGraphNode(TRigVMAssetInterfacePtr InBlueprint)
 	: FRigVMFindResult(InBlueprint), Glyph(FAppStyle::GetAppStyleSetName(), "")
 	, Class(nullptr)
 {
@@ -141,17 +168,17 @@ FRigVMFindReferencesGraphNode::FRigVMFindReferencesGraphNode(TWeakObjectPtr<URig
 
 FReply FRigVMFindReferencesGraphNode::OnClick()
 {
-	UBlueprint* Blueprint = GetBlueprint();
-	if(Blueprint)
+	TRigVMAssetInterfacePtr Blueprint = GetRigVMAsset();
+	if(Blueprint.GetObject())
 	{
 		UEdGraphNode* OutNode = NULL;
-		if(	UEdGraphNode* GraphNode = FBlueprintEditorUtils::GetNodeByGUID(Blueprint, NodeGuid) )
+		if(	UEdGraphNode* GraphNode = FRigVMBlueprintUtils::GetNodeByGUID(Blueprint, NodeGuid) )
 		{
-			TSharedPtr< IToolkit > FoundAssetEditor = FToolkitManager::Get().FindEditorForAsset(Blueprint);
+			TSharedPtr< IToolkit > FoundAssetEditor = FToolkitManager::Get().FindEditorForAsset(Blueprint->GetObject());
 			if (!FoundAssetEditor)
 			{
-				GEditor->GetEditorSubsystem<UAssetEditorSubsystem>()->OpenEditorForAsset(Blueprint);
-				FoundAssetEditor = FToolkitManager::Get().FindEditorForAsset(Blueprint);
+				GEditor->GetEditorSubsystem<UAssetEditorSubsystem>()->OpenEditorForAsset(Blueprint->GetObject());
+				FoundAssetEditor = FToolkitManager::Get().FindEditorForAsset(Blueprint->GetObject());
 			}
 
 			// If we found an Editor
@@ -227,9 +254,9 @@ void FRigVMFindReferencesGraphNode::FinalizeSearchData()
 	}
 }
 
-UObject* FRigVMFindReferencesGraphNode::GetObject(URigVMBlueprint* InBlueprint) const
+UObject* FRigVMFindReferencesGraphNode::GetObject(TRigVMAssetInterfacePtr InBlueprint) const
 {
-	return FBlueprintEditorUtils::GetNodeByGUID(InBlueprint, NodeGuid);
+	return FRigVMBlueprintUtils::GetNodeByGUID(InBlueprint, NodeGuid);
 }
 
 FText FRigVMFindReferencesGraphNode::GetCategory() const
@@ -241,6 +268,13 @@ FText FRigVMFindReferencesGraphNode::GetCategory() const
 // FRigVMFindReferencesPin
 
 FRigVMFindReferencesPin::FRigVMFindReferencesPin(TWeakObjectPtr<URigVMBlueprint> InBlueprint, FString InSchemaName)
+	: FRigVMFindResult(TRigVMAssetInterfacePtr(InBlueprint.Get()))
+	, SchemaName(InSchemaName)
+	, IconColor(FSlateColor::UseForeground())
+{
+}
+
+FRigVMFindReferencesPin::FRigVMFindReferencesPin(TRigVMAssetInterfacePtr InBlueprint, FString InSchemaName)
 	: FRigVMFindResult(InBlueprint)
 	, SchemaName(InSchemaName)
 	, IconColor(FSlateColor::UseForeground())
@@ -326,6 +360,11 @@ void FRigVMFindReferencesPin::FinalizeSearchData()
 // FRigVMFindReferencesVariable
 
 FRigVMFindReferencesVariable::FRigVMFindReferencesVariable(TWeakObjectPtr<URigVMBlueprint> InBlueprint)
+	: FRigVMFindResult(TRigVMAssetInterfacePtr(InBlueprint.Get()))
+{
+}
+
+FRigVMFindReferencesVariable::FRigVMFindReferencesVariable(TRigVMAssetInterfacePtr InBlueprint)
 	: FRigVMFindResult(InBlueprint)
 {
 }
@@ -383,17 +422,21 @@ void FRigVMFindReferencesVariable::FinalizeSearchData()
 // FRigVMFindReferencesGraph
 
 FRigVMFindReferencesGraph::FRigVMFindReferencesGraph(TWeakObjectPtr<URigVMBlueprint> InBlueprint, EGraphType InGraphType)
+	: FRigVMFindResult(TRigVMAssetInterfacePtr(InBlueprint.Get())), GraphType(InGraphType)
+{
+}
+FRigVMFindReferencesGraph::FRigVMFindReferencesGraph(TRigVMAssetInterfacePtr InBlueprint, EGraphType InGraphType)
 	: FRigVMFindResult(InBlueprint), GraphType(InGraphType)
 {
 }
 
 FReply FRigVMFindReferencesGraph::OnClick()
 {
-	UBlueprint* Blueprint = GetBlueprint();
-	if(Blueprint)
+	TRigVMAssetInterfacePtr Blueprint = GetRigVMAsset();
+	if(Blueprint.GetObject())
 	{
-		GEditor->GetEditorSubsystem<UAssetEditorSubsystem>()->OpenEditorForAsset(Blueprint);
-		TSharedPtr< IToolkit > FoundAssetEditor = FToolkitManager::Get().FindEditorForAsset(Blueprint);
+		GEditor->GetEditorSubsystem<UAssetEditorSubsystem>()->OpenEditorForAsset(Blueprint->GetObject());
+		TSharedPtr< IToolkit > FoundAssetEditor = FToolkitManager::Get().FindEditorForAsset(Blueprint->GetObject());
 
 		// If we found an Editor
 		if (FoundAssetEditor.IsValid())
@@ -401,7 +444,7 @@ FReply FRigVMFindReferencesGraph::OnClick()
 			TSharedPtr<FRigVMNewEditor> Editor = StaticCastSharedPtr<FRigVMNewEditor>(FoundAssetEditor);
 
 			TArray<UEdGraph*> BlueprintGraphs;
-			Blueprint->GetAllGraphs(BlueprintGraphs);
+			Blueprint->GetAllEdGraphs(BlueprintGraphs);
 
 			for(UEdGraph* Graph : BlueprintGraphs)
 			{
@@ -572,7 +615,7 @@ void SRigVMFindReferences::FindReferences(const FString& SearchTerms)
 		return;
 	}
 
-	URigVMBlueprint* Blueprint = Editor->GetRigVMBlueprint();
+	TRigVMAssetInterfacePtr Blueprint = Editor->GetRigVMAsset();
 	if (!Blueprint)
 	{
 		return;
@@ -634,7 +677,7 @@ void SRigVMFindReferences::FindReferences(const FString& SearchTerms)
 		{
 			UEdGraphNode* EdNode = EdPin->GetOwningNode();
 			
-			const FString SchemaClassName = Blueprint->GetRigVMEdGraphSchemaClass()->GetName();
+			const FString SchemaClassName = Blueprint->GetRigVMClientHost()->GetRigVMEdGraphSchemaClass()->GetName();
 			FRigVMSearchResult NodeResult = FindOrAddNodeResult(EdPin->GetOwningNode(), NodeMetaData);
 			PinResult = FRigVMSearchResult(new FRigVMFindReferencesPin(Blueprint, SchemaClassName));
 			ElementHashToResult.Add(GetTypeHash(EdPin), PinResult);
@@ -686,7 +729,7 @@ void SRigVMFindReferences::FindReferences(const FString& SearchTerms)
 		return VariableResult;
 	};
 
-	TArray<URigVMGraph*> Graphs = Blueprint->GetRigVMClient()->GetAllModels(true, true);
+	TArray<URigVMGraph*> Graphs = Blueprint->GetRigVMClientHost()->GetRigVMClient()->GetAllModels(true, true);
 	for (const URigVMGraph* Graph : Graphs)
 	{
 		if (URigVMEdGraph* EdGraph = Cast<URigVMEdGraph>(Blueprint->GetEdGraph(Graph)))
@@ -774,32 +817,34 @@ void SRigVMFindReferences::FindReferences(const FString& SearchTerms)
 		}
 	}
 
-	TArray<FRigVMGraphVariableDescription> BPVariables = Blueprint->GetMemberVariables();
+	TArray<FRigVMGraphVariableDescription> BPVariables = Blueprint->GetAssetVariables();
 	for (const FRigVMGraphVariableDescription& Variable : BPVariables)
 	{
-		TArray<UBlueprintExtension::FSearchTagDataPair> BPVariableMetadata;
-		Blueprint->AddVariableSearchMetaDataInfo(Variable.Name, BPVariableMetadata);
-		
-
-		const UBlueprintExtension::FSearchTagDataPair* MatchedVariablePair = nullptr;
-		for (const UBlueprintExtension::FSearchTagDataPair& MetadataTag : BPVariableMetadata)
+		if (TRigVMAssetInterfacePtr BP = Blueprint)
 		{
-			if (MetadataTag.Value.ToString().Contains(SearchTerms))
+			TArray<UBlueprintExtension::FSearchTagDataPair> BPVariableMetadata;
+			BP->AddVariableSearchMetaDataInfo(Variable.Name, BPVariableMetadata);
+
+			const UBlueprintExtension::FSearchTagDataPair* MatchedVariablePair = nullptr;
+			for (const UBlueprintExtension::FSearchTagDataPair& MetadataTag : BPVariableMetadata)
 			{
-				MatchedVariablePair = &MetadataTag;
-				break;
+				if (MetadataTag.Value.ToString().Contains(SearchTerms))
+				{
+					MatchedVariablePair = &MetadataTag;
+					break;
+				}
 			}
-		}
 
-		if (MatchedVariablePair)
-		{
-			FRigVMSearchResult Result = FindOrAddVariableResult(Variable.Name, nullptr, BPVariableMetadata);
+			if (MatchedVariablePair)
+			{
+				FRigVMSearchResult Result = FindOrAddVariableResult(Variable.Name, nullptr, BPVariableMetadata);
 
-			const FText DisplayText = FText::Format(LOCTEXT("RigVMFindReferencesValues", "{0}: {1}"), MatchedVariablePair->Key, MatchedVariablePair->Value);
-			FRigVMSearchResult DetailResult = FRigVMSearchResult(new FRigVMFindResult(Blueprint, DisplayText));
+				const FText DisplayText = FText::Format(LOCTEXT("RigVMFindReferencesValues", "{0}: {1}"), MatchedVariablePair->Key, MatchedVariablePair->Value);
+				FRigVMSearchResult DetailResult = FRigVMSearchResult(new FRigVMFindResult(Blueprint, DisplayText));
 
-			Result->Children.Add(DetailResult);
-			DetailResult->Parent = Result;
+				Result->Children.Add(DetailResult);
+				DetailResult->Parent = Result;
+			}
 		}
 	}
 
diff --git a/Engine/Plugins/Runtime/RigVM/Source/RigVMEditor/Private/Editor/RigVMGraphDetailCustomization.cpp b/Engine/Plugins/Runtime/RigVM/Source/RigVMEditor/Private/Editor/RigVMGraphDetailCustomization.cpp
--- a/Engine/Plugins/Runtime/RigVM/Source/RigVMEditor/Private/Editor/RigVMGraphDetailCustomization.cpp
+++ b/Engine/Plugins/Runtime/RigVM/Source/RigVMEditor/Private/Editor/RigVMGraphDetailCustomization.cpp
@@ -751,16 +751,16 @@ TSharedRef<IDetailCustomization> FRigVMGraphDetailCustomization::MakeInstance(TS
 	const TArray<UObject*>* Objects = (InEditor.IsValid() ? InEditor->GetObjectsCurrentlyBeingEdited() : nullptr);
 	if (Objects && Objects->Num() == 1)
 	{
-		if (URigVMBlueprint* RigVMBlueprint = Cast<URigVMBlueprint>((*Objects)[0]))
+		if (TRigVMAssetInterfacePtr RigVMBlueprint = (*Objects)[0])
 		{
-			if(RigVMBlueprint->GetClass() == InExpectedBlueprintClass)
+			if(RigVMBlueprint.GetObject()->GetClass() == InExpectedBlueprintClass)
 			{
 				return MakeShareable(new FRigVMGraphDetailCustomization(InEditor, RigVMBlueprint));
 			}
 		}
 	}
 
-	return MakeShareable(new FRigVMGraphDetailCustomization((TSharedPtr<IRigVMEditor>)nullptr, nullptr));
+	return MakeShareable(new FRigVMGraphDetailCustomization((TSharedPtr<IRigVMEditor>)nullptr, TRigVMAssetInterfacePtr(nullptr)));
 }
 
 #if WITH_RIGVMLEGACYEDITOR
@@ -769,19 +769,20 @@ TSharedPtr<IDetailCustomization> FRigVMGraphDetailCustomization::MakeLegacyInsta
 	const TArray<UObject*>* Objects = (InBlueprintEditor.IsValid() ? InBlueprintEditor->GetObjectsCurrentlyBeingEdited() : nullptr);
 	if (Objects && Objects->Num() == 1)
 	{
-		if (URigVMBlueprint* RigVMBlueprint = Cast<URigVMBlueprint>((*Objects)[0]))
+		if ((*Objects)[0]->Implements<URigVMAssetInterface>())
 		{
-			if(RigVMBlueprint->GetClass() == InExpectedBlueprintClass)
+			TRigVMAssetInterfacePtr RigVMBlueprint = (*Objects)[0];
+			if(RigVMBlueprint->GetObject()->GetClass() == InExpectedBlueprintClass)
 			{
 				return MakeShareable(new FRigVMGraphDetailCustomization(InBlueprintEditor, RigVMBlueprint));
 			}
 		}
 	}
 
-	return MakeShareable(new FRigVMGraphDetailCustomization((TSharedPtr<IRigVMEditor>)nullptr, nullptr));
+	return MakeShareable(new FRigVMGraphDetailCustomization((TSharedPtr<IRigVMEditor>)nullptr, TRigVMAssetInterfacePtr(nullptr)));
 }
 
-FRigVMGraphDetailCustomization::FRigVMGraphDetailCustomization(TSharedPtr<IBlueprintEditor> RigVMigEditor, URigVMBlueprint* RigVMBlueprint)
+FRigVMGraphDetailCustomization::FRigVMGraphDetailCustomization(TSharedPtr<IBlueprintEditor> RigVMigEditor, TRigVMAssetInterfacePtr RigVMBlueprint)
 	: RigVMBlueprintPtr(RigVMBlueprint)
 	, RigVMGraphDetailCustomizationImpl(MakeShared<FRigVMGraphDetailCustomizationImpl>())
 {
@@ -789,7 +790,7 @@ FRigVMGraphDetailCustomization::FRigVMGraphDetailCustomization(TSharedPtr<IBluep
 }
 #endif
 
-FRigVMGraphDetailCustomization::FRigVMGraphDetailCustomization(TSharedPtr<IRigVMEditor> RigVMigEditor, URigVMBlueprint* RigVMBlueprint)
+FRigVMGraphDetailCustomization::FRigVMGraphDetailCustomization(TSharedPtr<IRigVMEditor> RigVMigEditor, TRigVMAssetInterfacePtr RigVMBlueprint)
 	: RigVMEditorPtr(RigVMigEditor)
 	, RigVMBlueprintPtr(RigVMBlueprint)
 	, RigVMGraphDetailCustomizationImpl(MakeShared<FRigVMGraphDetailCustomizationImpl>())
@@ -804,7 +805,7 @@ void FRigVMGraphDetailCustomization::CustomizeDetails(IDetailLayoutBuilder& Deta
 	GraphPtr = CastChecked<URigVMEdGraph>(Objects[0].Get());
 	URigVMEdGraph* Graph = GraphPtr.Get();
 
-	URigVMBlueprint* Blueprint = RigVMBlueprintPtr.Get();
+	TRigVMAssetInterfacePtr Blueprint = RigVMBlueprintPtr.GetObject();
 	URigVMGraph* Model = nullptr;
 	URigVMController* Controller = nullptr;
 
@@ -824,8 +825,7 @@ void FRigVMGraphDetailCustomization::CustomizeDetails(IDetailLayoutBuilder& Deta
 		];
 		return;
 	}
-
-	RigVMGraphDetailCustomizationImpl->CustomizeDetails(DetailLayout, Model, Controller, Blueprint, RigVMEditorPtr);
+	RigVMGraphDetailCustomizationImpl->CustomizeDetails(DetailLayout, Model, Controller, TScriptInterface<IRigVMClientHost>(Blueprint->GetObject()).GetInterface(), RigVMEditorPtr);
 }
 
 void FRigVMGraphDetailCustomizationImpl::CustomizeDetails(IDetailLayoutBuilder& DetailLayout,
@@ -1462,7 +1462,7 @@ bool FRigVMGraphDetailCustomizationImpl::IsValidFunction() const
 		{
 			if (const URigVMLibraryNode* LibraryNode = Cast<URigVMLibraryNode>(Model->GetOuter()))
 			{
-				return LibraryNode->GetFunctionHeader(RigVMClientHost->GetRigVMGraphFunctionHost()).IsValid();
+				return LibraryNode->GetFunctionHeader(RigVMClientHost->GetRigVMGraphFunctionHost().GetInterface()).IsValid();
 			}
 		}
 	}
@@ -1477,7 +1477,7 @@ FRigVMVariant FRigVMGraphDetailCustomizationImpl::GetVariant() const
 		{
 			if (const URigVMLibraryNode* LibraryNode = Cast<URigVMLibraryNode>(Model->GetOuter()))
 			{
-				return LibraryNode->GetFunctionHeader(RigVMClientHost->GetRigVMGraphFunctionHost()).Variant;
+				return LibraryNode->GetFunctionHeader(RigVMClientHost->GetRigVMGraphFunctionHost().GetInterface()).Variant;
 			}
 		}
 	}
@@ -1507,7 +1507,7 @@ TArray<FRigVMVariantRef> FRigVMGraphDetailCustomizationImpl::GetVariantRefs() co
 		{
 			if (const URigVMLibraryNode* LibraryNode = Cast<URigVMLibraryNode>(Model->GetOuter()))
 			{
-				const FRigVMGraphFunctionHeader Header = LibraryNode->GetFunctionHeader(RigVMClientHost->GetRigVMGraphFunctionHost());
+				const FRigVMGraphFunctionHeader Header = LibraryNode->GetFunctionHeader(RigVMClientHost->GetRigVMGraphFunctionHost().GetInterface());
 				return Header.LibraryPointer.GetVariants(false);
 			}
 		}
@@ -1990,7 +1990,7 @@ void FRigVMWrappedNodeDetailCustomization::CustomizeDetails(IDetailLayoutBuilder
 		URigVMDetailsViewWrapperObject* WrapperObject = CastChecked<URigVMDetailsViewWrapperObject>(DetailObject.Get());
 		if(BlueprintBeingCustomized == nullptr)
 		{
-			BlueprintBeingCustomized = WrapperObject->GetTypedOuter<URigVMBlueprint>();
+			BlueprintBeingCustomized = IRigVMAssetInterface::GetInterfaceOuter(WrapperObject);
 		}
 
 		ObjectsBeingCustomized.Add(WrapperObject);
@@ -2541,7 +2541,7 @@ void FRigVMWrappedNodeDetailCustomization::CustomizeDetails(IDetailLayoutBuilder
 				[
 				SNew(SRigVMGraphVariableBinding)
 					.ModelPins(ModelPins)
-					.Blueprint(BlueprintBeingCustomized)
+					.Asset(BlueprintBeingCustomized)
 				];
 				continue;
 			}
@@ -2808,7 +2808,7 @@ void FRigVMWrappedNodeDetailCustomization::CustomizeLiveValues(IDetailLayoutBuil
 		return;
 	}
 
-	TSharedPtr<FRigVMParserAST> AST = FirstNode->GetGraph()->GetRuntimeAST(BlueprintBeingCustomized->VMCompileSettings.ASTSettings, false);
+	TSharedPtr<FRigVMParserAST> AST = FirstNode->GetGraph()->GetRuntimeAST(BlueprintBeingCustomized->GetVMCompileSettings().ASTSettings, false);
 	if(!AST.IsValid())
 	{
 		return;
@@ -3038,7 +3038,7 @@ void FRigVMGraphEnumDetailCustomization::CustomizeHeader(TSharedRef<IPropertyHan
 
 		if(BlueprintBeingCustomized == nullptr)
 		{
-			BlueprintBeingCustomized = Object->GetTypedOuter<URigVMBlueprint>();
+			BlueprintBeingCustomized = IRigVMAssetInterface::GetInterfaceOuter(Object);
 		}
 
 		if(GraphBeingCustomized == nullptr)
@@ -3517,7 +3517,7 @@ void FRigVMGraphMathTypeDetailCustomization::ConfigureTransformWidgetArgs(TShare
 		}
 		
 		URigVMController* Controller = nullptr;
-		if (TStrongObjectPtr<URigVMBlueprint> Blueprint = BlueprintBeingCustomized.Pin(); GraphBeingCustomized.IsValid())
+		if (TStrongObjectPtr<UObject> Blueprint = BlueprintBeingCustomized.GetWeakObjectPtr().Pin(); GraphBeingCustomized.IsValid())
 		{
 			Controller = BlueprintBeingCustomized->GetController(GraphBeingCustomized.Get());
 			if(bIsCommit)
@@ -3571,7 +3571,7 @@ void FRigVMGraphMathTypeDetailCustomization::ConfigureTransformWidgetArgs(TShare
 		}
 		
 		URigVMController* Controller = nullptr;
-		if (TStrongObjectPtr<URigVMBlueprint> Blueprint = BlueprintBeingCustomized.Pin(); GraphBeingCustomized.IsValid())
+		if (TStrongObjectPtr<UObject> Blueprint = BlueprintBeingCustomized.GetWeakObjectPtr().Pin(); GraphBeingCustomized.IsValid())
 		{
 			Controller = BlueprintBeingCustomized->GetController(GraphBeingCustomized.Get());
 			if(Controller)
@@ -3688,7 +3688,7 @@ void FRigVMGraphMathTypeDetailCustomization::ConfigureTransformWidgetArgs(TShare
 		}
 
 		URigVMController* Controller = nullptr;
-		if (TStrongObjectPtr<URigVMBlueprint> Blueprint = BlueprintBeingCustomized.Pin(); GraphBeingCustomized.IsValid())
+		if (TStrongObjectPtr<UObject> Blueprint = BlueprintBeingCustomized.GetWeakObjectPtr().Pin(); GraphBeingCustomized.IsValid())
 		{
 			Controller = BlueprintBeingCustomized->GetController(GraphBeingCustomized.Get());
 			Controller->OpenUndoBracket(FString::Printf(TEXT("Set %s"), *InPropertyHandle->GetProperty()->GetName()));
@@ -3805,7 +3805,7 @@ void FRigVMGraphMathTypeDetailCustomization::CustomizeHeader(TSharedRef<IPropert
 	{
 		if(BlueprintBeingCustomized == nullptr)
 		{
-			BlueprintBeingCustomized = Object->GetTypedOuter<URigVMBlueprint>();
+			BlueprintBeingCustomized = IRigVMAssetInterface::GetInterfaceOuter(Object);
 		}
 
 		if(GraphBeingCustomized == nullptr)
@@ -3863,7 +3863,7 @@ void FRigVMGraphMathTypeDetailCustomization::CustomizeChildren(TSharedRef<IPrope
 	{
 		if(BlueprintBeingCustomized == nullptr)
 		{
-			BlueprintBeingCustomized = Object->GetTypedOuter<URigVMBlueprint>();
+			BlueprintBeingCustomized = IRigVMAssetInterface::GetInterfaceOuter(Object);
 		}
 
 		if(GraphBeingCustomized == nullptr)
diff --git a/Engine/Plugins/Runtime/RigVM/Source/RigVMEditor/Private/Editor/RigVMLegacyEditor.cpp b/Engine/Plugins/Runtime/RigVM/Source/RigVMEditor/Private/Editor/RigVMLegacyEditor.cpp
--- a/Engine/Plugins/Runtime/RigVM/Source/RigVMEditor/Private/Editor/RigVMLegacyEditor.cpp
+++ b/Engine/Plugins/Runtime/RigVM/Source/RigVMEditor/Private/Editor/RigVMLegacyEditor.cpp
@@ -37,6 +37,16 @@ void FRigVMLegacyEditor::CreateEditorToolbar()
 	}
 }
 
+void FRigVMLegacyEditor::CommonInitialization(const TArray<TRigVMAssetInterfacePtr>& InitBlueprints, bool bShouldOpenInDefaultsMode)
+{
+	TArray<UBlueprint*> Blueprints;
+	Algo::Transform(InitBlueprints, Blueprints, [](TRigVMAssetInterfacePtr In)
+	{
+		return Cast<UBlueprint>(In->GetObject());
+	});
+	FBlueprintEditor::CommonInitialization(Blueprints, bShouldOpenInDefaultsMode);
+}
+
 TSharedPtr<FApplicationMode> FRigVMLegacyEditor::CreateEditorMode()
 {
 	return MakeShareable(new FRigVMLegacyEditorMode(StaticCastSharedRef<FRigVMLegacyEditor>(FBlueprintEditor::SharedThis<FBlueprintEditor>(this))));
diff --git a/Engine/Plugins/Runtime/RigVM/Source/RigVMEditor/Private/Editor/RigVMLegacyEditorMode.cpp b/Engine/Plugins/Runtime/RigVM/Source/RigVMEditor/Private/Editor/RigVMLegacyEditorMode.cpp
--- a/Engine/Plugins/Runtime/RigVM/Source/RigVMEditor/Private/Editor/RigVMLegacyEditorMode.cpp
+++ b/Engine/Plugins/Runtime/RigVM/Source/RigVMEditor/Private/Editor/RigVMLegacyEditorMode.cpp
@@ -13,7 +13,7 @@
 FRigVMLegacyEditorMode::FRigVMLegacyEditorMode(const TSharedRef<FRigVMLegacyEditor>& InRigVMEditor)
 	: FBlueprintEditorApplicationMode(InRigVMEditor, FRigVMEditorModes::RigVMEditorMode, FRigVMEditorModes::GetLocalizedMode, false, false)
 {
-	RigVMBlueprintPtr = CastChecked<URigVMBlueprint>(InRigVMEditor->GetRigVMBlueprint());
+	RigVMBlueprintPtr = CastChecked<URigVMBlueprint>(InRigVMEditor->GetRigVMAsset().GetObject());
 
 	TabFactories.RegisterFactory(MakeShared<FRigVMExecutionStackTabSummoner>(InRigVMEditor));
 	TabFactories.RegisterFactory(MakeShared<FRigVMEditorGraphExplorerTabSummoner>(InRigVMEditor));
diff --git a/Engine/Plugins/Runtime/RigVM/Source/RigVMEditor/Private/Editor/RigVMLocalVariableDetailCustomization.cpp b/Engine/Plugins/Runtime/RigVM/Source/RigVMEditor/Private/Editor/RigVMLocalVariableDetailCustomization.cpp
--- a/Engine/Plugins/Runtime/RigVM/Source/RigVMEditor/Private/Editor/RigVMLocalVariableDetailCustomization.cpp
+++ b/Engine/Plugins/Runtime/RigVM/Source/RigVMEditor/Private/Editor/RigVMLocalVariableDetailCustomization.cpp
@@ -10,7 +10,7 @@
 #include "Widgets/Input/SCheckBox.h"
 #include "Widgets/Input/STextComboBox.h"
 #include "RigVMBlueprintGeneratedClass.h"
-#include "RigVMBlueprint.h"
+#include "RigVMBlueprintLegacy.h"
 #include "RigVMCore/RigVM.h"
 #include "InstancedPropertyBagStructureDataProvider.h"
 
@@ -37,7 +37,7 @@ void FRigVMLocalVariableDetailCustomization::CustomizeDetails(IDetailLayoutBuild
 		}
 		if (ensure(GraphBeingCustomized))
 		{
-			BlueprintBeingCustomized = GraphBeingCustomized->GetTypedOuter<URigVMBlueprint>();
+			BlueprintBeingCustomized = IRigVMAssetInterface::GetInterfaceOuter(GraphBeingCustomized);
 		}
 
 		NameValidator = FRigVMLocalVariableNameValidator(BlueprintBeingCustomized, GraphBeingCustomized, VariableDescription.Name);
@@ -92,12 +92,11 @@ void FRigVMLocalVariableDetailCustomization::CustomizeDetails(IDetailLayoutBuild
 
 	if (BlueprintBeingCustomized)
 	{
-		URigVMBlueprintGeneratedClass* RigClass = BlueprintBeingCustomized->GetRigVMBlueprintGeneratedClass();
-		URigVMHost* CDO = Cast<URigVMHost>(RigClass->GetDefaultObject(true /* create if needed */));
-		if (CDO->GetVM() != nullptr)
+		URigVM* VM = BlueprintBeingCustomized->GetVM(false);
+		if (VM != nullptr)
 		{
 			FString SourcePath = FString::Printf(TEXT("LocalVariableDefault::%s|%s::Const"), *GraphBeingCustomized->GetGraphName(), *VariableDescription.Name.ToString());
-			FRigVMMemoryStorageStruct* LiteralMemory = CDO->GetVM()->GetLiteralMemory();
+			FRigVMMemoryStorageStruct* LiteralMemory = VM->GetLiteralMemory();
 			FProperty* Property = LiteralMemory->FindPropertyByName(*SourcePath);
 			if (Property)
 			{
diff --git a/Engine/Plugins/Runtime/RigVM/Source/RigVMEditor/Private/Editor/RigVMMinimalEnvironment.cpp b/Engine/Plugins/Runtime/RigVM/Source/RigVMEditor/Private/Editor/RigVMMinimalEnvironment.cpp
--- a/Engine/Plugins/Runtime/RigVM/Source/RigVMEditor/Private/Editor/RigVMMinimalEnvironment.cpp
+++ b/Engine/Plugins/Runtime/RigVM/Source/RigVMEditor/Private/Editor/RigVMMinimalEnvironment.cpp
@@ -1,7 +1,7 @@
 // Copyright Epic Games, Inc. All Rights Reserved.
 
 #include "Editor/RigVMMinimalEnvironment.h"
-#include "RigVMBlueprint.h"
+#include "RigVMBlueprintLegacy.h"
 #include "Editor/RigVMEditorTools.h"
 
 #define LOCTEXT_NAMESPACE "RigVMMinimalEnvironment"
@@ -74,7 +74,7 @@ void FRigVMMinimalEnvironment::SetFunctionNode(const FRigVMGraphFunctionIdentifi
 		{
 			if(const UClass* Class = AssetData.GetClass())
 			{
-				if(Class->IsChildOf(URigVMBlueprint::StaticClass()))
+				if(Class->ImplementsInterface(URigVMAssetInterface::StaticClass()))
 				{
 					SetSchemata(Class);
 				}
@@ -111,18 +111,18 @@ void FRigVMMinimalEnvironment::SetSchemata(const UClass* InRigVMBlueprintClass)
 {
 	check(InRigVMBlueprintClass);
 
-	const URigVMBlueprint* CDO = InRigVMBlueprintClass->GetDefaultObject<URigVMBlueprint>();
+	const TRigVMAssetInterfacePtr CDO = InRigVMBlueprintClass->GetDefaultObject();
 	check(CDO);
 
-	EdGraphClass = CDO->GetRigVMEdGraphClass();
-	EdGraphNodeClass = CDO->GetRigVMEdGraphNodeClass();
+	EdGraphClass = CDO->GetRigVMClientHost()->GetRigVMEdGraphClass();
+	EdGraphNodeClass = CDO->GetRigVMClientHost()->GetRigVMEdGraphNodeClass();
 
 	if(!EdGraph.IsValid() || EdGraph->GetClass() != EdGraphClass)
 	{
 		EdGraph = TStrongObjectPtr<URigVMEdGraph>(NewObject<URigVMEdGraph>(ModelGraph.Get(), EdGraphClass));
 	}
 
-	ModelController->SetSchemaClass(CDO->GetRigVMSchemaClass());
+	ModelController->SetSchemaClass(CDO->GetRigVMClientHost()->GetRigVMSchemaClass());
 	EdGraph->SetBlueprintClass(InRigVMBlueprintClass);
 }
 
