RigVM: Adding IRigVMAssetInterface as parent of RigVMBlueprint. First step to separation of RigVM and Blueprint.

diff --git a/Engine/Plugins/Runtime/RigVM/Source/RigVMEditor/Private/Editor/RigVMNewEditor.cpp b/Engine/Plugins/Runtime/RigVM/Source/RigVMEditor/Private/Editor/RigVMNewEditor.cpp
--- a/Engine/Plugins/Runtime/RigVM/Source/RigVMEditor/Private/Editor/RigVMNewEditor.cpp
+++ b/Engine/Plugins/Runtime/RigVM/Source/RigVMEditor/Private/Editor/RigVMNewEditor.cpp
@@ -29,7 +29,8 @@
 #include "GraphEditor.h"
 #include "GraphEditorActions.h"
 #include "MessageLogModule.h"
-#include "RigVMBlueprint.h"
+#include "RigVMBlueprintLegacy.h"
+#include "RigVMBlueprintUtils.h"
 #include "RigVMEditorCommands.h"
 #include "Editor/RigVMFindReferences.h"
 #include "Editor/RigVMCompilerResultsTabSummoner.h"
@@ -203,7 +204,7 @@ void FRigVMNewEditor::InitAssetEditor(const EToolkitMode::Type Mode, const TShar
 	FWorkflowCentricApplication::InitAssetEditor(Mode, InitToolkitHost, AppIdentifier, StandaloneDefaultLayout, bCreateDefaultStandaloneMenu, bCreateDefaultToolbar, ObjectsToEdit, bInIsToolbarFocusable, bInUseSmallToolbarIcons, InOpenMethod);
 }
 
-void FRigVMNewEditor::CommonInitialization(const TArray<UBlueprint*>& InitBlueprints, bool bShouldOpenInDefaultsMode)
+void FRigVMNewEditor::CommonInitialization(const TArray<TRigVMAssetInterfacePtr>& InitBlueprints, bool bShouldOpenInDefaultsMode)
 {
 	TSharedPtr<FRigVMNewEditor> ThisPtr = SharedThis(this);
 
@@ -258,16 +259,19 @@ void FRigVMNewEditor::CommonInitialization(const TArray<UBlueprint*>& InitBluepr
 
 		//FLoadObjectsFromAssetRegistryHelper::Load<UUserDefinedEnum>(UserDefinedEnumerators);
 
-		UBlueprint* InitBlueprint = InitBlueprints[0];
+		TRigVMAssetInterfacePtr InitBlueprint = InitBlueprints[0];
 
 		// Update the blueprint if required
-		EBlueprintStatus OldStatus = InitBlueprint->Status;
+		ERigVMAssetStatus OldStatus = InitBlueprint->GetAssetStatus();
 		//EnsureBlueprintIsUpToDate(InitBlueprint); // TODO sara-s
-		UPackage* BpPackage = InitBlueprint->GetOutermost();
-		bBlueprintModifiedOnOpen = (InitBlueprint->Status != OldStatus) && !BpPackage->HasAnyPackageFlags(PKG_NewlyCreated);
+		UPackage* BpPackage = InitBlueprint->GetObject()->GetOutermost();
+		bBlueprintModifiedOnOpen = (InitBlueprint->GetAssetStatus() != OldStatus) && !BpPackage->HasAnyPackageFlags(PKG_NewlyCreated);
 
 		// Flag the blueprint as having been opened
-		InitBlueprint->bIsNewlyCreated = false;
+		if (UBlueprint* Blueprint = Cast<UBlueprint>(InitBlueprint->GetObject()))
+		{
+			Blueprint->bIsNewlyCreated = false;
+		}
 
 		// When the blueprint that we are observing changes, it will notify this wrapper widget.
 		InitBlueprint->OnChanged().AddSP(this, &FRigVMNewEditor::OnBlueprintChanged); // TODO sara-s
@@ -291,15 +295,22 @@ void FRigVMNewEditor::CommonInitialization(const TArray<UBlueprint*>& InitBluepr
 	}
 }
 
-void FRigVMNewEditor::OnBlueprintChanged(UBlueprint* InBlueprint)
+void FRigVMNewEditor::OnBlueprintChanged(UObject* InBlueprint)
 {
 	if (InBlueprint)
 	{
 		// Notify that the blueprint has been changed (update Content browser, etc)
 		InBlueprint->PostEditChange();
 
-		// Call PostEditChange() on any Actors that are based on this Blueprint 
-		FBlueprintEditorUtils::PostEditChangeBlueprintActors(InBlueprint);
+		// Call PostEditChange() on any Actors that are based on this Blueprint
+		if (UBlueprint* Blueprint = Cast<UBlueprint>(InBlueprint))
+		{
+			FBlueprintEditorUtils::PostEditChangeBlueprintActors(Blueprint);
+		}
+		else
+		{
+			checkf(false, TEXT("RigVM asset is not a blueprint"));
+		}
 
 		// Refresh the graphs
 		//TODO sara-s
@@ -319,7 +330,7 @@ void FRigVMNewEditor::SaveEditedObjectState()
 	check(IsEditingSingleBlueprint());
 
 	// Clear currently edited documents
-	GetBlueprintObj()->LastEditedDocuments.Empty();
+	GetRigVMAsset()->GetLastEditedDocuments().Empty();
 
 	// Ask all open documents to save their state, which will update LastEditedDocuments
 	DocumentManager->SaveAllState();
@@ -339,7 +350,7 @@ void FRigVMNewEditor::SetCurrentMode(FName NewMode)
 void FRigVMNewEditor::PostLayoutBlueprintEditorInitialization()
 {
 	// TODO sara-s: Replicate what blueprint editor is doing
-	if (UBlueprint* Blueprint = GetBlueprintObj())
+	if (TRigVMAssetInterfacePtr Blueprint = GetRigVMAsset())
 	{
 		// Refresh the graphs
 		ForceEditorRefresh();
@@ -362,7 +373,7 @@ void FRigVMNewEditor::PostLayoutBlueprintEditorInitialization()
 
 			// Fire log message
 			FString BlueprintName;
-			Blueprint->GetName(BlueprintName);
+			Blueprint->GetObject()->GetName(BlueprintName);
 
 			FFormatNamedArguments Args;
 			Args.Add( TEXT("BlueprintName"), FText::FromString( BlueprintName ) );
@@ -473,7 +484,7 @@ void FRigVMNewEditor::OnGraphEditorFocused(const TSharedRef<SGraphEditor>& InGra
 
 		if ( FocusedGraph != nullptr )
 		{
-			if ( FBlueprintEditorUtils::FindBlueprintForGraph(FocusedGraph) )
+			if ( FRigVMBlueprintUtils::FindAssetForGraph(FocusedGraph) )
 			{
 				GraphExplorerWidget->Refresh();
 			}
@@ -512,19 +523,19 @@ void FRigVMNewEditor::OnGraphEditorBackgrounded(const TSharedRef<SGraphEditor>&
 FText FRigVMNewEditor::GetCompileStatusTooltip() const
 {
 	// Copied from FBlueprintEditorToolbar::GetStatusTooltip
-	
-	UBlueprint* BlueprintObj = GetBlueprintObj();
-	EBlueprintStatus Status = BlueprintObj->Status;
+
+	TRigVMAssetInterfacePtr BlueprintObj = GetRigVMAsset();
+	ERigVMAssetStatus Status = BlueprintObj->GetAssetStatus();
 
 	switch (Status)
 	{
-	case BS_Dirty:
+	case RVMA_Dirty:
 		return LOCTEXT("Dirty_Status", "Dirty; needs to be recompiled");
-	case BS_Error:
+	case RVMA_Error:
 		return LOCTEXT("CompileError_Status", "There was an error during compilation, see the log for details");
-	case BS_UpToDate:
+	case RVMA_UpToDate:
 		return LOCTEXT("GoodToGo_Status", "Good to go");
-	case BS_UpToDateWithWarnings:
+	case RVMA_UpToDateWithWarnings:
 		return LOCTEXT("GoodToGoWarning_Status", "There was a warning during compilation, see the log for details");
 	default:
 		return LOCTEXT("Recompile_Status", "Unknown status; should recompile");
@@ -534,15 +545,9 @@ FText FRigVMNewEditor::GetCompileStatusTooltip() const
 FSlateIcon FRigVMNewEditor::GetCompileStatusImage() const
 {
 	// Copied from FBlueprintEditorToolbar::GetStatusImage
-	
-	UBlueprint* BlueprintObj = GetBlueprintObj();
-	EBlueprintStatus Status = BlueprintObj->Status;
 
-	// For macro types, always show as up-to-date, since we don't compile them
-	if (BlueprintObj->BlueprintType == BPTYPE_MacroLibrary)
-	{
-		Status = BS_UpToDate;
-	}
+	TRigVMAssetInterfacePtr BlueprintObj = GetRigVMAsset();
+	ERigVMAssetStatus Status = BlueprintObj->GetAssetStatus();
 
 	
 	static const FName CompileStatusBackground("Blueprint.CompileStatus.Background");
@@ -553,11 +558,11 @@ FSlateIcon FRigVMNewEditor::GetCompileStatusImage() const
 
 	switch (Status)
 	{
-	case BS_Error:
+	case RVMA_Error:
 		return FSlateIcon(FAppStyle::GetAppStyleSetName(), CompileStatusBackground, NAME_None, CompileStatusError);
-	case BS_UpToDate:
+	case RVMA_UpToDate:
 		return FSlateIcon(FAppStyle::GetAppStyleSetName(), CompileStatusBackground, NAME_None, CompileStatusGood);
-	case BS_UpToDateWithWarnings:
+	case RVMA_UpToDateWithWarnings:
 		return FSlateIcon(FAppStyle::GetAppStyleSetName(), CompileStatusBackground, NAME_None, CompileStatusWarning);
 	default:
 		return FSlateIcon(FAppStyle::GetAppStyleSetName(), CompileStatusBackground, NAME_None, CompileStatusUnknown);
@@ -651,9 +656,9 @@ void FRigVMNewEditor::OnAddNewVariable()
 {
 	const FScopedTransaction Transaction( LOCTEXT("AddVariable", "Add Variable") );
 
-	FName VarName = FBlueprintEditorUtils::FindUniqueKismetName(GetBlueprintObj(), TEXT("NewVar"));
-
-	bool bSuccess = GraphExplorerWidget.IsValid() && FBlueprintEditorUtils::AddMemberVariable(GetBlueprintObj(), VarName, GraphExplorerWidget->GetLastPinTypeUsed());
+	FName VarName = FRigVMBlueprintUtils::FindUniqueVariableName(GetRigVMAsset().GetInterface(), TEXT("NewVar"));
+	
+	bool bSuccess = GraphExplorerWidget.IsValid() && GetRigVMAsset()->AddAssetVariableFromPinType(VarName, GraphExplorerWidget->GetLastPinTypeUsed()) != NAME_None;
 
 	if(!bSuccess)
 	{
@@ -688,34 +693,17 @@ void FRigVMNewEditor::ZoomToSelection_Clicked()
 
 void FRigVMNewEditor::RestoreEditedObjectState()
 {
-	check(IsEditingSingleBlueprint());	
-
-	UBlueprint* Blueprint = GetBlueprintObj();
-	if (Blueprint->LastEditedDocuments.Num() == 0)
-	{
-		if (FBlueprintEditorUtils::SupportsConstructionScript(Blueprint))
-		{
-			Blueprint->LastEditedDocuments.Add(FBlueprintEditorUtils::FindUserConstructionScript(Blueprint));
-		}
-
-		if (Blueprint->SupportsEventGraphs())
-		{
-			Blueprint->LastEditedDocuments.Add(FBlueprintEditorUtils::FindEventGraph(Blueprint));
-		}
-	}
+	check(IsEditingSingleBlueprint());
 
+	TRigVMAssetInterfacePtr Blueprint = GetRigVMAsset();
+	
 	TSet<FSoftObjectPath> PathsToRemove;
-	for (int32 i = 0; i < Blueprint->LastEditedDocuments.Num(); i++)
+	for (int32 i = 0; i < Blueprint->GetLastEditedDocuments().Num(); i++)
 	{
-		if (UObject* Obj = Blueprint->LastEditedDocuments[i].EditedObjectPath.ResolveObject())
+		if (UObject* Obj = Blueprint->GetLastEditedDocuments()[i].EditedObjectPath.ResolveObject())
 		{
 			if(UEdGraph* Graph = Cast<UEdGraph>(Obj))
 			{
-				if (FBlueprintEditorUtils::IsEventGraph(Graph) && !Blueprint->SupportsEventGraphs())
-				{
-					continue;
-				}
-
 				struct LocalStruct
 				{
 					static TSharedPtr<SDockTab> OpenGraphTree(FRigVMNewEditor* InBlueprintEditor, UEdGraph* InGraph)
@@ -745,7 +733,7 @@ void FRigVMNewEditor::RestoreEditedObjectState()
 				if (TabWithGraph.IsValid())
 				{
 					TSharedRef<SGraphEditor> GraphEditor = StaticCastSharedRef<SGraphEditor>(TabWithGraph->GetContent());
-					GraphEditor->SetViewLocation(Blueprint->LastEditedDocuments[i].SavedViewOffset, Blueprint->LastEditedDocuments[i].SavedZoomAmount);
+					GraphEditor->SetViewLocation(Blueprint->GetLastEditedDocuments()[i].SavedViewOffset, Blueprint->GetLastEditedDocuments()[i].SavedZoomAmount);
 				}
 			}
 			else
@@ -755,7 +743,7 @@ void FRigVMNewEditor::RestoreEditedObjectState()
 		}
 		else
 		{
-			PathsToRemove.Add(Blueprint->LastEditedDocuments[i].EditedObjectPath);
+			PathsToRemove.Add(Blueprint->GetLastEditedDocuments()[i].EditedObjectPath);
 		}
 	}
 
@@ -763,7 +751,7 @@ void FRigVMNewEditor::RestoreEditedObjectState()
 	// we'll check for that now. This is done to ensure we don't store invalid object paths indefinitely.
 	if (PathsToRemove.Num() > 0)
 	{
-		Blueprint->LastEditedDocuments.RemoveAll([&PathsToRemove](const FEditedDocumentInfo& Entry)
+		Blueprint->GetLastEditedDocuments().RemoveAll([&PathsToRemove](const FEditedDocumentInfo& Entry)
 		{
 			return PathsToRemove.Contains(Entry.EditedObjectPath);
 		});
@@ -840,7 +828,7 @@ void FRigVMNewEditor::SetUISelectionState(FName SelectionOwner)
 	}
 }
 
-void FRigVMNewEditor::AnalyticsTrackNodeEvent(UBlueprint* Blueprint, UEdGraphNode* GraphNode, bool bNodeDelete) const
+void FRigVMNewEditor::AnalyticsTrackNodeEvent(IRigVMAssetInterface* Blueprint, UEdGraphNode* GraphNode, bool bNodeDelete) const
 {
 	// TODO sara-s
 	// if(Blueprint && GraphNode && FEngineAnalytics::IsAvailable())
@@ -950,53 +938,6 @@ void FRigVMNewEditor::JumpToHyperlink(const UObject* ObjectReference, bool bRequ
 		// Open the document
 		OpenDocument(Graph, OpenMode);
 	}
-	else if (const AActor* ReferencedActor = Cast<const AActor>(ObjectReference))
-	{
-		// Check if the world is active in the editor. It's possible to open level BPs without formally opening
-		// the levels through Find-in-Blueprints
-		bool bInOpenWorld = false;
-		const TIndirectArray<FWorldContext>& WorldContextList = GEditor->GetWorldContexts();
-		const UWorld* ReferencedActorOwningWorld = ReferencedActor->GetWorld();
-		for (const FWorldContext& WorldContext : WorldContextList)
-		{
-			if (WorldContext.World() == ReferencedActorOwningWorld)
-			{
-				bInOpenWorld = true;
-				break;
-			}
-		}
-
-		// Clear the selection even if we couldn't find it, so the existing selection doesn't get mistaken for the desired to be selected actor
-		GEditor->SelectNone(false, false);
-
-		if (bInOpenWorld)
-		{
-			// Select the in-level actor
-			GEditor->SelectActor(const_cast<AActor*>(ReferencedActor), true, true, true);
-
-			// Point the camera at it
-			GUnrealEd->Exec(ReferencedActor->GetWorld(), TEXT("CAMERA ALIGN ACTIVEVIEWPORTONLY"));
-		}
-	}
-	else if(const UFunction* Function = Cast<const UFunction>(ObjectReference))
-	{
-		UBlueprint* BP = GetBlueprintObj();
-		if(BP)
-		{
-			if (UEdGraph* FunctionGraph = FBlueprintEditorUtils::FindScopeGraph(BP, Function))
-			{
-				OpenDocument(FunctionGraph, FDocumentTracker::OpenNewDocument);
-			}
-		}
-	}
-	else if(const UBlueprintGeneratedClass* Class = Cast<const UBlueprintGeneratedClass>(ObjectReference))
-	{
-		GEditor->GetEditorSubsystem<UAssetEditorSubsystem>()->OpenEditorForAsset(Class->ClassGeneratedBy);
-	}
-	// else if (const UTimelineTemplate* Timeline = Cast<const UTimelineTemplate>(ObjectReference))
-	// {
-	// 	OpenDocument(Timeline, FDocumentTracker::OpenNewDocument);
-	// }
 	else if ((ObjectReference != nullptr) && ObjectReference->IsAsset())
 	{
 		GEditor->GetEditorSubsystem<UAssetEditorSubsystem>()->OpenEditorForAsset(const_cast<UObject*>(ObjectReference));
@@ -1050,8 +991,8 @@ bool FRigVMNewEditor::IsSaveOnCompileEnabled() const
 {
 	// Copied from BlueprintEditorImpl::IsSaveOnCompileEnabled
 	
-	UBlueprint* Blueprint = GetBlueprintObj();
-	bool const bIsLevelScript = (Cast<ULevelScriptBlueprint>(Blueprint) != nullptr);
+	TRigVMAssetInterfacePtr Blueprint = GetRigVMAsset();
+	bool const bIsLevelScript = (Cast<ULevelScriptBlueprint>(Blueprint->GetObject()) != nullptr);
 
 	return !bIsLevelScript;
 }
@@ -1146,7 +1087,7 @@ UEdGraphNode* FRigVMNewEditor::FindNodeWithError(FCompilerResultsLog const& Erro
 
 void FRigVMNewEditor::CreateDefaultCommands()
 {
-	if (GetBlueprintObj())
+	if (GetRigVMAsset())
 	{
 		// TODO sara-s: Replicate what blueprint editor is doing
 		//FBlueprintEditor::CreateDefaultCommands();
@@ -1349,44 +1290,26 @@ void FRigVMNewEditor::CompileImpl()
 {
 	DECLARE_SCOPE_HIERARCHICAL_COUNTER_FUNC()
 
-	UBlueprint* BlueprintObj = GetBlueprintObj();
+	TRigVMAssetInterfacePtr BlueprintObj = GetRigVMAsset();
 	if (BlueprintObj)
 	{
 		FMessageLog BlueprintLog("BlueprintLog");
 
 		FFormatNamedArguments Arguments;
-		Arguments.Add(TEXT("BlueprintName"), FText::FromString(BlueprintObj->GetName()));
+		Arguments.Add(TEXT("BlueprintName"), FText::FromString(BlueprintObj->GetObject()->GetName()));
 		BlueprintLog.NewPage(FText::Format(LOCTEXT("CompilationPageLabel", "Compile {BlueprintName}"), Arguments));
 
-		FCompilerResultsLog LogResults;
-		LogResults.SetSourcePath(BlueprintObj->GetPathName());
-		LogResults.BeginEvent(TEXT("Compile"));
-		
-		// TODO: sara-s remove once blueprint backend is replaced
-		{
-			EBlueprintCompileOptions CompileOptions = EBlueprintCompileOptions::None;
-
-			// If compilation is enabled during PIE/simulation, references to the CDO might be held by a script variable.
-			// Thus, we set the flag to direct the compiler to allow those references to be replaced during reinstancing.
-			if (GEditor->PlayWorld != nullptr)
-			{
-				CompileOptions |= EBlueprintCompileOptions::IncludeCDOInReferenceReplacement;
-			}
-		
-			FKismetEditorUtilities::CompileBlueprint(BlueprintObj, CompileOptions, &LogResults);
-		}
-
-		LogResults.EndEvent();
+		FCompilerResultsLog LogResults = GetRigVMAsset()->CompileBlueprint();
 
-		// CachedNumWarnings = LogResults.NumWarnings;
-		// CachedNumErrors = LogResults.NumErrors;
+		CompilerResultsListing->ClearMessages();
 
-		const bool bForceMessageDisplay = (LogResults.NumWarnings > 0 || LogResults.NumErrors > 0) && !BlueprintObj->bIsRegeneratingOnLoad;
-		DumpMessagesToCompilerLog(LogResults.Messages, bForceMessageDisplay);
+		// Note we dont mirror to the output log here as the compiler already does that
+		CompilerResultsListing->AddMessages(LogResults.Messages, false);
 
-		if (BlueprintObj->UpgradeNotesLog.IsValid())
+		const bool bForceMessageDisplay = (LogResults.NumWarnings > 0 || LogResults.NumErrors > 0) && !BlueprintObj->IsRegeneratingOnLoad();
+		if (bForceMessageDisplay)
 		{
-			CompilerResultsListing->AddMessages(BlueprintObj->UpgradeNotesLog->Messages);
+			TabManager->TryInvokeTab(FRigVMCompilerResultsTabSummoner::TabID());
 		}
 
 		// send record when player clicks compile and send the result
@@ -1401,11 +1324,11 @@ bool FRigVMNewEditor::IsGraphInCurrentBlueprint(const UEdGraph* InGraph) const
 {
 	bool bEditable = true;
 
-	UBlueprint* EditingBP = GetBlueprintObj();
+	TRigVMAssetInterfacePtr EditingBP = GetRigVMAsset();
 	if(EditingBP)
 	{
 		TArray<UEdGraph*> Graphs;
-		EditingBP->GetAllGraphs(Graphs);
+		EditingBP->GetAllEdGraphs(Graphs);
 		bEditable &= Graphs.Contains(InGraph);
 	}
 
@@ -1467,16 +1390,16 @@ void FRigVMNewEditor::JumpToPin(const UEdGraphPin* Pin)
 	}
 }
 
-void FRigVMNewEditor::OnBlueprintChangedInnerImpl(UBlueprint* InBlueprint, bool bIsJustBeingCompiled)
+void FRigVMNewEditor::OnBlueprintChangedInnerImpl(IRigVMAssetInterface* InBlueprint, bool bIsJustBeingCompiled)
 {
 	if (InBlueprint)
 	{
 		// Notify that the blueprint has been changed (update Content browser, etc)
-		InBlueprint->PostEditChange();
-
-		// Call PostEditChange() on any Actors that are based on this Blueprint 
-		FBlueprintEditorUtils::PostEditChangeBlueprintActors(InBlueprint);
+		InBlueprint->GetObject()->PostEditChange();
 
+		// Call PostEditChange() on any Actors that are based on this Blueprint
+		InBlueprint->PostEditChangeBlueprintActors();
+		
 		// Refresh the graphs
 		ERefreshRigVMEditorReason::Type Reason = bIsJustBeingCompiled ? ERefreshRigVMEditorReason::BlueprintCompiled : ERefreshRigVMEditorReason::UnknownReason;
 		ForceEditorRefresh(Reason);
@@ -1489,7 +1412,7 @@ void FRigVMNewEditor::OnBlueprintChangedInnerImpl(UBlueprint* InBlueprint, bool
             	check(IsEditingSingleBlueprint());
             
             	// Clear currently edited documents
-            	GetBlueprintObj()->LastEditedDocuments.Empty();
+            	GetRigVMAsset()->GetLastEditedDocuments().Empty();
             
             	// Ask all open documents to save their state, which will update LastEditedDocuments
             	DocumentManager->SaveAllState();
@@ -1638,31 +1561,11 @@ FGraphAppearanceInfo FRigVMNewEditor::GetGraphAppearanceImpl(UEdGraph* InGraph)
 	// Create the appearance info
 	FGraphAppearanceInfo AppearanceInfo;
 
-	UBlueprint* Blueprint = (InGraph != nullptr) ? FBlueprintEditorUtils::FindBlueprintForGraph(InGraph) : GetBlueprintObj();
+	
+	auto Blueprint = (InGraph != nullptr) ? FRigVMBlueprintUtils::FindAssetForGraph(InGraph) : GetRigVMAsset();
 	if (Blueprint != nullptr)
 	{
-		if (FBlueprintEditorUtils::IsEditorUtilityBlueprint(Blueprint))
-		{
-			AppearanceInfo.CornerText = LOCTEXT("EditorUtilityAppearanceCornerText", "EDITOR UTILITY");
-		}
-		else
-		{
-			switch (Blueprint->BlueprintType)
-			{
-			case BPTYPE_LevelScript:
-				AppearanceInfo.CornerText = LOCTEXT("AppearanceCornerText_LevelScript", "LEVEL BLUEPRINT");
-				break;
-			case BPTYPE_MacroLibrary:
-				AppearanceInfo.CornerText = LOCTEXT("AppearanceCornerText_Macro", "MACRO");
-				break;
-			case BPTYPE_Interface:
-				AppearanceInfo.CornerText = LOCTEXT("AppearanceCornerText_Interface", "INTERFACE");
-				break;
-			default:
-				AppearanceInfo.CornerText = LOCTEXT("AppearanceCornerText_Blueprint", "BLUEPRINT");
-				break;
-			}
-		}
+		AppearanceInfo.CornerText = LOCTEXT("AppearanceCornerText_Blueprint", "BLUEPRINT");
 	}
 
 	UEdGraph const* EditingGraph = GetFocusedGraph();
@@ -1679,7 +1582,7 @@ FGraphAppearanceInfo FRigVMNewEditor::GetGraphAppearanceImpl(UEdGraph* InGraph)
 
 	// Copy from FBlueprintEditor::GetPIEStatus
 	{
-		UBlueprint* CurrentBlueprint = GetBlueprintObj();
+		TRigVMAssetInterfacePtr CurrentBlueprint = GetRigVMAsset();
 		UWorld *DebugWorld = nullptr;
 		ENetMode NetMode = NM_Standalone;
 		if (CurrentBlueprint)
@@ -1753,8 +1656,8 @@ void FRigVMNewEditor::NotifyPostChangeImpl(const FPropertyChangedEvent& Property
 	
 	if (IsEditingSingleBlueprint())
 	{
-		UBlueprint* Blueprint = GetBlueprintObj();
-		UPackage* BlueprintPackage = Blueprint->GetOutermost();
+		TRigVMAssetInterfacePtr Blueprint = GetRigVMAsset();
+		UPackage* BlueprintPackage = Blueprint->GetObject()->GetOutermost();
 
 		// if any of the objects being edited are in our package, mark us as dirty
 		bool bPropertyInBlueprint = false;
@@ -1773,10 +1676,10 @@ void FRigVMNewEditor::NotifyPostChangeImpl(const FPropertyChangedEvent& Property
 			// Note: if change type is "interactive," hold off on applying the change (e.g. this will occur if the user is scrubbing a spinbox value; we don't want to apply the change until the mouse is released, for performance reasons)
 			if (PropertyChangedEvent.ChangeType != EPropertyChangeType::Interactive)
 			{
-				FBlueprintEditorUtils::MarkBlueprintAsModified(Blueprint, PropertyChangedEvent);
+				Blueprint->MarkAssetAsModified(PropertyChangedEvent);
 
 				// Call PostEditChange() on any Actors that might be based on this Blueprint
-				FBlueprintEditorUtils::PostEditChangeBlueprintActors(Blueprint);
+				Blueprint->PostEditChangeBlueprintActors();
 			}
 
 			// TODO sara-s
@@ -1835,13 +1738,13 @@ void FRigVMNewEditor::EditGlobalOptions_Clicked()
 	}
 	else
 	{
-		UBlueprint* Blueprint = GetBlueprintObj();
+		TRigVMAssetInterfacePtr Blueprint = GetRigVMAsset();
 		if (Blueprint != nullptr)
 		{
 			// Show details for the Blueprint instance we're editing
 			if (Inspector.IsValid())
 			{
-				Inspector->ShowDetailsForSingleObject(Blueprint);
+				Inspector->ShowDetailsForSingleObject(Blueprint->GetObject());
 			}
 
 			TryInvokingDetailsTab();
@@ -1928,7 +1831,7 @@ void FRigVMNewEditor::StartEditingDefaults(bool bAutoFocus, bool bForceRefresh)
 
 	if (IsEditingSingleBlueprint())
 	{
-		if (GetBlueprintObj()->GeneratedClass != nullptr)
+		if (UObject* DefaultObject = GetRigVMAsset()->GetDefaultsObject())
 		{
 			// TODO sara-s
 			// if (SubobjectEditor.IsValid() && GetBlueprintObj()->GeneratedClass->IsChildOf<AActor>())
@@ -1937,8 +1840,6 @@ void FRigVMNewEditor::StartEditingDefaults(bool bAutoFocus, bool bForceRefresh)
 			// }
 			// else
 			{
-				UObject* DefaultObject = GetBlueprintObj()->GeneratedClass->GetDefaultObject();
-
 				// Update the details panel
 				FString Title;
 				DefaultObject->GetName(Title);
@@ -1967,16 +1868,6 @@ float FRigVMNewEditor::GetInstructionTextOpacity(UEdGraph* InGraph) const
 	if (InGraph)
 	{
 		bGraphReadOnly = !InGraph->bEditable;
-
-		if (!bGraphReadOnly)
-		{
-			const UBlueprint* BlueprintForGraph = FBlueprintEditorUtils::FindBlueprintForGraph(InGraph);
-			bool const bIsInterface = ((BlueprintForGraph != nullptr) && (BlueprintForGraph->BlueprintType == BPTYPE_Interface));
-			bool const bIsDelegate  = FBlueprintEditorUtils::IsDelegateSignatureGraph(InGraph);
-			bool const bIsMathExpression = FBlueprintEditorUtils::IsMathExpressionGraph(InGraph);
-
-			bGraphReadOnly = bIsInterface || bIsDelegate || bIsMathExpression;
-		}
 	}
 	
 	URigVMEditorSettings const* Settings = GetDefault<URigVMEditorSettings>();
@@ -2049,13 +1940,12 @@ void FRigVMNewEditor::DumpMessagesToCompilerLog(const TArray<TSharedRef<class FT
 	}
 }
 
-void FRigVMNewEditor::CreateDefaultTabContents(const TArray<UBlueprint*> InBlueprints)
+void FRigVMNewEditor::CreateDefaultTabContents(const TArray<TRigVMAssetInterfacePtr> InBlueprints)
 {
-	UBlueprint* InBlueprint = InBlueprints.Num() == 1 ? InBlueprints[0] : nullptr;
+	TRigVMAssetInterfacePtr InBlueprint = InBlueprints.Num() == 1 ? InBlueprints[0] : nullptr;
 
 	// Cache off whether or not this is an interface, since it is used to govern multiple widget's behavior
-	const bool bIsInterface = (InBlueprint && InBlueprint->BlueprintType == BPTYPE_Interface);
-	const bool bIsMacro = (InBlueprint && InBlueprint->BlueprintType == BPTYPE_MacroLibrary);
+	//const bool bIsInterface = (InBlueprint && InBlueprint->BlueprintType == BPTYPE_Interface);
 
 	if (InBlueprint)
 	{
@@ -2069,12 +1959,16 @@ void FRigVMNewEditor::CreateDefaultTabContents(const TArray<UBlueprint*> InBluep
 	{
 		//this->ReplaceReferencesWidget = SNew(SReplaceNodeReferences, SharedThis(this)); // TODO sara-s
 	}
-	
-	CompilerResultsListing = FCompilerResultsLog::GetBlueprintMessageLog(InBlueprint);
-	CompilerResultsListing->OnMessageTokenClicked().AddSP(this, &FRigVMNewEditor::OnLogTokenClicked); 
 
-	FMessageLogModule& MessageLogModule = FModuleManager::LoadModuleChecked<FMessageLogModule>("MessageLog");
-	CompilerResults = MessageLogModule.CreateLogListingWidget( CompilerResultsListing.ToSharedRef() );
+	if (UBlueprint* Blueprint = Cast<UBlueprint>(InBlueprint.GetObject()))
+	{
+		CompilerResultsListing = FCompilerResultsLog::GetBlueprintMessageLog(Blueprint);
+		CompilerResultsListing->OnMessageTokenClicked().AddSP(this, &FRigVMNewEditor::OnLogTokenClicked);
+
+		FMessageLogModule& MessageLogModule = FModuleManager::LoadModuleChecked<FMessageLogModule>("MessageLog");
+		CompilerResults = MessageLogModule.CreateLogListingWidget( CompilerResultsListing.ToSharedRef() );
+	} 
+
 	FindResults = SNew(SRigVMFindReferences, SharedThis(this));
 	
 	this->Inspector = 
@@ -2094,7 +1988,6 @@ void FRigVMNewEditor::CreateDefaultTabContents(const TArray<UBlueprint*> InBluep
 			SNew(SRigVMDetailsInspector)
 			.Editor(SharedThis(this))
 			. ViewIdentifier(FName("BlueprintDefaults"))
-			. IsEnabled(!bIsInterface)
 			//. ShowPublicViewControl(this, &FBlueprintEditor::ShouldShowPublicViewControl) // TODO sara-s
 			. ShowTitleArea(false)
 			. HideNameArea(bHideNameArea)
@@ -2102,13 +1995,13 @@ void FRigVMNewEditor::CreateDefaultTabContents(const TArray<UBlueprint*> InBluep
 		;
 	}
 
-	if (InBlueprint && 
-		InBlueprint->ParentClass &&
-		InBlueprint->ParentClass->IsChildOf(AActor::StaticClass()) && 
-		InBlueprint->SimpleConstructionScript )
-	{
-		//CreateSubobjectEditors(); // TODO sara-s
-	}
+	// if (InBlueprint && 
+	// 	InBlueprint->ParentClass &&
+	// 	InBlueprint->ParentClass->IsChildOf(AActor::StaticClass()) && 
+	// 	InBlueprint->SimpleConstructionScript )
+	// {
+	// 	//CreateSubobjectEditors(); // TODO sara-s
+	// }
 }
 
 const FSlateBrush* FRigVMNewEditor::GetGlyphForGraph(const UEdGraph* Graph, bool bInLargeIcon)
diff --git a/Engine/Plugins/Runtime/RigVM/Source/RigVMEditor/Private/Editor/RigVMVariableDetailCustomization.cpp b/Engine/Plugins/Runtime/RigVM/Source/RigVMEditor/Private/Editor/RigVMVariableDetailCustomization.cpp
--- a/Engine/Plugins/Runtime/RigVM/Source/RigVMEditor/Private/Editor/RigVMVariableDetailCustomization.cpp
+++ b/Engine/Plugins/Runtime/RigVM/Source/RigVMEditor/Private/Editor/RigVMVariableDetailCustomization.cpp
@@ -10,14 +10,13 @@
 #include "Widgets/Text/STextBlock.h"
 #include "Widgets/Input/SCheckBox.h"
 #include "Engine/BlueprintGeneratedClass.h"
-#include "RigVMBlueprint.h"
+#include "RigVMBlueprintLegacy.h"
 #include "Editor/RigVMEditor.h"
 #if WITH_RIGVMLEGACYEDITOR
 #include "Editor/RigVMLegacyEditor.h"
 #endif
 #include "RigVMModel/RigVMController.h"
 #include "Widgets/Input/SMultiLineEditableTextBox.h"
-#include "RigVMBlueprint.h"
 #include "SPinTypeSelector.h"
 #include "SupportedRangeTypes.h"
 #include "Editor/SRigVMDetailsInspector.h"
@@ -33,19 +32,16 @@ TSharedRef<IDetailCustomization> FRigVMVariableDetailCustomization::MakeInstance
 	const TArray<UObject*>* Objects = (InEditor.IsValid() ? InEditor->GetObjectsCurrentlyBeingEdited() : nullptr);
 	if (Objects && Objects->Num() == 1)
 	{
-		if (URigVMBlueprint* Blueprint = Cast<URigVMBlueprint>((*Objects)[0]))
+		if (TRigVMAssetInterfacePtr Blueprint = (*Objects)[0])
 		{
-			if (Blueprint->ParentClass && Blueprint->ParentClass->IsChildOf(URigVMHost::StaticClass()))
-			{
-				return MakeShareable(new FRigVMVariableDetailCustomization(InEditor, Blueprint));
-			}
+			return MakeShareable(new FRigVMVariableDetailCustomization(InEditor, Blueprint));
 		}
 	}
 
-	return MakeShareable(new FRigVMVariableDetailCustomization((TSharedPtr<IRigVMEditor>)nullptr, nullptr));
+	return MakeShareable(new FRigVMVariableDetailCustomization((TSharedPtr<IRigVMEditor>)nullptr, TRigVMAssetInterfacePtr(nullptr)));
 }
 
-FRigVMVariableDetailCustomization::FRigVMVariableDetailCustomization(TSharedPtr<IRigVMEditor> InEditor, URigVMBlueprint* Blueprint)
+FRigVMVariableDetailCustomization::FRigVMVariableDetailCustomization(TSharedPtr<IRigVMEditor> InEditor, TRigVMAssetInterfacePtr Blueprint)
 		: EditorPtr(InEditor)
 		, BlueprintPtr(Blueprint)
 {}
@@ -86,7 +82,7 @@ void FRigVMVariableDetailCustomization::CustomizeDetails(IDetailLayoutBuilder& D
 	if (Editor->GetRigVMInspector())
 	{
 		FName VariableName = GetVariableName();
-		CachedVariableProperty = BlueprintPtr->SkeletonGeneratedClass->FindPropertyByName(VariableName);
+		CachedVariableProperty = BlueprintPtr->FindGeneratedPropertyByName(VariableName);
 
 		if(!CachedVariableProperty.IsValid())
 		{
@@ -95,23 +91,24 @@ void FRigVMVariableDetailCustomization::CustomizeDetails(IDetailLayoutBuilder& D
 
 		Editor->OnRefresh().AddSP(this, &FRigVMVariableDetailCustomization::OnPostEditorRefresh);
 
-		// Get an appropiate name validator
+		// Get an appropriate name validator
 		TSharedPtr<INameValidatorInterface> NameValidator = nullptr;
 		{
 			const UEdGraphSchema* Schema = nullptr;
 			if (BlueprintPtr.IsValid())
 			{
 				TArray<UEdGraph*> Graphs;
-				BlueprintPtr->GetAllGraphs(Graphs);
+				BlueprintPtr->GetAllEdGraphs(Graphs);
 				if (Graphs.Num() > 0)
 				{
 					Schema = Graphs[0]->GetSchema();
 				}
 			}
 
-			if (Schema)
+			if (const URigVMEdGraphSchema* RigVMSchema = Cast<URigVMEdGraphSchema>(Schema))
 			{
-				NameValidator = Schema->GetNameValidator(BlueprintPtr.Get(), CachedVariableName, nullptr, FEdGraphSchemaAction_K2Var::StaticGetTypeId());
+				TRigVMAssetInterfacePtr AssetPtr = BlueprintPtr.GetObject();
+				NameValidator = RigVMSchema->GetNameValidator(AssetPtr, CachedVariableName, nullptr);
 			}
 		}
 
@@ -317,7 +314,7 @@ void FRigVMVariableDetailCustomization::CustomizeDetails(IDetailLayoutBuilder& D
 								OutErrorMessage = LOCTEXT("CategoryEmpty", "Cannot add a category with an empty string.");
 								return false;
 							}
-							if (InNewText.EqualTo(FText::FromString(BlueprintPtr->GetName())))
+							if (InNewText.EqualTo(FText::FromString(BlueprintPtr->GetObject()->GetName())))
 							{
 								OutErrorMessage = LOCTEXT("CategoryEqualsBlueprintName", "Cannot add a category with the same name as the blueprint.");
 								return false;
@@ -433,94 +430,102 @@ void FRigVMVariableDetailCustomization::CustomizeDetails(IDetailLayoutBuilder& D
 			]
 		];
 
-		TStrongObjectPtr<URigVMBlueprint> Blueprint = BlueprintPtr.Pin();
-
-		// Add in default value editing for properties that can be edited, local properties cannot be edited
-		if (Blueprint.IsValid() && (Blueprint->GeneratedClass != nullptr))
+		if (TStrongObjectPtr<UObject> StrongObj = BlueprintPtr.GetWeakObjectPtr().Pin())
 		{
-			bool bVariableRenamed = false;
-			if (VariableProperty != nullptr)
+			URigVMBlueprint* Blueprint = Cast<URigVMBlueprint>(StrongObj.Get());
+			if (!Blueprint)
+			{
+				checkf(false, TEXT("RigVM asset is not a blueprint"));
+			}
+			
+			// Add in default value editing for properties that can be edited, local properties cannot be edited
+			if (Blueprint && Blueprint->GeneratedClass != nullptr)
 			{
-				const int32 VarIndex = FBlueprintEditorUtils::FindNewVariableIndex(Blueprint.Get(), CachedVariableName);
-				if (VarIndex != INDEX_NONE)
+				bool bVariableRenamed = false;
+				if (VariableProperty != nullptr)
 				{
-					const FGuid VarGuid = Blueprint->NewVariables[VarIndex].VarGuid;
-					if (UBlueprintGeneratedClass* AuthoritiveBPGC = Cast<UBlueprintGeneratedClass>(Blueprint->GeneratedClass))
+					const int32 VarIndex = FBlueprintEditorUtils::FindNewVariableIndex(Blueprint, CachedVariableName);
+					if (VarIndex != INDEX_NONE)
 					{
-						if (const FName* OldName = AuthoritiveBPGC->PropertyGuids.FindKey(VarGuid))
+						const FGuid VarGuid = Blueprint->NewVariables[VarIndex].VarGuid;
+						if (UBlueprintGeneratedClass* AuthoritiveBPGC = Cast<UBlueprintGeneratedClass>(Blueprint->GeneratedClass))
 						{
-							bVariableRenamed = CachedVariableName != *OldName;
+							if (const FName* OldName = AuthoritiveBPGC->PropertyGuids.FindKey(VarGuid))
+							{
+								bVariableRenamed = CachedVariableName != *OldName;
+							}
 						}
 					}
-				}
 				
-				const FProperty* OriginalProperty = nullptr;
-				OriginalProperty = FindFProperty<FProperty>(Blueprint->GeneratedClass, VariableProperty->GetFName());
+					const FProperty* OriginalProperty = nullptr;
+					OriginalProperty = FindFProperty<FProperty>(Blueprint->GeneratedClass, VariableProperty->GetFName());
 
-				if (OriginalProperty == nullptr || bVariableRenamed)
-				{
-					// Prevent editing the default value of a skeleton property
-					VariableProperty = nullptr;
-				}
-				else if (const FStructProperty* StructProperty = CastField<const FStructProperty>(OriginalProperty))
-				{
-					// Prevent editing the default value of a stale struct
-					const UUserDefinedStruct* BGStruct = Cast<const UUserDefinedStruct>(StructProperty->Struct);
-					if (BGStruct && (EUserDefinedStructureStatus::UDSS_UpToDate != BGStruct->Status))
+					if (OriginalProperty == nullptr || bVariableRenamed)
 					{
+						// Prevent editing the default value of a skeleton property
 						VariableProperty = nullptr;
 					}
+					else if (const FStructProperty* StructProperty = CastField<const FStructProperty>(OriginalProperty))
+					{
+						// Prevent editing the default value of a stale struct
+						const UUserDefinedStruct* BGStruct = Cast<const UUserDefinedStruct>(StructProperty->Struct);
+						if (BGStruct && (EUserDefinedStructureStatus::UDSS_UpToDate != BGStruct->Status))
+						{
+							VariableProperty = nullptr;
+						}
+					}
 				}
-			}
 
-			// Find the class containing the variable
-			UClass* VariableClass = (VariableProperty ? VariableProperty->GetTypedOwner<UClass>() : nullptr);
+				// Find the class containing the variable
+				UClass* VariableClass = (VariableProperty ? VariableProperty->GetTypedOwner<UClass>() : nullptr);
 
-			FText ErrorMessage;
-			IDetailCategoryBuilder& DefaultValueCategory = DetailLayout.EditCategory(TEXT("DefaultValueCategory"), LOCTEXT("DefaultValueCategoryHeading", "Default Value"));
+				FText ErrorMessage;
+				IDetailCategoryBuilder& DefaultValueCategory = DetailLayout.EditCategory(TEXT("DefaultValueCategory"), LOCTEXT("DefaultValueCategoryHeading", "Default Value"));
 
-			if (VariableProperty == nullptr)
-			{
-				if (Blueprint->Status != BS_UpToDate)
+				if (VariableProperty == nullptr)
 				{
-					ErrorMessage = LOCTEXT("VariableMissing_DirtyBlueprint", "Please compile the blueprint");
+					if (Blueprint->Status != BS_UpToDate)
+					{
+						ErrorMessage = LOCTEXT("VariableMissing_DirtyBlueprint", "Please compile the blueprint");
+					}
+					else
+					{
+						ErrorMessage = LOCTEXT("VariableMissing_CleanBlueprint", "Failed to find variable property");
+					}
 				}
-				else
+				// Show the error message if something went wrong
+				if (!ErrorMessage.IsEmpty())
 				{
-					ErrorMessage = LOCTEXT("VariableMissing_CleanBlueprint", "Failed to find variable property");
+					DefaultValueCategory.AddCustomRow( ErrorMessage )
+					[
+						SNew(STextBlock)
+						.ToolTipText(ErrorMessage)
+						.Text(ErrorMessage)
+						.Font(DetailFontInfo)
+					];
 				}
-			}
-			// Show the error message if something went wrong
-			if (!ErrorMessage.IsEmpty())
-			{
-				DefaultValueCategory.AddCustomRow( ErrorMessage )
-				[
-					SNew(STextBlock)
-					.ToolTipText(ErrorMessage)
-					.Text(ErrorMessage)
-					.Font(DetailFontInfo)
-				];
-			}
-			else 
-			{
-				TSharedPtr<IDetailsView> DetailsView;
-				if (Editor.IsValid())
+				else 
 				{
-					DetailsView = Editor->GetRigVMInspector()->GetPropertyView();
-				}
+					TSharedPtr<IDetailsView> DetailsView;
+					if (Editor.IsValid())
+					{
+						DetailsView = Editor->GetRigVMInspector()->GetPropertyView();
+					}
 
-				{
-					// Things are in order, show the property and allow it to be edited
-					TArray<UObject*> ObjectList;
-					ObjectList.Add(Blueprint->GeneratedClass->GetDefaultObject());
-					IDetailPropertyRow* Row = DefaultValueCategory.AddExternalObjectProperty(ObjectList, VariableProperty->GetFName());
-
-					// if (DetailsView.IsValid())
-					// {
-					// 	DetailsView->OnFinishedChangingProperties().AddSP(this, &FRigVMVariableDetailCustomization::OnFinishedChangingVariable);
-					// }
+					{
+						// Things are in order, show the property and allow it to be edited
+						TArray<UObject*> ObjectList;
+						ObjectList.Add(Blueprint->GeneratedClass->GetDefaultObject());
+						IDetailPropertyRow* Row = DefaultValueCategory.AddExternalObjectProperty(ObjectList, VariableProperty->GetFName());
+
+						// if (DetailsView.IsValid())
+						// {
+						// 	DetailsView->OnFinishedChangingProperties().AddSP(this, &FRigVMVariableDetailCustomization::OnFinishedChangingVariable);
+						// }
+					}
 				}
 			}
+			
 		}
 	}
 }
@@ -528,12 +533,13 @@ void FRigVMVariableDetailCustomization::CustomizeDetails(IDetailLayoutBuilder& D
 void FRigVMVariableDetailCustomization::PopulateCategories()
 {
 	TSharedPtr<IRigVMEditor> Editor = EditorPtr.Pin();
-	TStrongObjectPtr<URigVMBlueprint> Blueprint = BlueprintPtr.Pin();
+	TStrongObjectPtr<UObject> Blueprint = BlueprintPtr.GetWeakObjectPtr().Pin();
 	
 	if (!Blueprint || !Editor)
 	{
 		return;
 	}
+	TRigVMAssetInterfacePtr Interface = Blueprint.Get();
 	
 	auto IsNewCategorySource = [&](const FText& NewCategory)
 	{
@@ -547,12 +553,12 @@ void FRigVMVariableDetailCustomization::PopulateCategories()
 	CategorySource.Add(MakeShared<FText>(UEdGraphSchema_K2::VR_DefaultCategory));
 
 
-	TArray<FRigVMGraphVariableDescription> Variables = Blueprint->GetMemberVariables();
+	TArray<FRigVMGraphVariableDescription> Variables = Interface->GetAssetVariables();
 	
 	for (const FRigVMGraphVariableDescription& Variable : Variables)
 	{
 		FText Category = Variable.Category;
-		if (!Category.IsEmpty() && !Category.EqualTo(FText::FromString(BlueprintPtr.Pin()->GetName())))
+		if (!Category.IsEmpty() && !Category.EqualTo(FText::FromString(Interface->GetObject()->GetName())))
 		{
 			if (IsNewCategorySource(Category))
 			{
@@ -575,7 +581,7 @@ void FRigVMVariableDetailCustomization::PopulateCategories()
 void FRigVMVariableDetailCustomization::OnPostEditorRefresh()
 {
 	FName VariableName = GetVariableName();
-	CachedVariableProperty = BlueprintPtr->SkeletonGeneratedClass->FindPropertyByName(VariableName);
+	CachedVariableProperty = BlueprintPtr->FindGeneratedPropertyByName(VariableName);
 }
 
 FName FRigVMVariableDetailCustomization::GetVariableName() const
@@ -603,7 +609,7 @@ void FRigVMVariableDetailCustomization::OnVarNameCommitted(const FText& InNewNam
 		return;
 	}
 
-	TStrongObjectPtr<URigVMBlueprint> Blueprint = BlueprintPtr.Pin();
+	TStrongObjectPtr<UObject> Blueprint = BlueprintPtr.GetWeakObjectPtr().Pin();
 	if (!Blueprint)
 	{
 		return;
@@ -618,8 +624,9 @@ void FRigVMVariableDetailCustomization::OnVarNameCommitted(const FText& InNewNam
 	ERigVMExplorerElementType::Type SelectedType = Explorer->GetSelectedType();
 	if (SelectedType == ERigVMExplorerElementType::Variable)
 	{
+		TRigVMAssetInterfacePtr Interface = Blueprint.Get();
 		Blueprint->Modify();
-		FBlueprintEditorUtils::RenameMemberVariable(Blueprint.Get(), CachedVariableName, *InNewName.ToString());
+		Interface->RenameMemberVariable(CachedVariableName, *InNewName.ToString());
 	}
 }
 
@@ -632,11 +639,12 @@ FEdGraphPinType FRigVMVariableDetailCustomization::OnGetVarType() const
 		return Type;
 	}
 
-	TStrongObjectPtr<URigVMBlueprint> Blueprint = BlueprintPtr.Pin();
+	TStrongObjectPtr<UObject> Blueprint = BlueprintPtr.GetWeakObjectPtr().Pin();
 	if (!Blueprint)
 	{
 		return Type;
 	}
+	TRigVMAssetInterfacePtr Interface = Blueprint.Get();
 
 	TSharedPtr<SRigVMEditorGraphExplorer> Explorer = Editor->GetGraphExplorerWidget();
 	if (!Explorer)
@@ -647,11 +655,11 @@ FEdGraphPinType FRigVMVariableDetailCustomization::OnGetVarType() const
 	ERigVMExplorerElementType::Type SelectionVarType = Explorer->GetSelectedType();
 	if (SelectionVarType == ERigVMExplorerElementType::Variable)
 	{
-		for (FBPVariableDescription& Var : Blueprint->NewVariables)
+		for (FRigVMGraphVariableDescription& Var : Interface->GetAssetVariables())
 		{
-			if (Var.VarName == CachedVariableName)
+			if (Var.Name == CachedVariableName)
 			{
-				return Var.VarType;
+				return Var.ToPinType();
 			}
 		}
 		return Type;
@@ -667,11 +675,12 @@ void FRigVMVariableDetailCustomization::OnVarTypeChanged(const FEdGraphPinType&
 		return;
 	}
 
-	TStrongObjectPtr<URigVMBlueprint> Blueprint = BlueprintPtr.Pin();
+	TStrongObjectPtr<UObject> Blueprint = BlueprintPtr.GetWeakObjectPtr().Pin();
 	if (!Blueprint)
 	{
 		return;
 	}
+	TRigVMAssetInterfacePtr Interface = Blueprint.Get();
 
 	TSharedPtr<SRigVMEditorGraphExplorer> Explorer = Editor->GetGraphExplorerWidget();
 	if (!Explorer)
@@ -682,11 +691,11 @@ void FRigVMVariableDetailCustomization::OnVarTypeChanged(const FEdGraphPinType&
 	ERigVMExplorerElementType::Type SelectedType = Explorer->GetSelectedType();
 	if (SelectedType == ERigVMExplorerElementType::Variable)
 	{
-		for (FRigVMGraphVariableDescription& Variable : Blueprint->GetMemberVariables())
+		for (FRigVMGraphVariableDescription& Variable : Interface->GetAssetVariables())
 		{
 			if (Variable.Name == CachedVariableName)
 			{
-				FBlueprintEditorUtils::ChangeMemberVariableType(Blueprint.Get(), CachedVariableName, NewPinType);
+				Interface->ChangeMemberVariableType(CachedVariableName, NewPinType);
 				Explorer->SetLastPinTypeUsed(NewPinType);
 				return;
 			}
@@ -738,11 +747,12 @@ FText FRigVMVariableDetailCustomization::OnGetTooltipText() const
 		return ToolTip;
 	}
 
-	TStrongObjectPtr<URigVMBlueprint> Blueprint = BlueprintPtr.Pin();
+	TStrongObjectPtr<UObject> Blueprint = BlueprintPtr.GetWeakObjectPtr().Pin();
 	if (!Blueprint)
 	{
 		return ToolTip;
 	}
+	TRigVMAssetInterfacePtr Interface = Blueprint.Get();
 
 	TSharedPtr<SRigVMEditorGraphExplorer> Explorer = Editor->GetGraphExplorerWidget();
 	if (!Explorer)
@@ -755,8 +765,13 @@ FText FRigVMVariableDetailCustomization::OnGetTooltipText() const
 	if (SelectedType == ERigVMExplorerElementType::Variable)
 	{
 		FString Result;
-		FBlueprintEditorUtils::GetBlueprintVariableMetaData(Blueprint.Get(), CachedVariableName, nullptr, TEXT("tooltip"), Result);
-		return FText::FromString(Result);
+		for (FRigVMGraphVariableDescription& Var : Interface->GetAssetVariables())
+		{
+			if (Var.Name == CachedVariableName)
+			{
+				return Var.Tooltip;
+			}
+		}
 	}
 
 	return ToolTip;
@@ -770,11 +785,12 @@ void FRigVMVariableDetailCustomization::OnTooltipTextCommitted(const FText& NewT
 		return;
 	}
 
-	TStrongObjectPtr<URigVMBlueprint> Blueprint = BlueprintPtr.Pin();
+	TStrongObjectPtr<UObject> Blueprint = BlueprintPtr.GetWeakObjectPtr().Pin();
 	if (!Blueprint)
 	{
 		return;
 	}
+	TRigVMAssetInterfacePtr Interface = Blueprint.Get();
 
 	TSharedPtr<SRigVMEditorGraphExplorer> Explorer = Editor->GetGraphExplorerWidget();
 	if (!Explorer)
@@ -786,7 +802,7 @@ void FRigVMVariableDetailCustomization::OnTooltipTextCommitted(const FText& NewT
 
 	if (SelectedType == ERigVMExplorerElementType::Variable)
 	{
-		FBlueprintEditorUtils::SetBlueprintVariableMetaData(Blueprint.Get(), CachedVariableName, nullptr, TEXT("tooltip"), NewText.ToString() );
+		Interface->SetVariableTooltip(CachedVariableName, NewText);
 	}
 }
 
@@ -819,15 +835,24 @@ FText FRigVMVariableDetailCustomization::OnGetCategoryText() const
 		return FText();
 	}
 
-	TStrongObjectPtr<URigVMBlueprint> Blueprint = BlueprintPtr.Pin();
+	TStrongObjectPtr<UObject> Blueprint = BlueprintPtr.GetWeakObjectPtr().Pin();
 	if (!Blueprint)
 	{
 		return FText();
 	}
+	TRigVMAssetInterfacePtr Interface = Blueprint.Get();
 
 	if (CachedVariableName != NAME_None)
 	{
-		FText Category = FBlueprintEditorUtils::GetBlueprintVariableCategory(Blueprint.Get(), CachedVariableName, nullptr);
+		FText Category;
+		for (FRigVMGraphVariableDescription& Variable : Interface->GetAssetVariables())
+		{
+			if (Variable.Name == CachedVariableName)
+			{
+				Category = Variable.Category;
+				break;
+			}
+		}
 
 		// Older blueprints will have their name as the default category and whenever it is the same as the default category, display localized text
 		if ( Category.EqualTo(FText::FromString(Blueprint->GetName())) || Category.EqualTo(UEdGraphSchema_K2::VR_DefaultCategory) )
@@ -844,15 +869,16 @@ FText FRigVMVariableDetailCustomization::OnGetCategoryText() const
 
 void FRigVMVariableDetailCustomization::OnCategoryTextCommitted(const FText& NewText, ETextCommit::Type InTextCommit, FName VarName)
 {
-	TStrongObjectPtr<URigVMBlueprint> Blueprint = BlueprintPtr.Pin();
+	TStrongObjectPtr<UObject> Blueprint = BlueprintPtr.GetWeakObjectPtr().Pin();
 	if (!Blueprint)
 	{
 		return;
 	}
+	TRigVMAssetInterfacePtr Interface = Blueprint.Get();
 	
 	if (InTextCommit == ETextCommit::OnEnter || InTextCommit == ETextCommit::OnUserMovedFocus)
 	{
-		FBlueprintEditorUtils::SetBlueprintVariableCategory(Blueprint.Get(), VarName, nullptr, NewText);
+		Interface->SetVariableCategory(CachedVariableName, NewText.ToString());
 		PopulateCategories();
 	}
 }
@@ -867,17 +893,17 @@ TSharedRef<ITableRow> FRigVMVariableDetailCustomization::MakeCategoryViewWidget(
 
 void FRigVMVariableDetailCustomization::OnCategorySelectionChanged(TSharedPtr<FText> ProposedSelection, ESelectInfo::Type)
 {
-	TStrongObjectPtr<URigVMBlueprint> Blueprint = BlueprintPtr.Pin();
+	TStrongObjectPtr<UObject> Blueprint = BlueprintPtr.GetWeakObjectPtr().Pin();
 	if (!Blueprint)
 	{
 		return;
 	}
+	TRigVMAssetInterfacePtr Interface = Blueprint.Get();
 	
 	if (ProposedSelection.IsValid() && CachedVariableName != NAME_None)
 	{
 		FText NewCategory = *ProposedSelection.Get();
-
-		FBlueprintEditorUtils::SetBlueprintVariableCategory(Blueprint.Get(), CachedVariableName, nullptr, NewCategory );
+		Interface->SetVariableCategory(CachedVariableName, NewCategory.ToString());
 		CategoryListView.Pin()->ClearSelection();
 		CategoryComboButton.Pin()->SetIsOpen(false);
 	}
@@ -895,23 +921,17 @@ ECheckBoxState FRigVMVariableDetailCustomization::OnGetExposedToSpawnCheckboxSta
 
 void FRigVMVariableDetailCustomization::OnExposedToSpawnChanged(ECheckBoxState InNewState)
 {
-	TStrongObjectPtr<URigVMBlueprint> Blueprint = BlueprintPtr.Pin();
+	TStrongObjectPtr<UObject> Blueprint = BlueprintPtr.GetWeakObjectPtr().Pin();
 	if (!Blueprint)
 	{
 		return;
 	}
+	TRigVMAssetInterfacePtr Interface = Blueprint.Get();
 	
 	if (CachedVariableName != NAME_None)
 	{
 		const bool bExposeOnSpawn = (InNewState == ECheckBoxState::Checked);
-		if(bExposeOnSpawn)
-		{
-			FBlueprintEditorUtils::SetBlueprintVariableMetaData(Blueprint.Get(), CachedVariableName, NULL, FBlueprintMetadata::MD_ExposeOnSpawn, TEXT("true"));
-		}
-		else
-		{
-			FBlueprintEditorUtils::RemoveBlueprintVariableMetaData(Blueprint.Get(), CachedVariableName, NULL, FBlueprintMetadata::MD_ExposeOnSpawn);
-		} 
+		Interface->SetVariableExposeOnSpawn(CachedVariableName, bExposeOnSpawn);
 	}
 }
 
@@ -927,23 +947,17 @@ ECheckBoxState FRigVMVariableDetailCustomization::OnGetPrivateCheckboxState() co
 
 void FRigVMVariableDetailCustomization::OnPrivateChanged(ECheckBoxState InNewState)
 {
-	TStrongObjectPtr<URigVMBlueprint> Blueprint = BlueprintPtr.Pin();
+	TStrongObjectPtr<UObject> Blueprint = BlueprintPtr.GetWeakObjectPtr().Pin();
 	if (!Blueprint)
 	{
 		return;
 	}
+	TRigVMAssetInterfacePtr Interface = Blueprint.Get();
 	
 	if (CachedVariableName != NAME_None)
 	{
-		const bool bExposeOnSpawn = (InNewState == ECheckBoxState::Checked);
-		if(bExposeOnSpawn)
-		{
-			FBlueprintEditorUtils::SetBlueprintVariableMetaData(Blueprint.Get(), CachedVariableName, NULL, FBlueprintMetadata::MD_Private, TEXT("true"));
-		}
-		else
-		{
-			FBlueprintEditorUtils::RemoveBlueprintVariableMetaData(Blueprint.Get(), CachedVariableName, NULL, FBlueprintMetadata::MD_Private);
-		}
+		const bool bPrivate = (InNewState == ECheckBoxState::Checked);
+		Interface->SetVariablePrivate(CachedVariableName, bPrivate);
 	}
 }
 
@@ -959,11 +973,12 @@ ECheckBoxState FRigVMVariableDetailCustomization::OnGetExposedToCinematicsCheckb
 
 void FRigVMVariableDetailCustomization::OnExposedToCinematicsChanged(ECheckBoxState InNewState)
 {
-	TStrongObjectPtr<URigVMBlueprint> Blueprint = BlueprintPtr.Pin();
+	TStrongObjectPtr<UObject> Blueprint = BlueprintPtr.GetWeakObjectPtr().Pin();
 	if (!Blueprint)
 	{
 		return;
 	}
+	TRigVMAssetInterfacePtr Interface = Blueprint.Get();
 	
 	// Toggle the flag on the blueprint's version of the variable description, based on state
 	const bool bExposeToCinematics = (InNewState == ECheckBoxState::Checked);
@@ -971,23 +986,23 @@ void FRigVMVariableDetailCustomization::OnExposedToCinematicsChanged(ECheckBoxSt
 	const FName VarName = CachedVariableName;
 	if (VarName != NAME_None)
 	{
-		FBlueprintEditorUtils::SetInterpFlag(Blueprint.Get(), VarName, bExposeToCinematics);
+		Interface->SetVariableExposeToCinematics(VarName, bExposeToCinematics);
 	}
 }
 
 FText FRigVMVariableDetailCustomization::OnGetMetaKeyValue(FName Key) const
 {
-	TStrongObjectPtr<URigVMBlueprint> Blueprint = BlueprintPtr.Pin();
+	TStrongObjectPtr<UObject> Blueprint = BlueprintPtr.GetWeakObjectPtr().Pin();
 	if (!Blueprint)
 	{
 		return FText();
 	}
+	TRigVMAssetInterfacePtr Interface = Blueprint.Get();
 	
 	if (CachedVariableName != NAME_None)
 	{
-		FString Result;
-		FBlueprintEditorUtils::GetBlueprintVariableMetaData(Blueprint.Get(), CachedVariableName, nullptr, Key, /*out*/ Result);
-
+		FString Result = Interface->GetVariableMetadataValue(CachedVariableName, Key);
+		
 		return FText::FromString(Result);
 	}
 	return FText();
@@ -995,19 +1010,20 @@ FText FRigVMVariableDetailCustomization::OnGetMetaKeyValue(FName Key) const
 
 void FRigVMVariableDetailCustomization::OnMetaKeyValueChanged(const FText& NewMinValue, ETextCommit::Type CommitInfo, FName Key)
 {
-	TStrongObjectPtr<URigVMBlueprint> Blueprint = BlueprintPtr.Pin();
+	TStrongObjectPtr<UObject> Blueprint = BlueprintPtr.GetWeakObjectPtr().Pin();
 	if (!Blueprint)
 	{
 		return;
 	}
-	
-	if (CachedVariableName != NAME_None)
-	{
-		if ((CommitInfo == ETextCommit::OnEnter) || (CommitInfo == ETextCommit::OnUserMovedFocus))
-		{
-			FBlueprintEditorUtils::SetBlueprintVariableMetaData(Blueprint.Get(), CachedVariableName, nullptr, Key, NewMinValue.ToString());
-		}
-	}
+
+	checkf(false, TEXT("Not implemented yet"));
+	// if (CachedVariableName != NAME_None)
+	// {
+	// 	if ((CommitInfo == ETextCommit::OnEnter) || (CommitInfo == ETextCommit::OnUserMovedFocus))
+	// 	{
+	// 		FBlueprintEditorUtils::SetBlueprintVariableMetaData(Blueprint.Get(), CachedVariableName, nullptr, Key, NewMinValue.ToString());
+	// 	}
+	// }
 }
 
 EVisibility FRigVMVariableDetailCustomization::RangeVisibility() const
diff --git a/Engine/Plugins/Runtime/RigVM/Source/RigVMEditor/Private/Editor/RigVMVariantDetailCustomization.cpp b/Engine/Plugins/Runtime/RigVM/Source/RigVMEditor/Private/Editor/RigVMVariantDetailCustomization.cpp
--- a/Engine/Plugins/Runtime/RigVM/Source/RigVMEditor/Private/Editor/RigVMVariantDetailCustomization.cpp
+++ b/Engine/Plugins/Runtime/RigVM/Source/RigVMEditor/Private/Editor/RigVMVariantDetailCustomization.cpp
@@ -30,9 +30,9 @@ void FRigVMVariantDetailCustomization::CustomizeHeader(TSharedRef<IPropertyHandl
 
 	for (UObject* Object : Objects)
 	{
-		if (Object->IsA<URigVMBlueprint>())
+		if (Object->Implements<URigVMAssetInterface>())
 		{
-			BlueprintBeingCustomized = Cast<URigVMBlueprint>(Object);
+			BlueprintBeingCustomized = Object;
 			break;
 		}
 	}
@@ -40,7 +40,7 @@ void FRigVMVariantDetailCustomization::CustomizeHeader(TSharedRef<IPropertyHandl
 	FRigVMVariantWidgetContext VariantContext;
 	if(BlueprintBeingCustomized)
 	{
-		VariantContext.ParentPath = BlueprintBeingCustomized->GetPathName();
+		VariantContext.ParentPath = BlueprintBeingCustomized->GetObject()->GetPathName();
 	}
 
 	HeaderRow
@@ -75,7 +75,7 @@ FRigVMVariant FRigVMVariantDetailCustomization::GetVariant() const
 {
 	if (BlueprintBeingCustomized)
 	{
-		return BlueprintBeingCustomized->AssetVariant;
+		return BlueprintBeingCustomized->GetAssetVariant();
 	}
 	return FRigVMVariant();
 }
@@ -93,9 +93,9 @@ TArray<FRigVMVariantRef> FRigVMVariantDetailCustomization::GetVariantRefs() cons
 {
 	if (BlueprintBeingCustomized)
 	{
-		const FRigVMVariant& Variant = BlueprintBeingCustomized->AssetVariant;
+		const FRigVMVariant& Variant = BlueprintBeingCustomized->GetAssetVariant();
 		TArray<FRigVMVariantRef> Variants = URigVMBuildData::Get()->FindAssetVariantRefs(Variant.Guid);
-		const FRigVMVariantRef MyVariantRef = FRigVMVariantRef(BlueprintBeingCustomized->GetPathName(), BlueprintBeingCustomized->AssetVariant);
+		const FRigVMVariantRef MyVariantRef = FRigVMVariantRef(BlueprintBeingCustomized->GetObject()->GetPathName(), BlueprintBeingCustomized->GetAssetVariant());
 		Variants.RemoveAll([MyVariantRef](const FRigVMVariantRef& VariantRef) -> bool
 		{
 			return VariantRef == MyVariantRef;
@@ -110,8 +110,8 @@ void FRigVMVariantDetailCustomization::OnVariantChanged(const FRigVMVariant& InN
 	if(BlueprintBeingCustomized)
 	{
 		FScopedTransaction Transaction(LOCTEXT("ChangedVariantInfo", "Changed Blueprint Variant Information"));
-		BlueprintBeingCustomized->Modify();
-		BlueprintBeingCustomized->AssetVariant = InNewVariant;
+		BlueprintBeingCustomized->GetObject()->Modify();
+		BlueprintBeingCustomized->GetAssetVariant() = InNewVariant;
 	}
 }
 
@@ -129,7 +129,7 @@ TArray<FRigVMTag> FRigVMVariantDetailCustomization::OnGetTags() const
 {
 	if(BlueprintBeingCustomized)
 	{
-		return BlueprintBeingCustomized->AssetVariant.Tags;
+		return BlueprintBeingCustomized->GetAssetVariant().Tags;
 	}
 	return {};
 }
@@ -143,14 +143,14 @@ void FRigVMVariantDetailCustomization::OnAddTag(const FName& InTagName)
 		{
 			if (const FRigVMTag* Tag = Settings->FindTag(InTagName))
 			{
-				if(!BlueprintBeingCustomized->AssetVariant.Tags.ContainsByPredicate([InTagName](const FRigVMTag& Tag) -> bool
+				if(!BlueprintBeingCustomized->GetAssetVariant().Tags.ContainsByPredicate([InTagName](const FRigVMTag& Tag) -> bool
 				{
 					return Tag.Name == InTagName;
 				}))
 				{
 					FScopedTransaction Transaction(LOCTEXT("AddedBlueprintVariantTag", "Added Blueprint Variant Tag"));
-					BlueprintBeingCustomized->Modify();
-					BlueprintBeingCustomized->AssetVariant.Tags.Add(*Tag);
+					BlueprintBeingCustomized->GetObject()->Modify();
+					BlueprintBeingCustomized->GetAssetVariant().Tags.Add(*Tag);
 				}
 			}
 		}
@@ -161,14 +161,14 @@ void FRigVMVariantDetailCustomization::OnRemoveTag(const FName& InTagName)
 {
 	if(BlueprintBeingCustomized)
 	{
-		if(BlueprintBeingCustomized->AssetVariant.Tags.ContainsByPredicate([InTagName](const FRigVMTag& Tag) -> bool
+		if(BlueprintBeingCustomized->GetAssetVariant().Tags.ContainsByPredicate([InTagName](const FRigVMTag& Tag) -> bool
 			{
 				return Tag.Name == InTagName;
 			}))
 		{
 			FScopedTransaction Transaction(LOCTEXT("RemovedBlueprintVariantTag", "Removed Blueprint Variant Tag"));
-			BlueprintBeingCustomized->Modify();
-			BlueprintBeingCustomized->AssetVariant.Tags.RemoveAll([InTagName](const FRigVMTag& Tag) -> bool
+			BlueprintBeingCustomized->GetObject()->Modify();
+			BlueprintBeingCustomized->GetAssetVariant().Tags.RemoveAll([InTagName](const FRigVMTag& Tag) -> bool
 			{
 				return Tag.Name == InTagName;
 			});
diff --git a/Engine/Plugins/Runtime/RigVM/Source/RigVMEditor/Private/Editor/SRigVMDetailsInspector.cpp b/Engine/Plugins/Runtime/RigVM/Source/RigVMEditor/Private/Editor/SRigVMDetailsInspector.cpp
--- a/Engine/Plugins/Runtime/RigVM/Source/RigVMEditor/Private/Editor/SRigVMDetailsInspector.cpp
+++ b/Engine/Plugins/Runtime/RigVM/Source/RigVMEditor/Private/Editor/SRigVMDetailsInspector.cpp
@@ -133,7 +133,7 @@ const TArray< TWeakObjectPtr<UObject> >& SRigVMDetailsInspector::GetSelectedObje
 	return SelectedObjects;
 }
 
-void SRigVMDetailsInspector::OnEditorClose(const IRigVMEditor* RigVMEditorBase, URigVMBlueprint* RigVMBlueprint)
+void SRigVMDetailsInspector::OnEditorClose(const IRigVMEditor* RigVMEditorBase, TRigVMAssetInterfacePtr RigVMBlueprint)
 {
 	PropertyView->UnregisterInstancedCustomPropertyLayout(URigVMEdGraph::StaticClass());
 	PropertyView->UnregisterInstancedCustomPropertyLayout(UPropertyWrapper::StaticClass());
@@ -241,7 +241,7 @@ void SRigVMDetailsInspector::Construct(const FArguments& InArgs)
 	IsPropertyEditingEnabledDelegate = InArgs._IsPropertyEditingEnabledDelegate;
 	UserOnFinishedChangingProperties = InArgs._OnFinishedChangingProperties;
 
-	const UClass* BlueprintClass = Editor->GetRigVMBlueprint()->GetClass();
+	const UClass* BlueprintClass = Editor->GetRigVMAsset()->GetObject()->GetClass();
 	FOnGetDetailCustomizationInstance FunctionDetails = FOnGetDetailCustomizationInstance::CreateStatic(&FRigVMGraphDetailCustomization::MakeInstance, Editor, BlueprintClass);
 	PropertyView->RegisterInstancedCustomPropertyLayout(URigVMEdGraph::StaticClass(), FunctionDetails);
 
@@ -509,13 +509,13 @@ bool SRigVMDetailsInspector::IsPropertyVisible( const FPropertyAndParent& Proper
 		// this is necessary because of Instanced objects, which will have a different owning class yet are conceptually contained in this blueprint
 		bool bVariableAddedInCurrentBlueprint = false;
 		TSharedPtr<IRigVMEditor> Editor = WeakEditor.Pin();
-		const UBlueprint* Blueprint = Editor.IsValid() ? Editor->GetRigVMBlueprint() : nullptr;
+		TRigVMAssetInterfacePtr Blueprint = Editor.IsValid() ? Editor->GetRigVMAsset() : nullptr;
 
 		auto WasAddedInThisBlueprint = [Blueprint](const FProperty* Property)
 		{
 			if (const UClass* OwningClass = Property->GetOwnerClass())
 			{
-				return Blueprint && OwningClass->ClassGeneratedBy.Get() == Blueprint;
+				return Blueprint && OwningClass->ClassGeneratedBy.Get() == Blueprint->GetObject();
 			}
 			return false;
 		};
diff --git a/Engine/Plugins/Runtime/RigVM/Source/RigVMEditor/Private/Editor/SRigVMGraphTitleBar.cpp b/Engine/Plugins/Runtime/RigVM/Source/RigVMEditor/Private/Editor/SRigVMGraphTitleBar.cpp
--- a/Engine/Plugins/Runtime/RigVM/Source/RigVMEditor/Private/Editor/SRigVMGraphTitleBar.cpp
+++ b/Engine/Plugins/Runtime/RigVM/Source/RigVMEditor/Private/Editor/SRigVMGraphTitleBar.cpp
@@ -3,6 +3,7 @@
 
 #include "SRigVMGraphTitleBar.h"
 
+#include "RigVMBlueprintUtils.h"
 #include "Editor/RigVMNewEditor.h"
 #include "Containers/Array.h"
 #include "Delegates/Delegate.h"
@@ -216,10 +217,10 @@ void SRigVMGraphTitleBar::Construct( const FArguments& InArgs )
 	RebuildBreadcrumbTrail();
 	BreadcrumbTrailScrollBox->ScrollToEnd();
 
-	UBlueprint* BlueprintObj = FBlueprintEditorUtils::FindBlueprintForGraph(this->EdGraphObj);
+	TRigVMAssetInterfacePtr BlueprintObj = FRigVMBlueprintUtils::FindAssetForGraph(this->EdGraphObj);
 	if (BlueprintObj)
 	{
-		BlueprintTitle = FText::FromString(BlueprintObj->GetFriendlyName());
+		BlueprintTitle = FText::FromString(BlueprintObj->GetObject()->GetName());
 
 		// Register for notifications to refresh UI
 		if( TSharedPtr<FRigVMNewEditor> Editor = EditorWeak.Pin() )
@@ -286,9 +287,9 @@ void SRigVMGraphTitleBar::Refresh()
 	// Refresh UI on request
 	if (EdGraphObj)
 	{
-		if (UBlueprint* BlueprintObj = FBlueprintEditorUtils::FindBlueprintForGraph(this->EdGraphObj))
+		if (TRigVMAssetInterfacePtr BlueprintObj = FRigVMBlueprintUtils::FindAssetForGraph(this->EdGraphObj))
 		{
-			BlueprintTitle = FText::FromString(BlueprintObj->GetFriendlyName());
+			BlueprintTitle = FText::FromString(BlueprintObj->GetObject()->GetName());
 			RebuildBreadcrumbTrail();
 		}
 	}
diff --git a/Engine/Plugins/Runtime/RigVM/Source/RigVMEditor/Private/RigVMEditorBlueprintLibrary.cpp b/Engine/Plugins/Runtime/RigVM/Source/RigVMEditor/Private/RigVMEditorBlueprintLibrary.cpp
--- a/Engine/Plugins/Runtime/RigVM/Source/RigVMEditor/Private/RigVMEditorBlueprintLibrary.cpp
+++ b/Engine/Plugins/Runtime/RigVM/Source/RigVMEditor/Private/RigVMEditorBlueprintLibrary.cpp
@@ -3,6 +3,8 @@
 #include "RigVMEditorBlueprintLibrary.h"
 #include "AssetRegistry/AssetRegistryModule.h"
 #include "Misc/ScopedSlowTask.h"
+#include "UObject/UObjectIterator.h"
+
 
 #include UE_INLINE_GENERATED_CPP_BY_NAME(RigVMEditorBlueprintLibrary)
 
@@ -297,16 +299,35 @@ TArray<FAssetData> URigVMEditorBlueprintLibrary::GetAssetsWithFilter(
 		TSubclassOf<URigVMBlueprint> InClass,
 		FRigVMAssetDataFilter InAssetDataFilter)
 {
+	FARFilter Filter;
 	if(InClass == nullptr)
 	{
-		InClass = URigVMBlueprint::StaticClass();
+		for (TObjectIterator<UClass> It; It; ++It)
+		{
+			UClass* Class = *It;
+
+			// Skip abstract and non-blueprintable classes if needed
+			if (Class->HasAnyClassFlags(CLASS_Abstract))
+			{
+				continue;
+			}
+
+			if (Class->ImplementsInterface(URigVMAssetInterface::StaticClass()))
+			{
+				Filter.ClassPaths.Add(Class->GetClassPathName());
+			}
+		}
+	}
+	else
+	{
+		Filter.ClassPaths.Add(InClass->GetClassPathName());
 	}
 
 	const FAssetRegistryModule& AssetRegistryModule = FModuleManager::LoadModuleChecked<FAssetRegistryModule>(TEXT("AssetRegistry"));
 
 	// Collect a full list of assets with the specified class
 	TArray<FAssetData> AssetDataList;
-	AssetRegistryModule.Get().GetAssetsByClass(InClass->GetClassPathName(), AssetDataList, true);
+	AssetRegistryModule.Get().GetAssets(Filter, AssetDataList);
 
 	if(InAssetDataFilter.IsBound())
 	{
