RigVM: Adding IRigVMAssetInterface as parent of RigVMBlueprint. First step to separation of RigVM and Blueprint.

diff --git a/Engine/Plugins/Runtime/RigVM/Source/RigVMEditor/Private/RigVMEditorModule.cpp b/Engine/Plugins/Runtime/RigVM/Source/RigVMEditor/Private/RigVMEditorModule.cpp
--- a/Engine/Plugins/Runtime/RigVM/Source/RigVMEditor/Private/RigVMEditorModule.cpp
+++ b/Engine/Plugins/Runtime/RigVM/Source/RigVMEditor/Private/RigVMEditorModule.cpp
@@ -75,8 +75,8 @@ void FRigVMEditorModule::StartupModule()
 		FEdGraphUtilities::RegisterVisualNodeFactory(EdGraphPanelNodeFactory);
 		FEdGraphUtilities::RegisterVisualPinFactory(EdGraphPanelPinFactory);
 
-		ReconstructAllNodesDelegateHandle = FBlueprintEditorUtils::OnReconstructAllNodesEvent.AddStatic(&FRigVMBlueprintUtils::HandleReconstructAllNodes);
-		RefreshAllNodesDelegateHandle = FBlueprintEditorUtils::OnRefreshAllNodesEvent.AddStatic(&FRigVMBlueprintUtils::HandleRefreshAllNodes);
+		ReconstructAllNodesDelegateHandle = FBlueprintEditorUtils::OnReconstructAllNodesEvent.AddStatic(&FRigVMBlueprintUtils::HandleReconstructAllBlueprintNodes);
+		RefreshAllNodesDelegateHandle = FBlueprintEditorUtils::OnRefreshAllNodesEvent.AddStatic(&FRigVMBlueprintUtils::HandleRefreshAllBlueprintNodes);
 
 #if WITH_RIGVMLEGACYEDITOR
 		// Register Blueprint editor variable customization
@@ -107,7 +107,7 @@ void FRigVMEditorModule::StartupModule()
 						}
 
 						const FAssetData& SelectedAssetData = SelectedAssets[0];
-						if (!SelectedAssetData.IsInstanceOf<URigVMBlueprint>())
+						if (!SelectedAssetData.GetClass()->ImplementsInterface(URigVMAssetInterface::StaticClass()))
 						{
 							// We aren't dealing with a RigVMBlueprint derived type
 							return;
@@ -135,9 +135,9 @@ void FRigVMEditorModule::StartupModule()
 									FPaths::Split(PackageName, PathName, ObjectName, Extension);
 
 									UObject* DuplicateAsset = AssetToolsModule.Get().DuplicateAsset(ObjectName, PathName, SelectedObject);
-									if (URigVMBlueprint* DuplicateBlueprint = Cast<URigVMBlueprint>(DuplicateAsset))
+									if (IRigVMAssetInterface* DuplicateBlueprint = Cast<IRigVMAssetInterface>(DuplicateAsset))
 									{
-										DuplicateBlueprint->AssetVariant = Cast<URigVMBlueprint>(SelectedObject)->AssetVariant;
+										DuplicateBlueprint->GetAssetVariant() = Cast<IRigVMAssetInterface>(SelectedObject)->GetAssetVariant();
 									}
 								}));
 						}
@@ -155,9 +155,8 @@ void FRigVMEditorModule::StartupModuleCommon()
 	MenuExtensibilityManager = MakeShareable(new FExtensibilityManager);
     ToolBarExtensibilityManager = MakeShareable(new FExtensibilityManager);
 
-	const UClass* BlueprintClass = GetRigVMBlueprintClass();
-	const URigVMBlueprint* BlueprintCDO = GetRigVMBlueprintCDO();
-	const UClass* EdGraphSchemaClass = BlueprintCDO->GetRigVMEdGraphSchemaClass();
+	const UClass* BlueprintClass = GetRigVMAssetClass();
+	const UClass* EdGraphSchemaClass = GetRigVMEdGraphSchemaClass();
 	const URigVMEdGraphSchema* SchemaCDO = CastChecked<URigVMEdGraphSchema>(EdGraphSchemaClass->GetDefaultObject());
 #if WITH_RIGVMLEGACYEDITOR
 	FBlueprintEditorModule& BlueprintEditorModule = FModuleManager::LoadModuleChecked<FBlueprintEditorModule>("Kismet");
@@ -203,8 +202,7 @@ void FRigVMEditorModule::ShutdownModuleCommon()
 {
 	if (!IsEngineExitRequested())
 	{
-		const URigVMBlueprint* BlueprintCDO = GetRigVMBlueprintCDO();
-		const UClass* EdGraphSchemaClass = BlueprintCDO->GetRigVMEdGraphSchemaClass();
+		const UClass* EdGraphSchemaClass = GetRigVMEdGraphSchemaClass();
 		const URigVMEdGraphSchema* SchemaCDO = CastChecked<URigVMEdGraphSchema>(EdGraphSchemaClass->GetDefaultObject());
 #if WITH_RIGVMLEGACYEDITOR
 		FBlueprintEditorModule& BlueprintEditorModule = FModuleManager::LoadModuleChecked<FBlueprintEditorModule>("Kismet");
@@ -216,22 +214,20 @@ void FRigVMEditorModule::ShutdownModuleCommon()
 	FKismetEditorUtilities::UnregisterAutoBlueprintNodeCreation(this);
 }
 
-UClass* FRigVMEditorModule::GetRigVMBlueprintClass() const
-{
-	return URigVMBlueprint::StaticClass();
-}
-
-const URigVMBlueprint* FRigVMEditorModule::GetRigVMBlueprintCDO() const
-{
-	const UClass* RigVMBlueprintClass = GetRigVMBlueprintClass();
-	return CastChecked<URigVMBlueprint>(RigVMBlueprintClass->GetDefaultObject());
-}
-
-void FRigVMEditorModule::GetTypeActions(URigVMBlueprint* RigVMBlueprint, FBlueprintActionDatabaseRegistrar& ActionRegistrar)
+void FRigVMEditorModule::GetTypeActions(TRigVMAssetInterfacePtr RigVMBlueprint, FBlueprintActionDatabaseRegistrar& ActionRegistrar)
 {
 	// Only register actions for ourselves
-	UClass* BlueprintClass = GetRigVMBlueprintClass();
-	if(RigVMBlueprint->GetClass() != BlueprintClass)
+	UClass* InClass = RigVMBlueprint.GetObject()->GetClass();
+	UClass* BlueprintClass = GetRigVMAssetClass();
+
+	if (BlueprintClass->GetDefaultObject()->IsA<UInterface>())
+	{
+		if (!InClass->ImplementsInterface(BlueprintClass))
+		{
+			return;
+		}
+	}
+	else if(InClass != BlueprintClass)
 	{
 		return;
 	}
@@ -242,7 +238,7 @@ void FRigVMEditorModule::GetTypeActions(URigVMBlueprint* RigVMBlueprint, FBluepr
 	// actions might have to be updated (or deleted) if their object-key is  
 	// mutated (or removed)... here we use the class (so if the class 
 	// type disappears, then the action should go with it)
-	UClass* ActionKey = RigVMBlueprint->GetClass();
+	UClass* ActionKey = RigVMBlueprint->GetObject()->GetClass();
 	// to keep from needlessly instantiating a UBlueprintNodeSpawner, first   
 	// check to make sure that the registrar is looking for actions of this type
 	// (could be regenerating actions for a specific asset, and therefore the 
@@ -255,8 +251,6 @@ void FRigVMEditorModule::GetTypeActions(URigVMBlueprint* RigVMBlueprint, FBluepr
 	FRigVMRegistry& Registry = FRigVMRegistry::Get();
 	Registry.RefreshEngineTypes();
 
-	const URigVMBlueprint* BlueprintCDO = GetRigVMBlueprintCDO();
-
 	for (const FRigVMTemplate& Template : Registry.GetTemplates())
 	{
 		// factories are registered below
@@ -277,7 +271,7 @@ void FRigVMEditorModule::GetTypeActions(URigVMBlueprint* RigVMBlueprint, FBluepr
 			continue;
 		}
 
-		if(!Template.SupportsExecuteContextStruct(BlueprintCDO->GetRigVMExecuteContextStruct()))
+		if(!Template.SupportsExecuteContextStruct(GetRigVMExecuteContextStruct()))
 		{
 			continue;
 		}
@@ -286,15 +280,15 @@ void FRigVMEditorModule::GetTypeActions(URigVMBlueprint* RigVMBlueprint, FBluepr
 		FText MenuDesc = FText::FromName(Template.GetName());
 		FText ToolTip = Template.GetTooltipText();
 
-		URigVMEdGraphNodeSpawner* NodeSpawner = URigVMEdGraphTemplateNodeSpawner::CreateFromNotation(Template.GetNotation(), MenuDesc, NodeCategory, ToolTip);
-		check(NodeSpawner != nullptr);
+		TSharedRef<URigVMEdGraphTemplateNodeSpawner> NodeSpawner = MakeShared<URigVMEdGraphTemplateNodeSpawner>(Template.GetNotation(), MenuDesc, NodeCategory, ToolTip);
 		NodeSpawner->SetRelatedBlueprintClass(BlueprintClass);
-		ActionRegistrar.AddBlueprintAction(ActionKey, NodeSpawner);
+		URigVMEdGraphNodeBlueprintSpawner* BlueprintSpawner = URigVMEdGraphNodeBlueprintSpawner::CreateFromRigVMSpawner(NodeSpawner);
+		ActionRegistrar.AddBlueprintAction(ActionKey, BlueprintSpawner);
 	};
 
 	for (const FRigVMDispatchFactory* Factory : Registry.GetFactories())
 	{
-		if(!Factory->SupportsExecuteContextStruct(BlueprintCDO->GetRigVMExecuteContextStruct()))
+		if(!Factory->SupportsExecuteContextStruct(GetRigVMExecuteContextStruct()))
 		{
 			continue;
 		}
@@ -309,10 +303,10 @@ void FRigVMEditorModule::GetTypeActions(URigVMBlueprint* RigVMBlueprint, FBluepr
 		FText MenuDesc = FText::FromString(Factory->GetNodeTitle(FRigVMTemplateTypeMap()));
 		FText ToolTip = Factory->GetNodeTooltip(FRigVMTemplateTypeMap());
 
-		URigVMEdGraphNodeSpawner* NodeSpawner = URigVMEdGraphTemplateNodeSpawner::CreateFromNotation(Template->GetNotation(), MenuDesc, NodeCategory, ToolTip);
-		check(NodeSpawner != nullptr);
+		TSharedRef<URigVMEdGraphTemplateNodeSpawner> NodeSpawner = MakeShared<URigVMEdGraphTemplateNodeSpawner>(Template->GetNotation(), MenuDesc, NodeCategory, ToolTip);
 		NodeSpawner->SetRelatedBlueprintClass(BlueprintClass);
-		ActionRegistrar.AddBlueprintAction(ActionKey, NodeSpawner);
+		URigVMEdGraphNodeBlueprintSpawner* BlueprintSpawner = URigVMEdGraphNodeBlueprintSpawner::CreateFromRigVMSpawner(NodeSpawner);
+		ActionRegistrar.AddBlueprintAction(ActionKey, BlueprintSpawner);
 	};
 
 	// Add all rig units
@@ -324,7 +318,7 @@ void FRigVMEditorModule::GetTypeActions(URigVMBlueprint* RigVMBlueprint, FBluepr
 			continue;
 		}
 
-		if(!Function.SupportsExecuteContextStruct(BlueprintCDO->GetRigVMExecuteContextStruct()))
+		if(!Function.SupportsExecuteContextStruct(GetRigVMExecuteContextStruct()))
 		{
 			continue;
 		}
@@ -358,10 +352,10 @@ void FRigVMEditorModule::GetTypeActions(URigVMBlueprint* RigVMBlueprint, FBluepr
 		FText MenuDesc = FText::FromString(DisplayNameMetadata + MenuDescSuffixMetadata);
 		FText ToolTip = Struct->GetToolTipText();
 
-		URigVMEdGraphNodeSpawner* NodeSpawner = URigVMEdGraphUnitNodeSpawner::CreateFromStruct(Struct, Function.GetMethodName(), MenuDesc, NodeCategory, ToolTip);
-		check(NodeSpawner != nullptr);
+		TSharedRef<URigVMEdGraphUnitNodeSpawner> NodeSpawner = MakeShared<URigVMEdGraphUnitNodeSpawner>(Struct, Function.GetMethodName(), MenuDesc, NodeCategory, ToolTip);
 		NodeSpawner->SetRelatedBlueprintClass(BlueprintClass);
-		ActionRegistrar.AddBlueprintAction(ActionKey, NodeSpawner);
+		URigVMEdGraphNodeBlueprintSpawner* BlueprintSpawner = URigVMEdGraphNodeBlueprintSpawner::CreateFromRigVMSpawner(NodeSpawner);
+		ActionRegistrar.AddBlueprintAction(ActionKey, BlueprintSpawner);
 	};
 
 	for (TObjectIterator<UEnum> EnumIt; EnumIt; ++EnumIt)
@@ -387,10 +381,10 @@ void FRigVMEditorModule::GetTypeActions(URigVMBlueprint* RigVMBlueprint, FBluepr
 		FText MenuDesc = FText::FromString(FString::Printf(TEXT("Enum %s"), *EnumToConsider->GetName()));
 		FText ToolTip = MenuDesc;
 
-		URigVMEdGraphNodeSpawner* NodeSpawner = URigVMEdGraphEnumNodeSpawner::CreateForEnum(EnumToConsider, MenuDesc, NodeCategory, ToolTip);
-		check(NodeSpawner != nullptr);
+		TSharedRef<URigVMEdGraphEnumNodeSpawner> NodeSpawner = MakeShared<URigVMEdGraphEnumNodeSpawner>(EnumToConsider, MenuDesc, NodeCategory, ToolTip);
 		NodeSpawner->SetRelatedBlueprintClass(BlueprintClass);
-		ActionRegistrar.AddBlueprintAction(ActionKey, NodeSpawner);
+		URigVMEdGraphNodeBlueprintSpawner* BlueprintSpawner = URigVMEdGraphNodeBlueprintSpawner::CreateFromRigVMSpawner(NodeSpawner);
+		ActionRegistrar.AddBlueprintAction(ActionKey, BlueprintSpawner);
 	}
 
 	FArrayProperty* PublicGraphFunctionsProperty = CastField<FArrayProperty>(URigVMBlueprint::StaticClass()->FindPropertyByName(TEXT("PublicGraphFunctions")));
@@ -421,8 +415,8 @@ void FRigVMEditorModule::GetTypeActions(URigVMBlueprint* RigVMBlueprint, FBluepr
 				}
 			}
 		};
-		AddClassToFilter(BlueprintCDO->GetClass());
-		AddClassToFilter(BlueprintCDO->GetRigVMBlueprintGeneratedClassPrototype());
+		AddClassToFilter(RigVMBlueprint->GetObject()->GetClass());
+		AddClassToFilter(RigVMBlueprint->GetRigVMGeneratedClassPrototype());
 
 		AssetRegistryModule.Get().GetAssets(ControlRigAssetFilter, ControlRigAssetDatas);
 
@@ -446,9 +440,9 @@ void FRigVMEditorModule::GetTypeActions(URigVMBlueprint* RigVMBlueprint, FBluepr
 			if (ControlRigAssetData.IsAssetLoaded())
 			{
 				UObject* AssetObject = ControlRigAssetData.GetAsset();
-				if (URigVMBlueprint* Blueprint = Cast<URigVMBlueprint>(AssetObject))
+				if (IRigVMAssetInterface* Asset = Cast<IRigVMAssetInterface>(AssetObject))
 				{
-					PublicFunctions = Blueprint->PublicGraphFunctions;
+					PublicFunctions = Asset->GetPublicGraphFunctions();
 				}
 				else if(URigVMBlueprintGeneratedClass* GeneratedClass = Cast<URigVMBlueprintGeneratedClass>(AssetObject))
 				{
@@ -490,10 +484,10 @@ void FRigVMEditorModule::GetTypeActions(URigVMBlueprint* RigVMBlueprint, FBluepr
 					PublicFunctionsProperty->ImportText_Direct(*PublicFunctionsString, &OldPublicFunctions, nullptr, EPropertyPortFlags::PPF_None);
 					for(const FRigVMOldPublicFunctionData& PublicFunction : OldPublicFunctions)
 					{
-						URigVMEdGraphNodeSpawner* NodeSpawner = URigVMEdGraphFunctionRefNodeSpawner::CreateFromAssetData(ControlRigAssetData, PublicFunction);
-						check(NodeSpawner != nullptr);
+						TSharedRef<URigVMEdGraphFunctionRefNodeSpawner> NodeSpawner = MakeShared<URigVMEdGraphFunctionRefNodeSpawner>(ControlRigAssetData, PublicFunction);
 						NodeSpawner->SetRelatedBlueprintClass(BlueprintClass);
-						ActionRegistrar.AddBlueprintAction(ActionKey, NodeSpawner);
+						URigVMEdGraphNodeBlueprintSpawner* BlueprintSpawner = URigVMEdGraphNodeBlueprintSpawner::CreateFromRigVMSpawner(NodeSpawner);
+						ActionRegistrar.AddBlueprintAction(ActionKey, BlueprintSpawner);
 					}
 				}
 
@@ -518,49 +512,49 @@ void FRigVMEditorModule::GetTypeActions(URigVMBlueprint* RigVMBlueprint, FBluepr
 			{
 				if (PublicFunction.LibraryPointer.IsValid())
 				{
-					URigVMEdGraphNodeSpawner* NodeSpawner = URigVMEdGraphFunctionRefNodeSpawner::CreateFromAssetData(ControlRigAssetData, PublicFunction);
-					check(NodeSpawner != nullptr);
+					TSharedPtr<URigVMEdGraphFunctionRefNodeSpawner> NodeSpawner = MakeShared<URigVMEdGraphFunctionRefNodeSpawner>(ControlRigAssetData, PublicFunction);
 					NodeSpawner->SetRelatedBlueprintClass(BlueprintClass);
-					ActionRegistrar.AddBlueprintAction(ActionKey, NodeSpawner);
+					URigVMEdGraphNodeBlueprintSpawner* BlueprintSpawner = URigVMEdGraphNodeBlueprintSpawner::CreateFromRigVMSpawner(NodeSpawner);
+					ActionRegistrar.AddBlueprintAction(ActionKey, BlueprintSpawner);
 				}
 			}
 		}
 	}
 }
 
-void FRigVMEditorModule::GetInstanceActions(URigVMBlueprint* RigVMBlueprint, FBlueprintActionDatabaseRegistrar& ActionRegistrar)
+void FRigVMEditorModule::GetInstanceActions(TRigVMAssetInterfacePtr RigVMBlueprint, FBlueprintActionDatabaseRegistrar& ActionRegistrar)
 {
 	// Only register actions for ourselves
-	UClass* BlueprintClass = GetRigVMBlueprintClass();
-	if(RigVMBlueprint->GetClass() != BlueprintClass)
+	UClass* BlueprintClass = GetRigVMAssetClass();
+	if(RigVMBlueprint->GetObject()->GetClass() != BlueprintClass)
 	{
 		return;
 	}
 
-	if (URigVMBlueprintGeneratedClass* GeneratedClass = RigVMBlueprint->GetRigVMBlueprintGeneratedClass())
+	if (UClass* GeneratedClass = RigVMBlueprint->GetRigVMGeneratedClass())
 	{
 		if (!ActionRegistrar.IsOpenForRegistration(GeneratedClass))
 		{
 			return;
 		}
-		if (URigVMHost* CDO = Cast<URigVMHost>(GeneratedClass->GetDefaultObject()))
-		{
-			static const FText NodeCategory = LOCTEXT("Variables", "Variables");
-
-			TArray<FRigVMExternalVariable> ExternalVariables = CDO->GetExternalVariables();
-			for (const FRigVMExternalVariable& ExternalVariable : ExternalVariables)
-			{
-				FText MenuDesc = FText::FromName(ExternalVariable.Name);
-				FText ToolTip = FText::FromString(FString::Printf(TEXT("Get the value of variable %s"), *ExternalVariable.Name.ToString()));
-				URigVMEdGraphNodeSpawner* GetNodeSpawner = URigVMEdGraphVariableNodeSpawner::CreateFromExternalVariable(RigVMBlueprint, ExternalVariable, true, MenuDesc, NodeCategory, ToolTip);
-				GetNodeSpawner->SetRelatedBlueprintClass(BlueprintClass);
-				ActionRegistrar.AddBlueprintAction(GeneratedClass, GetNodeSpawner);
-
-				ToolTip = FText::FromString(FString::Printf(TEXT("Set the value of variable %s"), *ExternalVariable.Name.ToString()));
-				URigVMEdGraphNodeSpawner* SetNodeSpawner = URigVMEdGraphVariableNodeSpawner::CreateFromExternalVariable(RigVMBlueprint, ExternalVariable, false, MenuDesc, NodeCategory, ToolTip);
-				SetNodeSpawner->SetRelatedBlueprintClass(BlueprintClass);
-				ActionRegistrar.AddBlueprintAction(GeneratedClass, SetNodeSpawner);
-			}
+		static const FText NodeCategory = LOCTEXT("Variables", "Variables");
+
+		TArray<FRigVMGraphVariableDescription> Variables = RigVMBlueprint->GetAssetVariables();
+		for (const FRigVMGraphVariableDescription& Variable : Variables)
+		{
+			const FRigVMExternalVariable ExternalVariable = Variable.ToExternalVariable();
+			FText MenuDesc = FText::FromName(ExternalVariable.Name);
+			FText ToolTip = FText::FromString(FString::Printf(TEXT("Get the value of variable %s"), *ExternalVariable.Name.ToString()));
+			TSharedPtr<URigVMEdGraphVariableNodeSpawner> GetNodeSpawner = MakeShared<URigVMEdGraphVariableNodeSpawner>(RigVMBlueprint, ExternalVariable, true, MenuDesc, NodeCategory, ToolTip);
+			GetNodeSpawner->SetRelatedBlueprintClass(BlueprintClass);
+			URigVMEdGraphNodeBlueprintSpawner* BlueprintGetSpawner = URigVMEdGraphNodeBlueprintSpawner::CreateFromRigVMSpawner(GetNodeSpawner);
+			ActionRegistrar.AddBlueprintAction(GeneratedClass, BlueprintGetSpawner);
+
+			ToolTip = FText::FromString(FString::Printf(TEXT("Set the value of variable %s"), *ExternalVariable.Name.ToString()));
+			TSharedPtr<URigVMEdGraphVariableNodeSpawner> SetNodeSpawner = MakeShared<URigVMEdGraphVariableNodeSpawner>(RigVMBlueprint, ExternalVariable, false, MenuDesc, NodeCategory, ToolTip);
+			SetNodeSpawner->SetRelatedBlueprintClass(BlueprintClass);
+			URigVMEdGraphNodeBlueprintSpawner* BlueprintSetSpawner = URigVMEdGraphNodeBlueprintSpawner::CreateFromRigVMSpawner(SetNodeSpawner);
+			ActionRegistrar.AddBlueprintAction(GeneratedClass, BlueprintSetSpawner);
 		}
 
 		if (URigVMFunctionLibrary* LocalFunctionLibrary = RigVMBlueprint->GetLocalFunctionLibrary())
@@ -570,7 +564,7 @@ void FRigVMEditorModule::GetInstanceActions(URigVMBlueprint* RigVMBlueprint, FBl
 			for (URigVMLibraryNode* Function : Functions)
 			{
 				// Avoid adding functions that are already added by the GetTypeActions functions (public functions that are already saved into the blueprint tag)
-				if (RigVMBlueprint->PublicGraphFunctions.ContainsByPredicate([LocalLibrarySoftPath, Function](const FRigVMGraphFunctionHeader& Header) -> bool
+				if (RigVMBlueprint->GetPublicGraphFunctions().ContainsByPredicate([LocalLibrarySoftPath, Function](const FRigVMGraphFunctionHeader& Header) -> bool
 				{
 					return FRigVMGraphFunctionIdentifier(LocalLibrarySoftPath, Function->GetPathName()) == Header.LibraryPointer;
 				}))
@@ -578,27 +572,29 @@ void FRigVMEditorModule::GetInstanceActions(URigVMBlueprint* RigVMBlueprint, FBl
 					continue;
 				}
 				
-				URigVMEdGraphNodeSpawner* NodeSpawner = URigVMEdGraphFunctionRefNodeSpawner::CreateFromFunction(Function);
-				check(NodeSpawner != nullptr);
+				TSharedPtr<URigVMEdGraphFunctionRefNodeSpawner> NodeSpawner = MakeShared<URigVMEdGraphFunctionRefNodeSpawner>(Function);
 				NodeSpawner->SetRelatedBlueprintClass(BlueprintClass);
-				ActionRegistrar.AddBlueprintAction(GeneratedClass, NodeSpawner);
+				URigVMEdGraphNodeBlueprintSpawner* BlueprintSpawner = URigVMEdGraphNodeBlueprintSpawner::CreateFromRigVMSpawner(NodeSpawner);
+				ActionRegistrar.AddBlueprintAction(GeneratedClass, BlueprintSpawner);
 			}
 
-			static const FText NodeCategory = LOCTEXT("LocalVariables", "Local Variables");
+			static const FText LocalNodeCategory = LOCTEXT("LocalVariables", "Local Variables");
 			for (URigVMLibraryNode* Function : Functions)
 			{
 				for (const FRigVMGraphVariableDescription& LocalVariable : Function->GetContainedGraph()->GetLocalVariables())
 				{
 					FText MenuDesc = FText::FromName(LocalVariable.Name);
 					FText ToolTip = FText::FromString(FString::Printf(TEXT("Get the value of variable %s"), *LocalVariable.Name.ToString()));
-					URigVMEdGraphNodeSpawner* GetNodeSpawner = URigVMEdGraphVariableNodeSpawner::CreateFromLocalVariable(RigVMBlueprint, Function->GetContainedGraph(), LocalVariable, true, MenuDesc, NodeCategory, ToolTip);
+					TSharedPtr<URigVMEdGraphVariableNodeSpawner> GetNodeSpawner = MakeShared<URigVMEdGraphVariableNodeSpawner>(RigVMBlueprint, Function->GetContainedGraph(), LocalVariable, true, MenuDesc, LocalNodeCategory, ToolTip);
 					GetNodeSpawner->SetRelatedBlueprintClass(BlueprintClass);
-					ActionRegistrar.AddBlueprintAction(GeneratedClass, GetNodeSpawner);
+					URigVMEdGraphNodeBlueprintSpawner* BlueprintGetSpawner = URigVMEdGraphNodeBlueprintSpawner::CreateFromRigVMSpawner(GetNodeSpawner);
+					ActionRegistrar.AddBlueprintAction(GeneratedClass, BlueprintGetSpawner);
 
 					ToolTip = FText::FromString(FString::Printf(TEXT("Set the value of variable %s"), *LocalVariable.Name.ToString()));
-					URigVMEdGraphNodeSpawner* SetNodeSpawner = URigVMEdGraphVariableNodeSpawner::CreateFromLocalVariable(RigVMBlueprint, Function->GetContainedGraph(), LocalVariable, false, MenuDesc, NodeCategory, ToolTip);
+					TSharedPtr<URigVMEdGraphVariableNodeSpawner> SetNodeSpawner = MakeShared<URigVMEdGraphVariableNodeSpawner>(RigVMBlueprint, Function->GetContainedGraph(), LocalVariable, false, MenuDesc, LocalNodeCategory, ToolTip);
 					SetNodeSpawner->SetRelatedBlueprintClass(BlueprintClass);
-					ActionRegistrar.AddBlueprintAction(GeneratedClass, SetNodeSpawner);
+					URigVMEdGraphNodeBlueprintSpawner* BlueprintSetSpawner = URigVMEdGraphNodeBlueprintSpawner::CreateFromRigVMSpawner(SetNodeSpawner);
+					ActionRegistrar.AddBlueprintAction(GeneratedClass, BlueprintSetSpawner);
 				}
 			}
 		}
@@ -607,22 +603,23 @@ void FRigVMEditorModule::GetInstanceActions(URigVMBlueprint* RigVMBlueprint, FBl
 		{
 			if (Graph->GetEntryNode())
 			{
-				static const FText NodeCategory = LOCTEXT("InputArguments", "Input Arguments");
+				static const FText InputNodeCategory = LOCTEXT("InputArguments", "Input Arguments");
 				for (const FRigVMGraphVariableDescription& InputArgument : Graph->GetInputArguments())
 				{
 					FText MenuDesc = FText::FromName(InputArgument.Name);
 					FText ToolTip = FText::FromString(FString::Printf(TEXT("Get the value of input %s"), *InputArgument.Name.ToString()));
-					URigVMEdGraphNodeSpawner* NodeSpawner = URigVMEdGraphVariableNodeSpawner::CreateFromLocalVariable(RigVMBlueprint, Graph, InputArgument, true, MenuDesc, NodeCategory, ToolTip);
+					TSharedPtr<URigVMEdGraphVariableNodeSpawner> NodeSpawner = MakeShared<URigVMEdGraphVariableNodeSpawner>(RigVMBlueprint, Graph, InputArgument, true, MenuDesc, InputNodeCategory, ToolTip);
 					NodeSpawner->SetRelatedBlueprintClass(BlueprintClass);
-					ActionRegistrar.AddBlueprintAction(GeneratedClass, NodeSpawner);
+					URigVMEdGraphNodeBlueprintSpawner* BlueprintSpawner = URigVMEdGraphNodeBlueprintSpawner::CreateFromRigVMSpawner(NodeSpawner);
+					ActionRegistrar.AddBlueprintAction(GeneratedClass, BlueprintSpawner);
 				}			
 			}
 		}
 
-		const TArray<FName> EntryNames = RigVMBlueprint->GetRigVMClient()->GetEntryNames();
+		const TArray<FName> EntryNames = RigVMBlueprint->GetRigVMClientHost()->GetRigVMClient()->GetEntryNames();
 		if(!EntryNames.IsEmpty())
 		{
-			static const FText NodeCategory = LOCTEXT("Events", "Events");
+			static const FText EventNodeCategory = LOCTEXT("Events", "Events");
 			for (const FName& EntryName : EntryNames)
 			{
 				static constexpr TCHAR EventStr[] = TEXT("Event");
@@ -630,22 +627,23 @@ void FRigVMEditorModule::GetInstanceActions(URigVMBlueprint* RigVMBlueprint, FBl
 				FString Suffix = EntryName.ToString().EndsWith(EventStr) ? FString() : EventSuffix;
 				FText MenuDesc = FText::FromString(FString::Printf(TEXT("Run %s%s"), *EntryName.ToString(), *Suffix));
 				FText ToolTip = FText::FromString(FString::Printf(TEXT("Runs the %s%s"), *EntryName.ToString(), *Suffix));
-				URigVMEdGraphNodeSpawner* NodeSpawner = URigVMEdGraphInvokeEntryNodeSpawner::CreateForEntry(RigVMBlueprint, EntryName, MenuDesc, NodeCategory, ToolTip);
+				TSharedPtr<URigVMEdGraphInvokeEntryNodeSpawner> NodeSpawner = MakeShared<URigVMEdGraphInvokeEntryNodeSpawner>(RigVMBlueprint, EntryName, MenuDesc, EventNodeCategory, ToolTip);
 				NodeSpawner->SetRelatedBlueprintClass(BlueprintClass);
-				ActionRegistrar.AddBlueprintAction(GeneratedClass, NodeSpawner);
+				URigVMEdGraphNodeBlueprintSpawner* BlueprintSpawner = URigVMEdGraphNodeBlueprintSpawner::CreateFromRigVMSpawner(NodeSpawner);
+				ActionRegistrar.AddBlueprintAction(GeneratedClass, BlueprintSpawner);
 			}
 		}
 	}
 }
 
-void FRigVMEditorModule::GetNodeContextMenuActions(IRigVMClientHost* RigVMClientHost, const URigVMEdGraphNode* EdGraphNode, URigVMNode* ModelNode, UToolMenu* Menu) const
+void FRigVMEditorModule::GetNodeContextMenuActions(TScriptInterface<IRigVMClientHost> RigVMClientHost, const URigVMEdGraphNode* EdGraphNode, URigVMNode* ModelNode, UToolMenu* Menu) const
 {
 	// Only register menu actions for ourselves if we are a blueprint
-	if(URigVMBlueprint* RigVMBlueprint = Cast<URigVMBlueprint>(RigVMClientHost))
+	if(TRigVMAssetInterfacePtr RigVMBlueprint = RigVMClientHost.GetObject())
 	{
-		UClass* BPClass = RigVMBlueprint->GetClass();
-		UClass* BaseClass = GetRigVMBlueprintClass();
-		if(BPClass != GetRigVMBlueprintClass() && !BPClass->IsChildOf(BaseClass))
+		UClass* BPClass = RigVMBlueprint->GetObject()->GetClass();
+		UClass* BaseClass = GetRigVMAssetClass();
+		if(BPClass != GetRigVMAssetClass() && !BPClass->IsChildOf(BaseClass))
 		{
 			return;
 		}
@@ -664,7 +662,7 @@ void FRigVMEditorModule::GetNodeContextMenuActions(IRigVMClientHost* RigVMClient
 	GetNodeDisplayContextMenuActions(RigVMClientHost, EdGraphNode, ModelNode, Menu);
 }
 
-void FRigVMEditorModule::GetNodeWorkflowContextMenuActions(IRigVMClientHost* RigVMClientHost, const URigVMEdGraphNode* EdGraphNode, URigVMNode* ModelNode, UToolMenu* Menu) const
+void FRigVMEditorModule::GetNodeWorkflowContextMenuActions(TScriptInterface<IRigVMClientHost> RigVMClientHost, const URigVMEdGraphNode* EdGraphNode, URigVMNode* ModelNode, UToolMenu* Menu) const
 {
 	const TArray<FRigVMUserWorkflow> Workflows = ModelNode->GetSupportedWorkflows(ERigVMUserWorkflowType::NodeContext, ModelNode);
 	if(!Workflows.IsEmpty())
@@ -698,9 +696,9 @@ void FRigVMEditorModule::GetNodeWorkflowContextMenuActions(IRigVMClientHost* Rig
 	}
 }
 
-void FRigVMEditorModule::GetNodeEventsContextMenuActions(IRigVMClientHost* RigVMClientHost, const URigVMEdGraphNode* EdGraphNode, URigVMNode* ModelNode, UToolMenu* Menu) const
+void FRigVMEditorModule::GetNodeEventsContextMenuActions(TScriptInterface<IRigVMClientHost> RigVMClientHost, const URigVMEdGraphNode* EdGraphNode, URigVMNode* ModelNode, UToolMenu* Menu) const
 {
-	URigVMBlueprint* RigVMBlueprint = Cast<URigVMBlueprint>(RigVMClientHost);
+	TRigVMAssetInterfacePtr RigVMBlueprint = RigVMClientHost.GetObject();
 	if(RigVMBlueprint && ModelNode->IsEvent())
 	{
 		const FName& EventName = ModelNode->GetEventName();
@@ -754,7 +752,7 @@ void FRigVMEditorModule::GetNodeEventsContextMenuActions(IRigVMClientHost* RigVM
 	}
 }
 
-void FRigVMEditorModule::GetNodeDefaultValueContextMenuActions(IRigVMClientHost* RigVMClientHost, const URigVMEdGraphNode* EdGraphNode, URigVMNode* ModelNode, UToolMenu* Menu) const
+void FRigVMEditorModule::GetNodeDefaultValueContextMenuActions(TScriptInterface<IRigVMClientHost> RigVMClientHost, const URigVMEdGraphNode* EdGraphNode, URigVMNode* ModelNode, UToolMenu* Menu) const
 {
 	if(!CVarRigVMEnablePinOverrides.GetValueOnAnyThread())
 	{
@@ -845,7 +843,7 @@ void FRigVMEditorModule::GetNodeDefaultValueContextMenuActions(IRigVMClientHost*
 	);
 }
 
-void FRigVMEditorModule::GetNodeConversionContextMenuActions(IRigVMClientHost* RigVMClientHost, const URigVMEdGraphNode* EdGraphNode, URigVMNode* ModelNode, UToolMenu* Menu) const
+void FRigVMEditorModule::GetNodeConversionContextMenuActions(TScriptInterface<IRigVMClientHost> RigVMClientHost, const URigVMEdGraphNode* EdGraphNode, URigVMNode* ModelNode, UToolMenu* Menu) const
 {
 	if(URigVMDispatchNode* DispatchNode = Cast<URigVMDispatchNode>(ModelNode))
 	{
@@ -890,10 +888,9 @@ void FRigVMEditorModule::GetNodeConversionContextMenuActions(IRigVMClientHost* R
 	}
 }
 
-void FRigVMEditorModule::GetNodeDebugContextMenuActions(IRigVMClientHost* RigVMClientHost, const URigVMEdGraphNode* EdGraphNode, URigVMNode* ModelNode, UToolMenu* Menu) const
+void FRigVMEditorModule::GetNodeDebugContextMenuActions(TScriptInterface<IRigVMClientHost> RigVMClientHost, const URigVMEdGraphNode* EdGraphNode, URigVMNode* ModelNode, UToolMenu* Menu) const
 {
-	URigVMBlueprint* RigVMBlueprint = Cast<URigVMBlueprint>(RigVMClientHost);
-	if(RigVMBlueprint)
+	if(TRigVMAssetInterfacePtr RigVMBlueprint = RigVMClientHost.GetObject())
 	{
 		FToolMenuSection& DebugSection = Menu->AddSection("RigVMEditorContextMenuDebug", LOCTEXT("DebugHeader", "Debug"));
 		bool bNoneHasBreakpoint = true;
@@ -959,10 +956,9 @@ void FRigVMEditorModule::GetNodeDebugContextMenuActions(IRigVMClientHost* RigVMC
 	}
 }
 
-void FRigVMEditorModule::GetNodeVariablesContextMenuActions(IRigVMClientHost* RigVMClientHost, const URigVMEdGraphNode* EdGraphNode, URigVMNode* ModelNode, UToolMenu* Menu) const
+void FRigVMEditorModule::GetNodeVariablesContextMenuActions(TScriptInterface<IRigVMClientHost> RigVMClientHost, const URigVMEdGraphNode* EdGraphNode, URigVMNode* ModelNode, UToolMenu* Menu) const
 {
-	URigVMBlueprint* RigVMBlueprint = Cast<URigVMBlueprint>(RigVMClientHost);
-	if (RigVMBlueprint)
+	if (TRigVMAssetInterfacePtr RigVMBlueprint = RigVMClientHost.GetObject())
 	{
 		const URigVMGraph* Model = ModelNode->GetGraph();
 		URigVMController* Controller = RigVMBlueprint->GetController(Model);
@@ -1000,11 +996,11 @@ void FRigVMEditorModule::GetNodeVariablesContextMenuActions(IRigVMClientHost* Ri
 							if(!ExternalVariables.IsEmpty())
 							{
 								FScopedTransaction Transaction(LOCTEXT("MakeVariablesFromFunctionReferenceNode", "Create required variables"));
-								RigVMBlueprint->Modify();
+								RigVMBlueprint->GetObject()->Modify();
 
 								if (URigVMLibraryNode* LibraryNode = FunctionReferenceNode->LoadReferencedNode())
 								{
-									URigVMBlueprint* ReferencedBlueprint = LibraryNode->GetTypedOuter<URigVMBlueprint>();
+									TRigVMAssetInterfacePtr ReferencedBlueprint = IRigVMAssetInterface::GetInterfaceOuter(LibraryNode);
 								   // ReferencedBlueprint != RigVMBlueprint - since only FunctionReferenceNodes from other assets have the potential to be unmapped
                             
 								   for(const FRigVMExternalVariable& ExternalVariable : ExternalVariables)
@@ -1012,9 +1008,9 @@ void FRigVMEditorModule::GetNodeVariablesContextMenuActions(IRigVMClientHost* Ri
 									   FString DefaultValue;
 									   if(ReferencedBlueprint)
 									   {
-										   for(const FBPVariableDescription& NewVariable : ReferencedBlueprint->NewVariables)
+										   for(const FRigVMGraphVariableDescription& NewVariable : ReferencedBlueprint->GetAssetVariables())
 										   {
-											   if(NewVariable.VarName == ExternalVariable.Name)
+											   if(NewVariable.Name == ExternalVariable.Name)
 											   {
 												   DefaultValue = NewVariable.DefaultValue;
 												   break;
@@ -1030,7 +1026,7 @@ void FRigVMEditorModule::GetNodeVariablesContextMenuActions(IRigVMClientHost* Ri
 								   }
 								}
 
-								FBlueprintEditorUtils::MarkBlueprintAsModified(RigVMBlueprint);
+								RigVMBlueprint->MarkAssetAsModified();
 							}
                         
 						})
@@ -1041,7 +1037,7 @@ void FRigVMEditorModule::GetNodeVariablesContextMenuActions(IRigVMClientHost* Ri
 	}
 }
 
-void FRigVMEditorModule::GetNodeTemplatesContextMenuActions(IRigVMClientHost* RigVMClientHost, const URigVMEdGraphNode* EdGraphNode, URigVMNode* ModelNode, UToolMenu* Menu) const
+void FRigVMEditorModule::GetNodeTemplatesContextMenuActions(TScriptInterface<IRigVMClientHost> RigVMClientHost, const URigVMEdGraphNode* EdGraphNode, URigVMNode* ModelNode, UToolMenu* Menu) const
 {
 	if (URigVMTemplateNode* TemplateNode = Cast<URigVMTemplateNode>(EdGraphNode->GetModelNode()))
 	{
@@ -1065,9 +1061,9 @@ void FRigVMEditorModule::GetNodeTemplatesContextMenuActions(IRigVMClientHost* Ri
 	}
 }
 
-void FRigVMEditorModule::GetNodeOrganizationContextMenuActions(IRigVMClientHost* RigVMClientHost, const URigVMEdGraphNode* EdGraphNode, URigVMNode* ModelNode, UToolMenu* Menu) const
+void FRigVMEditorModule::GetNodeOrganizationContextMenuActions(TScriptInterface<IRigVMClientHost> RigVMClientHost, const URigVMEdGraphNode* EdGraphNode, URigVMNode* ModelNode, UToolMenu* Menu) const
 {
-	URigVMBlueprint* RigVMBlueprint = Cast<URigVMBlueprint>(RigVMClientHost);
+	TRigVMAssetInterfacePtr RigVMBlueprint = RigVMClientHost.GetObject();
 	URigVMGraph* Model = ModelNode->GetGraph();
 	URigVMController* Controller = RigVMClientHost->GetRigVMClient()->GetController(Model);
 
@@ -1336,7 +1332,7 @@ void FRigVMEditorModule::GetNodeOrganizationContextMenuActions(IRigVMClientHost*
 	}));
 }
 
-void FRigVMEditorModule::GetNodeVariantContextMenuActions(IRigVMClientHost* RigVMClientHost, const URigVMEdGraphNode* EdGraphNode, URigVMNode* ModelNode, UToolMenu* Menu) const
+void FRigVMEditorModule::GetNodeVariantContextMenuActions(TScriptInterface<IRigVMClientHost> RigVMClientHost, const URigVMEdGraphNode* EdGraphNode, URigVMNode* ModelNode, UToolMenu* Menu) const
 {
 	const URigVMGraph* Model = ModelNode->GetGraph();
 	URigVMController* Controller = RigVMClientHost->GetRigVMClient()->GetController(Model);
@@ -1424,7 +1420,7 @@ void FRigVMEditorModule::GetNodeVariantContextMenuActions(IRigVMClientHost* RigV
 	}
 }
 
-void FRigVMEditorModule::GetNodeTestContextMenuActions(IRigVMClientHost* RigVMClientHost, const URigVMEdGraphNode* EdGraphNode, URigVMNode* ModelNode, UToolMenu* Menu) const
+void FRigVMEditorModule::GetNodeTestContextMenuActions(TScriptInterface<IRigVMClientHost> RigVMClientHost, const URigVMEdGraphNode* EdGraphNode, URigVMNode* ModelNode, UToolMenu* Menu) const
 {
 	// this struct is only available in EngineTest for now
 	static const FString TraitObjectPath = TEXT("/Script/EngineTestEditor.EngineTestRigVM_SimpleTrait");
@@ -1454,9 +1450,9 @@ void FRigVMEditorModule::GetNodeTestContextMenuActions(IRigVMClientHost* RigVMCl
 	);
 }
 
-void FRigVMEditorModule::GetNodeDisplayContextMenuActions(IRigVMClientHost* RigVMClientHost, const URigVMEdGraphNode* EdGraphNode, URigVMNode* ModelNode, UToolMenu* Menu) const
+void FRigVMEditorModule::GetNodeDisplayContextMenuActions(TScriptInterface<IRigVMClientHost> RigVMClientHost, const URigVMEdGraphNode* EdGraphNode, URigVMNode* ModelNode, UToolMenu* Menu) const
 {
-	URigVMBlueprint* Blueprint = Cast<URigVMBlueprint>(RigVMClientHost);
+	TRigVMAssetInterfacePtr Blueprint = RigVMClientHost.GetObject();
 	if(Blueprint == nullptr)
 	{
 		return;
@@ -1472,15 +1468,15 @@ void FRigVMEditorModule::GetNodeDisplayContextMenuActions(IRigVMClientHost* RigV
 		FUIAction(FExecuteAction::CreateLambda([Blueprint]()
 		{
 			FScopedTransaction Transaction(LOCTEXT("ToggleProfiler", "Toggle Profiler"));
-			Blueprint->Modify();
-			Blueprint->VMRuntimeSettings.bEnableProfiling = !Blueprint->VMRuntimeSettings.bEnableProfiling;
+			Blueprint.GetObject()->Modify();
+			Blueprint->GetVMRuntimeSettings().bEnableProfiling = !Blueprint->GetVMRuntimeSettings().bEnableProfiling;
 			Blueprint->PropagateRuntimeSettingsFromBPToInstances();
 			Blueprint->RequestAutoVMRecompilation();
 		}),
 		FCanExecuteAction(),
 		FGetActionCheckState::CreateLambda([Blueprint]() -> ECheckBoxState
 		{
-			return Blueprint->VMRuntimeSettings.bEnableProfiling ? ECheckBoxState::Checked : ECheckBoxState::Unchecked;
+			return Blueprint->GetVMRuntimeSettings().bEnableProfiling ? ECheckBoxState::Checked : ECheckBoxState::Unchecked;
 		})),
 		EUserInterfaceActionType::ToggleButton
 	);
@@ -1493,15 +1489,15 @@ void FRigVMEditorModule::GetNodeDisplayContextMenuActions(IRigVMClientHost* RigV
 		FUIAction(FExecuteAction::CreateLambda([Blueprint]()
 		{
 			FScopedTransaction Transaction(LOCTEXT("ToggleTagDisplayMode", "Toggle Tag Display Mode"));
-			Blueprint->Modify();
+			Blueprint->GetObject()->Modify();
 
-			if(Blueprint->RigGraphDisplaySettings.TagDisplayMode == ERigVMTagDisplayMode::All)
+			if(Blueprint->GetRigGraphDisplaySettings().TagDisplayMode == ERigVMTagDisplayMode::All)
 			{
-				Blueprint->RigGraphDisplaySettings.TagDisplayMode = ERigVMTagDisplayMode::DeprecationOnly;
+				Blueprint->GetRigGraphDisplaySettings().TagDisplayMode = ERigVMTagDisplayMode::DeprecationOnly;
 			}
 			else
 			{
-				Blueprint->RigGraphDisplaySettings.TagDisplayMode = ERigVMTagDisplayMode::All;
+				Blueprint->GetRigGraphDisplaySettings().TagDisplayMode = ERigVMTagDisplayMode::All;
 			}
 
 			// this causes all nodes to refresh
@@ -1510,13 +1506,13 @@ void FRigVMEditorModule::GetNodeDisplayContextMenuActions(IRigVMClientHost* RigV
 		FCanExecuteAction(),
 		FGetActionCheckState::CreateLambda([Blueprint]() -> ECheckBoxState
 		{
-			return Blueprint->RigGraphDisplaySettings.TagDisplayMode == ERigVMTagDisplayMode::All ? ECheckBoxState::Checked : ECheckBoxState::Unchecked;
+			return Blueprint->GetRigGraphDisplaySettings().TagDisplayMode == ERigVMTagDisplayMode::All ? ECheckBoxState::Checked : ECheckBoxState::Unchecked;
 		})),
 		EUserInterfaceActionType::ToggleButton
 	);
 }
 
-void FRigVMEditorModule::GetPinContextMenuActions(IRigVMClientHost* RigVMClientHost, const UEdGraphPin* EdGraphPin, URigVMPin* ModelPin, UToolMenu* Menu) const
+void FRigVMEditorModule::GetPinContextMenuActions(TScriptInterface<IRigVMClientHost> RigVMClientHost, const UEdGraphPin* EdGraphPin, URigVMPin* ModelPin, UToolMenu* Menu) const
 {
 	GetExposedPinContextMenuActions(RigVMClientHost, EdGraphPin, ModelPin, Menu);
 	GetPinWorkflowContextMenuActions(RigVMClientHost, EdGraphPin, ModelPin, Menu);
@@ -1530,7 +1526,7 @@ void FRigVMEditorModule::GetPinContextMenuActions(IRigVMClientHost* RigVMClientH
 	GetPinInjectedNodesContextMenuActions(RigVMClientHost, EdGraphPin, ModelPin, Menu);
 }
 
-void FRigVMEditorModule::GetExposedPinContextMenuActions(IRigVMClientHost* RigVMClientHost, const UEdGraphPin* EdGraphPin, URigVMPin* ModelPin, UToolMenu* Menu) const
+void FRigVMEditorModule::GetExposedPinContextMenuActions(TScriptInterface<IRigVMClientHost> RigVMClientHost, const UEdGraphPin* EdGraphPin, URigVMPin* ModelPin, UToolMenu* Menu) const
 {
 	if (!ModelPin->IsRootPin() || !ModelPin->GetNode()->IsA<URigVMFunctionInterfaceNode>())
 	{
@@ -1556,7 +1552,7 @@ void FRigVMEditorModule::GetExposedPinContextMenuActions(IRigVMClientHost* RigVM
 	));
 }
 
-void FRigVMEditorModule::GetPinWorkflowContextMenuActions(IRigVMClientHost* RigVMClientHost, const UEdGraphPin* EdGraphPin, URigVMPin* ModelPin, UToolMenu* Menu) const
+void FRigVMEditorModule::GetPinWorkflowContextMenuActions(TScriptInterface<IRigVMClientHost> RigVMClientHost, const UEdGraphPin* EdGraphPin, URigVMPin* ModelPin, UToolMenu* Menu) const
 {
 	const TArray<FRigVMUserWorkflow> Workflows = ModelPin->GetNode()->GetSupportedWorkflows(ERigVMUserWorkflowType::PinContext, ModelPin);
 	if(!Workflows.IsEmpty())
@@ -1591,10 +1587,9 @@ void FRigVMEditorModule::GetPinWorkflowContextMenuActions(IRigVMClientHost* RigV
 	}
 }
 
-void FRigVMEditorModule::GetPinDebugContextMenuActions(IRigVMClientHost* RigVMClientHost, const UEdGraphPin* EdGraphPin, URigVMPin* ModelPin, UToolMenu* Menu) const
+void FRigVMEditorModule::GetPinDebugContextMenuActions(TScriptInterface<IRigVMClientHost> RigVMClientHost, const UEdGraphPin* EdGraphPin, URigVMPin* ModelPin, UToolMenu* Menu) const
 {
-	URigVMBlueprint* RigVMBlueprint = Cast<URigVMBlueprint>(RigVMClientHost);
-	if (RigVMBlueprint)
+	if (TRigVMAssetInterfacePtr RigVMBlueprint = RigVMClientHost.GetObject())
 	{
 		const bool bIsEditablePin = !ModelPin->IsExecuteContext() && !ModelPin->IsWildCard();
 
@@ -1604,7 +1599,7 @@ void FRigVMEditorModule::GetPinDebugContextMenuActions(IRigVMClientHost* RigVMCl
 			{
 				// Add the watch pin / unwatch pin menu items
 				FToolMenuSection& Section = Menu->AddSection("RigVMEditorContextMenuWatches", LOCTEXT("WatchesHeader", "Watches"));
-				if (FKismetDebugUtilities::IsPinBeingWatched(RigVMBlueprint, EdGraphPin))
+				if (RigVMBlueprint->IsPinBeingWatched(EdGraphPin))
 				{
 					Section.AddMenuEntry(FGraphEditorCommands::Get().StopWatchingPin);
 				}
@@ -1617,7 +1612,7 @@ void FRigVMEditorModule::GetPinDebugContextMenuActions(IRigVMClientHost* RigVMCl
 	}
 }
 
-void FRigVMEditorModule::GetPinArrayContextMenuActions(IRigVMClientHost* RigVMClientHost, const UEdGraphPin* EdGraphPin, URigVMPin* ModelPin, UToolMenu* Menu) const
+void FRigVMEditorModule::GetPinArrayContextMenuActions(TScriptInterface<IRigVMClientHost> RigVMClientHost, const UEdGraphPin* EdGraphPin, URigVMPin* ModelPin, UToolMenu* Menu) const
 {
 	URigVMController* Controller = RigVMClientHost->GetRigVMClient()->GetController(ModelPin->GetGraph());
 	if (ModelPin->IsArray() && !ModelPin->IsExecuteContext())
@@ -1658,9 +1653,9 @@ void FRigVMEditorModule::GetPinArrayContextMenuActions(IRigVMClientHost* RigVMCl
 	}
 }
 
-void FRigVMEditorModule::GetPinAggregateContextMenuActions(IRigVMClientHost* RigVMClient, const UEdGraphPin* EdGraphPin, URigVMPin* ModelPin, UToolMenu* Menu) const
+void FRigVMEditorModule::GetPinAggregateContextMenuActions(TScriptInterface<IRigVMClientHost> RigVMClientHost, const UEdGraphPin* EdGraphPin, URigVMPin* ModelPin, UToolMenu* Menu) const
 {
-	URigVMController* Controller = RigVMClient->GetRigVMClient()->GetController(ModelPin->GetGraph());
+	URigVMController* Controller = RigVMClientHost->GetRigVMClient()->GetController(ModelPin->GetGraph());
 	if (Cast<URigVMAggregateNode>(ModelPin->GetNode()))
 	{
 		FToolMenuSection& Section = Menu->AddSection("RigVMEditorContextMenuAggregatePin", LOCTEXT("AggregatePin", "Aggregates"));
@@ -1676,7 +1671,7 @@ void FRigVMEditorModule::GetPinAggregateContextMenuActions(IRigVMClientHost* Rig
 	}
 }
 
-void FRigVMEditorModule::GetPinTemplateContextMenuActions(IRigVMClientHost* RigVMClientHost, const UEdGraphPin* EdGraphPin, URigVMPin* ModelPin, UToolMenu* Menu) const
+void FRigVMEditorModule::GetPinTemplateContextMenuActions(TScriptInterface<IRigVMClientHost> RigVMClientHost, const UEdGraphPin* EdGraphPin, URigVMPin* ModelPin, UToolMenu* Menu) const
 {
 	if (URigVMTemplateNode* TemplateNode = Cast<URigVMTemplateNode>(ModelPin->GetNode()))
 	{
@@ -1726,7 +1721,7 @@ void FRigVMEditorModule::GetPinTemplateContextMenuActions(IRigVMClientHost* RigV
 	}
 }
 
-void FRigVMEditorModule::GetPinConversionContextMenuActions(IRigVMClientHost* RigVMClientHost, const UEdGraphPin* EdGraphPin, URigVMPin* ModelPin, UToolMenu* Menu) const
+void FRigVMEditorModule::GetPinConversionContextMenuActions(TScriptInterface<IRigVMClientHost> RigVMClientHost, const UEdGraphPin* EdGraphPin, URigVMPin* ModelPin, UToolMenu* Menu) const
 {
 	if (URigVMRerouteNode* RerouteNode = Cast<URigVMRerouteNode>(ModelPin->GetNode()))
 	{
@@ -1784,9 +1779,9 @@ void FRigVMEditorModule::GetPinConversionContextMenuActions(IRigVMClientHost* Ri
 	}
 }
 
-void FRigVMEditorModule::GetPinVariableContextMenuActions(IRigVMClientHost* RigVMClientHost, const UEdGraphPin* EdGraphPin, URigVMPin* ModelPin, UToolMenu* Menu) const
+void FRigVMEditorModule::GetPinVariableContextMenuActions(TScriptInterface<IRigVMClientHost> RigVMClientHost, const UEdGraphPin* EdGraphPin, URigVMPin* ModelPin, UToolMenu* Menu) const
 {
-	if(URigVMBlueprint* RigVMBlueprint = Cast<URigVMBlueprint>(RigVMClientHost))
+	if(TRigVMAssetInterfacePtr RigVMBlueprint = RigVMClientHost.GetObject())
 	{
 		const bool bIsEditablePin = !ModelPin->IsExecuteContext() && !ModelPin->IsWildCard();
 
@@ -1845,7 +1840,7 @@ void FRigVMEditorModule::GetPinVariableContextMenuActions(IRigVMClientHost* RigV
 
 					TSharedRef<SRigVMGraphVariableBinding> VariableBindingWidget =
 						SNew(SRigVMGraphVariableBinding)
-						.Blueprint(RigVMBlueprint)
+						.Asset(RigVMBlueprint)
 						.ModelPins({ModelPin})
 						.CanRemoveBinding(false);
 
@@ -1856,7 +1851,7 @@ void FRigVMEditorModule::GetPinVariableContextMenuActions(IRigVMClientHost* RigV
 	}
 }
 
-void FRigVMEditorModule::GetPinResetDefaultContextMenuActions(IRigVMClientHost* RigVMClientHost, const UEdGraphPin* EdGraphPin, URigVMPin* ModelPin, UToolMenu* Menu) const
+void FRigVMEditorModule::GetPinResetDefaultContextMenuActions(TScriptInterface<IRigVMClientHost> RigVMClientHost, const UEdGraphPin* EdGraphPin, URigVMPin* ModelPin, UToolMenu* Menu) const
 {
 	if (Cast<URigVMUnitNode>(ModelPin->GetNode()) != nullptr || 
 		Cast<URigVMDispatchNode>(ModelPin->GetNode()) != nullptr || 
@@ -1895,7 +1890,7 @@ void FRigVMEditorModule::GetPinResetDefaultContextMenuActions(IRigVMClientHost*
 	}
 }
 
-void FRigVMEditorModule::GetPinInjectedNodesContextMenuActions(IRigVMClientHost* RigVMClientHost, const UEdGraphPin* EdGraphPin, URigVMPin* ModelPin, UToolMenu* Menu) const
+void FRigVMEditorModule::GetPinInjectedNodesContextMenuActions(TScriptInterface<IRigVMClientHost> RigVMClientHost, const UEdGraphPin* EdGraphPin, URigVMPin* ModelPin, UToolMenu* Menu) const
 {
 	if (ModelPin->GetRootPin() == ModelPin && (
 	Cast<URigVMUnitNode>(ModelPin->GetNode()) != nullptr ||
@@ -2168,11 +2163,20 @@ void FRigVMEditorModule::GetContextMenuActions(const URigVMEdGraphSchema* Schema
 		{
 			if(const UEdGraph* Graph = Context->Graph)
 			{
-				if(IRigVMClientHost* RigVMClientHost = Graph->GetImplementingOuter<IRigVMClientHost>())
+				// Can't get the UObject with GetImplementingOuter, which I need to create a TScriptInterface<IRigVMClientHost>
+				UObject* Outer = Graph->GetOuter();
+				while (Outer && !Outer->Implements<URigVMClientHost>())
+				{
+					Outer = Outer->GetOuter();
+				}
+				if (Outer)
 				{
-					if (URigVMPin* ModelPin = RigVMClientHost->GetRigVMClient()->GetModel(Graph)->FindPin(InGraphPin->GetName()))
+					if(TScriptInterface<IRigVMClientHost> RigVMClientHost = Outer)
 					{
-						GetPinContextMenuActions(RigVMClientHost, InGraphPin, ModelPin, Menu);
+						if (URigVMPin* ModelPin = RigVMClientHost->GetRigVMClient()->GetModel(Graph)->FindPin(InGraphPin->GetName()))
+						{
+							GetPinContextMenuActions(RigVMClientHost, InGraphPin, ModelPin, Menu);
+						}
 					}
 				}
 			}
@@ -2181,11 +2185,20 @@ void FRigVMEditorModule::GetContextMenuActions(const URigVMEdGraphSchema* Schema
 		{
 			if (const UEdGraph* Graph = Context->Graph)
 			{
-				if (IRigVMClientHost* RigVMClientHost = Graph->GetImplementingOuter<IRigVMClientHost>())
+				// Can't get the UObject with GetImplementingOuter, which I need to create a TScriptInterface<IRigVMClientHost>
+				UObject* Outer = Graph->GetOuter();
+				while (Outer && !Outer->Implements<URigVMClientHost>())
 				{
-					if(URigVMNode* ModelNode = EdGraphNode->GetModelNode())
+					Outer = Outer->GetOuter();
+				}
+				if (Outer)
+				{
+					if (TScriptInterface<IRigVMClientHost> RigVMClientHost = Outer)
 					{
-						GetNodeContextMenuActions(RigVMClientHost, EdGraphNode, ModelNode, Menu);
+						if(URigVMNode* ModelNode = EdGraphNode->GetModelNode())
+						{
+							GetNodeContextMenuActions(RigVMClientHost, EdGraphNode, ModelNode, Menu);
+						}
 					}
 				}
 			}
@@ -2266,7 +2279,7 @@ void FRigVMEditorModule::PostChange(const UUserDefinedStruct* Changed, FStructur
 		return;
 	}
 
-	TArray<URigVMBlueprint*> BlueprintsToRefresh;
+	TArray<TRigVMAssetInterfacePtr> BlueprintsToRefresh;
 	TArray<UEdGraph*> EdGraphsToRefresh;
 
 	for (TObjectIterator<URigVMPin> It(RF_Transient | RF_ClassDefaultObject, /** bIncludeDerivedClasses */ true, /** InternalExcludeFlags */ EInternalObjectFlags::Garbage); It; ++It)
@@ -2275,7 +2288,7 @@ void FRigVMEditorModule::PostChange(const UUserDefinedStruct* Changed, FStructur
 		// GetCPPTypeObject also makes sure the pin's type information is update to date
 		if (Pin && Pin->GetCPPTypeObject() == Changed)
 		{
-			if (URigVMBlueprint* RigVMBlueprint = Pin->GetTypedOuter<URigVMBlueprint>())
+			if (TRigVMAssetInterfacePtr RigVMBlueprint = IRigVMAssetInterface::GetInterfaceOuter(Pin))
 			{
 				BlueprintsToRefresh.AddUnique(RigVMBlueprint);
 				
@@ -2284,7 +2297,7 @@ void FRigVMEditorModule::PostChange(const UUserDefinedStruct* Changed, FStructur
 				if (Pin->GetGraph() == RigVMBlueprint->GetLocalFunctionLibrary())
 				{
 					TArray< TSoftObjectPtr<URigVMFunctionReferenceNode> > References;
-					References = RigVMBlueprint->RigVMClient.GetOrCreateFunctionLibrary(false)->GetReferencesForFunction(Pin->GetNode()->GetFName());
+					References = RigVMBlueprint->GetRigVMClientHost()->GetRigVMClient()->GetOrCreateFunctionLibrary(false)->GetReferencesForFunction(Pin->GetNode()->GetFName());
 					
 					for (const TSoftObjectPtr<URigVMFunctionReferenceNode>& Reference : References)
 					{
@@ -2294,7 +2307,7 @@ void FRigVMEditorModule::PostChange(const UUserDefinedStruct* Changed, FStructur
 							continue;
 						}
 						
-						if (URigVMBlueprint* FunctionUserBlueprint = RefNode->GetTypedOuter<URigVMBlueprint>())
+						if (TRigVMAssetInterfacePtr FunctionUserBlueprint = IRigVMAssetInterface::GetInterfaceOuter(RefNode))
 						{
 							BlueprintsToRefresh.AddUnique(FunctionUserBlueprint);
 						}
@@ -2309,7 +2322,7 @@ void FRigVMEditorModule::PostChange(const UUserDefinedStruct* Changed, FStructur
 		}
 	}
 
-	for (URigVMBlueprint* RigVMBlueprint : BlueprintsToRefresh)
+	for (TRigVMAssetInterfacePtr RigVMBlueprint : BlueprintsToRefresh)
 	{
 		RigVMBlueprint->OnRigVMRegistryChanged();
 		(void)RigVMBlueprint->MarkPackageDirty();
@@ -2321,12 +2334,11 @@ void FRigVMEditorModule::PostChange(const UUserDefinedStruct* Changed, FStructur
 		Graph->NotifyGraphChanged();
 	}
 
-	for (URigVMBlueprint* RigVMBlueprint : BlueprintsToRefresh)
+	for (TRigVMAssetInterfacePtr RigVMBlueprint : BlueprintsToRefresh)
 	{
 		// this should make sure variables in BP are updated with the latest struct object
 		// otherwise RigVMCompiler validation would complain about variable type - pin type mismatch
-		FCompilerResultsLog	ResultsLog;
-		FKismetEditorUtilities::CompileBlueprint(RigVMBlueprint, EBlueprintCompileOptions::None, &ResultsLog);
+		FCompilerResultsLog	ResultsLog = RigVMBlueprint->CompileBlueprint();
 		
 		// BP compiler always initialize the new CDO by copying from the old CDO,
 		// however, in case that a BP variable type has changed, the data old CDO would be invalid because
@@ -2336,19 +2348,22 @@ void FRigVMEditorModule::PostChange(const UUserDefinedStruct* Changed, FStructur
 		// so to simplify things, here we just reset all rigs upon error
 		if (ResultsLog.NumErrors > 0)
 		{
-			URigVMBlueprintGeneratedClass* RigClass = RigVMBlueprint->GetRigVMBlueprintGeneratedClass();
-			URigVMHost* CDO = Cast<URigVMHost>(RigClass->GetDefaultObject(true /* create if needed */));
-			if (CDO->GetVM() != nullptr)
+			if (URigVM* VM = RigVMBlueprint->GetVM(false))
 			{
-				CDO->GetVM()->Reset(CDO->GetRigVMExtendedExecuteContext());
+				if (FRigVMExtendedExecuteContext* Context = RigVMBlueprint->GetRigVMExtendedExecuteContext())
+				{
+					VM->Reset(*Context);
+				}
 			}
-			TArray<UObject*> ArchetypeInstances;
-			CDO->GetArchetypeInstances(ArchetypeInstances);
-			for (UObject* Instance : ArchetypeInstances)
+			TArray<UObject*> Instances = RigVMBlueprint->GetArchetypeInstances(true, true);
+			for (UObject* Instance : Instances)
 			{
-				if (URigVMHost* InstanceHost = Cast<URigVMHost>(Instance))
+				if (URigVMHost* RigVMHost = Cast<URigVMHost>(Instance))
 				{
-					InstanceHost->GetVM()->Reset(InstanceHost->GetRigVMExtendedExecuteContext());
+					if (RigVMHost->GetVM())
+					{
+						RigVMHost->GetVM()->Reset(RigVMHost->GetRigVMExtendedExecuteContext());
+					}
 				}
 			}
 		}
@@ -2360,17 +2375,17 @@ TArray<FRigVMEditorModule::FRigVMEditorToolbarExtender>& FRigVMEditorModule::Get
 	return RigVMEditorToolbarExtenders;
 }
 
-void FRigVMEditorModule::CreateRootGraphIfRequired(URigVMBlueprint* InBlueprint) const
+void FRigVMEditorModule::CreateRootGraphIfRequired(TRigVMAssetInterfacePtr InBlueprint) const
 {
 	if(InBlueprint == nullptr)
 	{
 		return;
 	}
 
-	const URigVMBlueprint* BlueprintCDO = GetRigVMBlueprintCDO();
-	UClass* EdGraphClass = BlueprintCDO->GetRigVMEdGraphClass();
+	
+	UClass* EdGraphClass = InBlueprint->GetRigVMClientHost()->GetRigVMEdGraphClass();
 
-	for(const UEdGraph* EdGraph : InBlueprint->UbergraphPages)
+	for(const UEdGraph* EdGraph : InBlueprint->GetUberGraphs())
 	{
 		if(EdGraph->IsA(EdGraphClass))
 		{
@@ -2378,20 +2393,21 @@ void FRigVMEditorModule::CreateRootGraphIfRequired(URigVMBlueprint* InBlueprint)
 		}
 	}
 
-	UClass* EdGraphSchemaClass = BlueprintCDO->GetRigVMEdGraphSchemaClass();
+	UClass* EdGraphSchemaClass = InBlueprint->GetRigVMClientHost()->GetRigVMEdGraphSchemaClass();
 	const URigVMEdGraphSchema* SchemaCDO = CastChecked<URigVMEdGraphSchema>(EdGraphSchemaClass->GetDefaultObject());
 	
 	// add an initial graph for us to work in
-	UEdGraph* ControlRigGraph = FBlueprintEditorUtils::CreateNewGraph(InBlueprint, SchemaCDO->GetRootGraphName(), EdGraphClass, EdGraphSchemaClass);
+	
+	URigVMEdGraph* ControlRigGraph = FRigVMBlueprintUtils::CreateNewGraph(InBlueprint->GetObject(), SchemaCDO->GetRootGraphName(), EdGraphClass, EdGraphSchemaClass);
 	ControlRigGraph->bAllowDeletion = false;
-	FBlueprintEditorUtils::AddUbergraphPage(InBlueprint, ControlRigGraph);
-	InBlueprint->LastEditedDocuments.AddUnique(ControlRigGraph);
+	InBlueprint->AddUbergraphPage(ControlRigGraph);
+	InBlueprint->GetLastEditedDocuments().AddUnique(ControlRigGraph);
 	InBlueprint->PostLoad();
 }
 
 void FRigVMEditorModule::HandleNewBlueprintCreated(UBlueprint* InBlueprint)
 {
-	CreateRootGraphIfRequired(Cast<URigVMBlueprint>(InBlueprint));
+	CreateRootGraphIfRequired(TRigVMAssetInterfacePtr(InBlueprint));
 }
 
 bool FRigVMEditorModule::ShowWorkflowOptionsDialog(URigVMUserWorkflowOptions* InOptions) const
@@ -2427,7 +2443,12 @@ FConnectionDrawingPolicy* FRigVMEditorModule::CreateConnectionDrawingPolicy(int3
 
 bool FRigVMEditorModule::IsRigVMEditorModuleBase() const
 {
-	return GetRigVMBlueprintClass() == URigVMBlueprint::StaticClass();
+	UClass* Class = GetRigVMAssetClass();
+	if (Class == URigVMAssetInterface::StaticClass() || Class->ImplementsInterface(URigVMAssetInterface::StaticClass()))
+	{
+		return true;
+	}
+	return false;
 }
 
 #undef LOCTEXT_NAMESPACE
diff --git a/Engine/Plugins/Runtime/RigVM/Source/RigVMEditor/Private/Widgets/RigVMTreeToolkitNode.cpp b/Engine/Plugins/Runtime/RigVM/Source/RigVMEditor/Private/Widgets/RigVMTreeToolkitNode.cpp
--- a/Engine/Plugins/Runtime/RigVM/Source/RigVMEditor/Private/Widgets/RigVMTreeToolkitNode.cpp
+++ b/Engine/Plugins/Runtime/RigVM/Source/RigVMEditor/Private/Widgets/RigVMTreeToolkitNode.cpp
@@ -11,13 +11,13 @@ FAssetData FRigVMTreeNode::GetAssetData() const
 	return FRigVMTreeContext::FindAssetFromAnyPath(GetPath(), false);
 }
 
-URigVMBlueprint* FRigVMTreeNode::GetBlueprint() const
+IRigVMAssetInterface* FRigVMTreeNode::GetAsset() const
 {
 	const FAssetData AssetData = FRigVMTreeContext::FindAssetFromAnyPath(GetPath(), true);
 	if(AssetData.IsAssetLoaded())
 	{
 		UObject* Asset = AssetData.GetAsset();
-		return Cast<URigVMBlueprint>(Asset);
+		return Cast<IRigVMAssetInterface>(Asset);
 	}
 	return nullptr;
 }
@@ -381,14 +381,14 @@ const TArray<FRigVMTag>& FRigVMTreePackageNode::GetTags() const
 	{
 		if(IsLoaded())
 		{
-			if(const URigVMBlueprint* Blueprint = GetBlueprint())
+			if(const IRigVMAssetInterface* Blueprint = GetAsset())
 			{
-				Tags = Blueprint->AssetVariant.Tags;
+				Tags = Blueprint->GetAssetVariant().Tags;
 			}
 			else
 			{
 				const FAssetData AssetData = GetAssetData();
-				if(AssetData.FindTag(GET_MEMBER_NAME_CHECKED(URigVMBlueprint, AssetVariant)))
+				if(AssetData.FindTag(TEXT("AssetVariant")))
 				{
 					//const FRigVMVariant& Variant = AssetData.GetTagValueRef<FRigVMVariant>(GET_MEMBER_NAME_CHECKED(URigVMBlueprint, AssetVariant));
 					//Tags = Variant.Tags;
diff --git a/Engine/Plugins/Runtime/RigVM/Source/RigVMEditor/Private/Widgets/SRigVMBulkEditWidget.cpp b/Engine/Plugins/Runtime/RigVM/Source/RigVMEditor/Private/Widgets/SRigVMBulkEditWidget.cpp
--- a/Engine/Plugins/Runtime/RigVM/Source/RigVMEditor/Private/Widgets/SRigVMBulkEditWidget.cpp
+++ b/Engine/Plugins/Runtime/RigVM/Source/RigVMEditor/Private/Widgets/SRigVMBulkEditWidget.cpp
@@ -251,10 +251,10 @@ void SRigVMBulkEditWidget::Tick(const FGeometry& AllottedGeometry, const double
 				const TSharedRef<FRigVMTreePhase>& Phase = GetActivePhase();
 
 				TSharedPtr<TGuardValue<FRigVMReportDelegate>> ReportDelegateGuard; 
-				if(URigVMBlueprint* Blueprint = Task->GetBlueprint(Phase))
+				if(TRigVMAssetInterfacePtr Blueprint = Task->GetRigVMAsset(Phase))
 				{
 					ReportDelegateGuard = MakeShareable(new TGuardValue<FRigVMReportDelegate>(
-						Blueprint->VMCompileSettings.ASTSettings.ReportDelegate,
+						Blueprint->GetVMCompileSettings().ASTSettings.ReportDelegate,
 						FRigVMReportDelegate::CreateLambda(
 						[this](EMessageSeverity::Type InSeverity, UObject* InSubject, const FString& InMessage)
 						{
diff --git a/Engine/Plugins/Runtime/RigVM/Source/RigVMEditor/Private/Widgets/SRigVMEditorGraphExplorer.cpp b/Engine/Plugins/Runtime/RigVM/Source/RigVMEditor/Private/Widgets/SRigVMEditorGraphExplorer.cpp
--- a/Engine/Plugins/Runtime/RigVM/Source/RigVMEditor/Private/Widgets/SRigVMEditorGraphExplorer.cpp
+++ b/Engine/Plugins/Runtime/RigVM/Source/RigVMEditor/Private/Widgets/SRigVMEditorGraphExplorer.cpp
@@ -14,7 +14,7 @@
 // Ick.
 #include "GraphActionNode.h"
 #include "GraphEditorActions.h"
-#include "RigVMBlueprint.h"
+#include "RigVMBlueprintLegacy.h"
 #include "SPinTypeSelector.h"
 #include "Kismet2/BlueprintEditorUtils.h"
 #include "RigVMModel/Nodes/RigVMAggregateNode.h"
@@ -36,7 +36,13 @@
 #define LOCTEXT_NAMESPACE "RigVMGraphExplorer"
 
 
-TSharedRef<FRigVMGraphExplorerDragDropOp> FRigVMGraphExplorerDragDropOp::New(const ::FRigVMExplorerElementKey& InElement, TObjectPtr<URigVMBlueprint> InBlueprint)
+TSharedRef<FRigVMGraphExplorerDragDropOp> FRigVMGraphExplorerDragDropOp::New(const FRigVMExplorerElementKey& InElement, TObjectPtr<URigVMBlueprint> InBlueprint)
+{
+	const TRigVMAssetInterfacePtr Asset(InBlueprint);
+	return New(InElement, Asset);
+}
+
+TSharedRef<FRigVMGraphExplorerDragDropOp> FRigVMGraphExplorerDragDropOp::New(const ::FRigVMExplorerElementKey& InElement, TRigVMAssetInterfacePtr InBlueprint)
 {
 	TSharedRef<FRigVMGraphExplorerDragDropOp> Operation = MakeShared<FRigVMGraphExplorerDragDropOp>();
 	Operation->Element = InElement;
@@ -95,10 +101,10 @@ FReply FRigVMGraphExplorerDragDropOp::DroppedOnPanel(const TSharedRef<SWidget>&
 {
 	if (URigVMEdGraph* TargetRigGraph = Cast<URigVMEdGraph>(&Graph))
 	{
-		if(URigVMBlueprint* Blueprint = TargetRigGraph->GetBlueprint())
+		if(TRigVMAssetInterfacePtr Blueprint = TargetRigGraph->GetAsset())
 		{
 			//  Find the appropriate asset editor where the drop is happening
-			const TArray<IAssetEditorInstance*> AssetEditors = GEditor->GetEditorSubsystem<UAssetEditorSubsystem>()->FindEditorsForAsset(Blueprint);
+			const TArray<IAssetEditorInstance*> AssetEditors = GEditor->GetEditorSubsystem<UAssetEditorSubsystem>()->FindEditorsForAsset(Blueprint->GetObject());
 			IAssetEditorInstance* FocusedAssetEditor = nullptr;
 			for (IAssetEditorInstance* AssetEditor : AssetEditors)
 			{
@@ -124,8 +130,8 @@ FReply FRigVMGraphExplorerDragDropOp::DroppedOnPanel(const TSharedRef<SWidget>&
 
 			if (!FocusedAssetEditor)
 			{
-				GEditor->GetEditorSubsystem<UAssetEditorSubsystem>()->OpenEditorForAsset(Blueprint);
-				FocusedAssetEditor = GEditor->GetEditorSubsystem<UAssetEditorSubsystem>()->FindEditorForAsset(Blueprint, /*bFocusIfOpen =*/true);
+				GEditor->GetEditorSubsystem<UAssetEditorSubsystem>()->OpenEditorForAsset(Blueprint->GetObject());
+				FocusedAssetEditor = GEditor->GetEditorSubsystem<UAssetEditorSubsystem>()->FindEditorForAsset(Blueprint->GetObject(), /*bFocusIfOpen =*/true);
 			}
 			if (FocusedAssetEditor)
 			{
@@ -143,13 +149,13 @@ FReply FRigVMGraphExplorerDragDropOp::DroppedOnPin(const FVector2f& ScreenPositi
 {
 	UEdGraph* TargetGraph = GetHoveredGraph();
 	UEdGraphPin* TargetPin = GetHoveredPin();
-	TObjectPtr<URigVMBlueprint> Blueprint = GetBlueprint();
+	TRigVMAssetInterfacePtr Blueprint = GetRigVMAsset();
 	
 	const URigVMEdGraphSchema* EdSchema = GetDefault<URigVMEdGraphSchema>();
 
 	if (Element.Type == ERigVMExplorerElementType::Variable)
 	{
-		if (FProperty* Property = Blueprint->SkeletonGeneratedClass->FindPropertyByName(*GetElement().Name))
+		if (FProperty* Property = Blueprint->FindGeneratedPropertyByName(*GetElement().Name))
 		{
 			if (EdSchema->RequestVariableDropOnPin(TargetGraph, Property, TargetPin, FDeprecateSlateVector2D(GraphPosition), FDeprecateSlateVector2D(ScreenPosition)))
 			{
@@ -161,7 +167,7 @@ FReply FRigVMGraphExplorerDragDropOp::DroppedOnPin(const FVector2f& ScreenPositi
 	{
 		for (FRigVMGraphVariableDescription& Variable : Blueprint->GetFocusedModel()->GetLocalVariables())
 		{
-			if (Variable.Name == Element.Name)
+			if (Variable.Name == *Element.Name)
 			{
 				if (EdSchema->RequestVariableDropOnPin(TargetGraph, Variable.ToExternalVariable(), TargetPin, FDeprecateSlateVector2D(GraphPosition), FDeprecateSlateVector2D(ScreenPosition)))
 				{
@@ -539,7 +545,7 @@ TArray<const URigVMGraph*> SRigVMEditorGraphExplorer::GetRootGraphs() const
 		return Graphs;
 	}
 
-	FRigVMClient* RigVMClient = Editor->GetRigVMBlueprint()->GetRigVMClient();
+	FRigVMClient* RigVMClient = Editor->GetRigVMAsset()->GetRigVMClientHost()->GetRigVMClient();
 	for (URigVMGraph* Graph : RigVMClient->GetModels())
 	{
 		Graphs.Add(Graph);
@@ -558,7 +564,7 @@ TArray<const URigVMGraph*> SRigVMEditorGraphExplorer::GetChildrenGraphs(const FS
 		return Children;
 	}
 
-	FRigVMClient* RigVMClient = Editor->GetRigVMBlueprint()->GetRigVMClient();
+	FRigVMClient* RigVMClient = Editor->GetRigVMAsset()->GetRigVMClientHost()->GetRigVMClient();
 	URigVMGraph* ParentGraph = RigVMClient->GetModel(InParentGraphPath);
 	if (!ParentGraph)
 	{
@@ -590,7 +596,7 @@ TArray<URigVMNode*> SRigVMEditorGraphExplorer::GetEventNodesInGraph(const FStrin
 		return Events;
 	}
 
-	FRigVMClient* RigVMClient = Editor->GetRigVMBlueprint()->GetRigVMClient();
+	FRigVMClient* RigVMClient = Editor->GetRigVMAsset()->GetRigVMClientHost()->GetRigVMClient();
 	URigVMGraph* ParentGraph = RigVMClient->GetModel(InParentGraphPath);
 	if (!ParentGraph)
 	{
@@ -628,7 +634,7 @@ TArray<URigVMLibraryNode*> SRigVMEditorGraphExplorer::GetFunctions() const
 		return Functions;
 	}
 
-	FRigVMClient* RigVMClient = Editor->GetRigVMBlueprint()->GetRigVMClient();
+	FRigVMClient* RigVMClient = Editor->GetRigVMAsset()->GetRigVMClientHost()->GetRigVMClient();
 	if (URigVMFunctionLibrary* Library = RigVMClient->GetFunctionLibrary())
 	{
 		return Library->GetFunctions();
@@ -646,7 +652,7 @@ TArray<FRigVMGraphVariableDescription> SRigVMEditorGraphExplorer::GetVariables()
 		return Variables;
 	}
 
-	return Editor->GetRigVMBlueprint()->GetMemberVariables();
+	return Editor->GetRigVMAsset()->GetAssetVariables();
 }
 
 TArray<FRigVMGraphVariableDescription> SRigVMEditorGraphExplorer::GetLocalVariables() const
@@ -658,7 +664,7 @@ TArray<FRigVMGraphVariableDescription> SRigVMEditorGraphExplorer::GetLocalVariab
 		return Variables;
 	}
 
-	if (URigVMBlueprint* Blueprint = Editor->GetRigVMBlueprint())
+	if (TRigVMAssetInterfacePtr Blueprint = Editor->GetRigVMAsset())
 	{
 		if (URigVMGraph* Function = Blueprint->GetFocusedModel())
 		{
@@ -678,12 +684,12 @@ FText SRigVMEditorGraphExplorer::GetGraphDisplayName(const FString& InGraphPath)
 		return DisplayName;
 	}
 
-	FRigVMClient* RigVMClient = Editor->GetRigVMBlueprint()->GetRigVMClient();
+	FRigVMClient* RigVMClient = Editor->GetRigVMAsset()->GetRigVMClientHost()->GetRigVMClient();
 	if (const URigVMGraph* Graph = RigVMClient->GetModel(InGraphPath))
 	{
 		const URigVMEdGraphSchema* EdSchema = GetDefault<URigVMEdGraphSchema>();
 		FGraphDisplayInfo DisplayInfo;
-		if (UEdGraph* EdGraph = Editor->GetRigVMBlueprint()->GetEdGraph(Graph))
+		if (UEdGraph* EdGraph = Editor->GetRigVMAsset()->GetEdGraph(Graph))
 		{
 			EdSchema->GetGraphDisplayInformation(*EdGraph, DisplayInfo);
 		}
@@ -703,7 +709,7 @@ FText SRigVMEditorGraphExplorer::GetEventDisplayName(const FString& InNodePath)
 		return DisplayName;
 	}
 
-	FRigVMClient* RigVMClient = Editor->GetRigVMBlueprint()->GetRigVMClient();
+	FRigVMClient* RigVMClient = Editor->GetRigVMAsset()->GetRigVMClientHost()->GetRigVMClient();
 	if (const URigVMNode* Node = RigVMClient->FindNode(InNodePath))
 	{
 		DisplayName = FText::FromName(Node->GetEventName());
@@ -721,7 +727,7 @@ const FSlateBrush* SRigVMEditorGraphExplorer::GetGraphIcon(const FString& InGrap
 		return Brush;
 	}
 
-	FRigVMClient* RigVMClient = Editor->GetRigVMBlueprint()->GetRigVMClient();
+	FRigVMClient* RigVMClient = Editor->GetRigVMAsset()->GetRigVMClientHost()->GetRigVMClient();
 	if (const URigVMGraph* Graph = RigVMClient->GetModel(InGraphPath))
 	{
 		if (Graph->IsRootGraph())
@@ -745,12 +751,12 @@ FText SRigVMEditorGraphExplorer::GetGraphTooltip(const FString& InGraphPath) con
 		return Tooltip;
 	}
 
-	FRigVMClient* RigVMClient = Editor->GetRigVMBlueprint()->GetRigVMClient();
+	FRigVMClient* RigVMClient = Editor->GetRigVMAsset()->GetRigVMClientHost()->GetRigVMClient();
 	if (const URigVMGraph* Graph = RigVMClient->GetModel(InGraphPath))
 	{
 		const URigVMEdGraphSchema* EdSchema = GetDefault<URigVMEdGraphSchema>();
 		FGraphDisplayInfo DisplayInfo;
-		UEdGraph* EdGraph = Editor->GetRigVMBlueprint()->GetEdGraph(Graph);
+		UEdGraph* EdGraph = Editor->GetRigVMAsset()->GetEdGraph(Graph);
 		EdSchema->GetGraphDisplayInformation(*EdGraph, DisplayInfo);
 
 		Tooltip = DisplayInfo.Tooltip;
@@ -767,11 +773,11 @@ void SRigVMEditorGraphExplorer::OnGraphClicked(const FString& InGraphPath)
 		return;
 	}
 
-	FRigVMClient* RigVMClient = Editor->GetRigVMBlueprint()->GetRigVMClient();
+	FRigVMClient* RigVMClient = Editor->GetRigVMAsset()->GetRigVMClientHost()->GetRigVMClient();
 	if (const URigVMGraph* Graph = RigVMClient->GetModel(InGraphPath))
 	{
 		const URigVMEdGraphSchema* EdSchema = GetDefault<URigVMEdGraphSchema>();
-		if (UEdGraph* EdGraph = Editor->GetRigVMBlueprint()->GetEdGraph(Graph))
+		if (UEdGraph* EdGraph = Editor->GetRigVMAsset()->GetEdGraph(Graph))
 		{
 			FGraphDisplayInfo DisplayInfo;
 			const UEdGraphSchema* Schema = EdGraph->GetSchema();
@@ -824,13 +830,13 @@ void SRigVMEditorGraphExplorer::OnFunctionClicked(const FString& InFunctionPath)
 		return;
 	}
 
-	FRigVMClient* RigVMClient = Editor->GetRigVMBlueprint()->GetRigVMClient();
+	FRigVMClient* RigVMClient = Editor->GetRigVMAsset()->GetRigVMClientHost()->GetRigVMClient();
 	if (URigVMFunctionLibrary* FunctionLibrary = RigVMClient->GetFunctionLibrary())
 	{
 		if (const URigVMLibraryNode* FunctionNode = FunctionLibrary->FindFunction(*InFunctionPath))
 		{
 			const URigVMEdGraphSchema* EdSchema = GetDefault<URigVMEdGraphSchema>();
-			if (UEdGraph* EdGraph = Editor->GetRigVMBlueprint()->GetEdGraph(FunctionNode->GetContainedGraph()))
+			if (UEdGraph* EdGraph = Editor->GetRigVMAsset()->GetEdGraph(FunctionNode->GetContainedGraph()))
 			{
 				FGraphDisplayInfo DisplayInfo;
 				const UEdGraphSchema* Schema = EdGraph->GetSchema();
@@ -859,11 +865,10 @@ void SRigVMEditorGraphExplorer::OnVariableClicked(const FRigVMExplorerElementKey
 		return;
 	}
 
-	URigVMBlueprint* Blueprint = Editor->GetRigVMBlueprint();
+	TRigVMAssetInterfacePtr Blueprint = Editor->GetRigVMAsset();
 	if (InVariable.Type == ERigVMExplorerElementType::Variable)
 	{
-		FRigVMClient* RigVMClient = Blueprint->GetRigVMClient();
-		FProperty* Prop = Blueprint->SkeletonGeneratedClass->FindPropertyByName(*InVariable.Name);
+		FProperty* Prop = Blueprint->FindGeneratedPropertyByName(*InVariable.Name);
 		UPropertyWrapper* PropWrap = (Prop ? Prop->GetUPropertyWrapper() : nullptr);
 #if WITH_RIGVMLEGACYEDITOR
 		if (TSharedPtr<SKismetInspector> KismetInspector = Editor->GetKismetInspector())
@@ -902,11 +907,11 @@ void SRigVMEditorGraphExplorer::OnGraphDoubleClicked(const FString& InGraphPath)
 		return;
 	}
 
-	FRigVMClient* RigVMClient = Editor->GetRigVMBlueprint()->GetRigVMClient();
+	FRigVMClient* RigVMClient = Editor->GetRigVMAsset()->GetRigVMClientHost()->GetRigVMClient();
 	if (const URigVMGraph* Graph = RigVMClient->GetModel(InGraphPath))
 	{
 		const URigVMEdGraphSchema* EdSchema = GetDefault<URigVMEdGraphSchema>();
-		if (UEdGraph* EdGraph = Editor->GetRigVMBlueprint()->GetEdGraph(Graph))
+		if (UEdGraph* EdGraph = Editor->GetRigVMAsset()->GetEdGraph(Graph))
 		{
 			Editor->JumpToHyperlink(EdGraph);
 		}
@@ -921,8 +926,8 @@ void SRigVMEditorGraphExplorer::OnEventDoubleClicked(const FString& InEventPath)
 		return;
 	}
 
-	URigVMBlueprint* Blueprint = Editor->GetRigVMBlueprint();
-	FRigVMClient* RigVMClient = Blueprint->GetRigVMClient();
+	TRigVMAssetInterfacePtr Blueprint = Editor->GetRigVMAsset();
+	FRigVMClient* RigVMClient = Blueprint->GetRigVMClientHost()->GetRigVMClient();
 
 	if(const URigVMNode* ModelNode = RigVMClient->FindNode(InEventPath))
 	{
@@ -938,13 +943,13 @@ void SRigVMEditorGraphExplorer::OnFunctionDoubleClicked(const FString& InFunctio
 		return;
 	}
 
-	FRigVMClient* RigVMClient = Editor->GetRigVMBlueprint()->GetRigVMClient();
+	FRigVMClient* RigVMClient = Editor->GetRigVMAsset()->GetRigVMClientHost()->GetRigVMClient();
 	if (URigVMFunctionLibrary* FunctionLibrary = RigVMClient->GetFunctionLibrary())
 	{
 		if (const URigVMLibraryNode* FunctionNode = FunctionLibrary->FindFunction(*InFunctionPath))
 		{
 			const URigVMEdGraphSchema* EdSchema = GetDefault<URigVMEdGraphSchema>();
-			if (UEdGraph* EdGraph = Editor->GetRigVMBlueprint()->GetEdGraph(FunctionNode->GetContainedGraph()))
+			if (UEdGraph* EdGraph = Editor->GetRigVMAsset()->GetEdGraph(FunctionNode->GetContainedGraph()))
 			{
 				Editor->JumpToHyperlink(EdGraph);
 			}
@@ -961,11 +966,11 @@ bool SRigVMEditorGraphExplorer::OnSetFunctionCategory(const FString& InFunctionP
 	}
 
 	const URigVMEdGraphSchema* Schema = GetDefault<URigVMEdGraphSchema>();
-	if (URigVMFunctionLibrary* Library = Editor->GetRigVMBlueprint()->GetLocalFunctionLibrary())
+	if (URigVMFunctionLibrary* Library = Editor->GetRigVMAsset()->GetLocalFunctionLibrary())
 	{
 		if (URigVMLibraryNode* Function = Library->FindFunction(*InFunctionPath))
 		{
-			if (UEdGraph* EdGraph = Editor->GetRigVMBlueprint()->GetEdGraph(Function->GetContainedGraph()))
+			if (UEdGraph* EdGraph = Editor->GetRigVMAsset()->GetEdGraph(Function->GetContainedGraph()))
 			{
 				Schema->TrySetGraphCategory(EdGraph, FText::FromString(InCategory));
 				return true;
@@ -985,7 +990,7 @@ FString SRigVMEditorGraphExplorer::OnGetFunctionCategory(const FString& InFuncti
 	}
 
 	const URigVMEdGraphSchema* Schema = GetDefault<URigVMEdGraphSchema>();
-	if (URigVMFunctionLibrary* Library = Editor->GetRigVMBlueprint()->GetLocalFunctionLibrary())
+	if (URigVMFunctionLibrary* Library = Editor->GetRigVMAsset()->GetLocalFunctionLibrary())
 	{
 		if (URigVMLibraryNode* Function = Library->FindFunction(*InFunctionPath))
 		{
@@ -1004,7 +1009,7 @@ FText SRigVMEditorGraphExplorer::OnGetFunctionTooltip(const FString& InFunctionP
 	}
 
 	const URigVMEdGraphSchema* Schema = GetDefault<URigVMEdGraphSchema>();
-	if (URigVMFunctionLibrary* Library = Editor->GetRigVMBlueprint()->GetLocalFunctionLibrary())
+	if (URigVMFunctionLibrary* Library = Editor->GetRigVMAsset()->GetLocalFunctionLibrary())
 	{
 		if (URigVMLibraryNode* Function = Library->FindFunction(*InFunctionPath))
 		{
@@ -1022,9 +1027,7 @@ FText SRigVMEditorGraphExplorer::OnGetVariableTooltip(const FString& InVariable)
 		return FText();
 	}
 
-	FString Result;
-	FBlueprintEditorUtils::GetBlueprintVariableMetaData(Editor->GetRigVMBlueprint(), *InVariable, nullptr, TEXT("tooltip"), Result);
-	return FText::FromString(Result);
+	return Editor->GetRigVMAsset()->GetVariableTooltip(*InVariable);
 }
 
 bool SRigVMEditorGraphExplorer::OnSetVariableCategory(const FString& InVariable, const FString& InCategory)
@@ -1035,7 +1038,7 @@ bool SRigVMEditorGraphExplorer::OnSetVariableCategory(const FString& InVariable,
 		return false;
 	}
 
-	FBlueprintEditorUtils::SetBlueprintVariableCategory(Editor->GetRigVMBlueprint(), *InVariable, nullptr, FText::FromString(InCategory), true);
+	Editor->GetRigVMAsset()->SetVariableCategory(*InVariable, InCategory);
 	bNeedsRefresh = true;
 	return true;
 }
@@ -1048,7 +1051,7 @@ FString SRigVMEditorGraphExplorer::OnGetVariableCategory(const FString& InVariab
 		return FString();
 	}
 
-	return FBlueprintEditorUtils::GetBlueprintVariableCategory(Editor->GetRigVMBlueprint(), *InVariable, nullptr).ToString();
+	return Editor->GetRigVMAsset()->GetVariableCategory(*InVariable);
 }
 
 FEdGraphPinType SRigVMEditorGraphExplorer::OnGetVariablePinType(const FRigVMExplorerElementKey& InVariable)
@@ -1059,12 +1062,12 @@ FEdGraphPinType SRigVMEditorGraphExplorer::OnGetVariablePinType(const FRigVMExpl
 		return FEdGraphPinType();
 	}
 
-	URigVMBlueprint* Blueprint = Editor->GetRigVMBlueprint();
+	TRigVMAssetInterfacePtr Blueprint = Editor->GetRigVMAsset();
 	if (InVariable.Type == ERigVMExplorerElementType::Variable)
 	{
-		for (FRigVMGraphVariableDescription& Variable : Blueprint->GetMemberVariables())
+		for (FRigVMGraphVariableDescription& Variable : Blueprint->GetAssetVariables())
 		{
-			if (Variable.Name == InVariable.Name)
+			if (Variable.Name == *InVariable.Name)
 			{
 				return Variable.ToPinType();
 			}
@@ -1076,7 +1079,7 @@ FEdGraphPinType SRigVMEditorGraphExplorer::OnGetVariablePinType(const FRigVMExpl
 		{
 			for (FRigVMGraphVariableDescription& Variable : Function->GetLocalVariables())
 			{
-				if (Variable.Name == InVariable.Name)
+				if (Variable.Name == *InVariable.Name)
 				{
 					return Variable.ToPinType();
 				}
@@ -1095,15 +1098,18 @@ bool SRigVMEditorGraphExplorer::OnSetVariablePinType(const FRigVMExplorerElement
 		return false;
 	}
 
-	URigVMBlueprint* Blueprint = Editor->GetRigVMBlueprint();
+	TRigVMAssetInterfacePtr Blueprint = Editor->GetRigVMAsset();
 
 	if (InVariable.Type == ERigVMExplorerElementType::Variable)
 	{
-		for (FRigVMGraphVariableDescription& Variable : Blueprint->GetMemberVariables())
+		for (FRigVMGraphVariableDescription& Variable : Blueprint->GetAssetVariables())
 		{
-			if (Variable.Name == InVariable.Name)
+			if (Variable.Name == *InVariable.Name)
 			{
-				FBlueprintEditorUtils::ChangeMemberVariableType(Blueprint, *InVariable.Name, InType);
+				FString CPPType;
+				UObject* CPPTypeObject = nullptr;
+				RigVMTypeUtils::CPPTypeFromPinType(InType, CPPType, &CPPTypeObject);
+				Blueprint->ChangeMemberVariableType(*InVariable.Name, CPPType);
 				SetLastPinTypeUsed(InType);
 				return true;
 			}
@@ -1113,7 +1119,7 @@ bool SRigVMEditorGraphExplorer::OnSetVariablePinType(const FRigVMExplorerElement
 	{
 		if (URigVMGraph* Function = Blueprint->GetFocusedModel())
 		{
-			if (URigVMController* Controller = Blueprint->GetRigVMClient()->GetController(Function))
+			if (URigVMController* Controller = Blueprint->GetRigVMClientHost()->GetRigVMClient()->GetController(Function))
 			{
 				FString NewCPPType;
 				UObject* NewCPPTypeObject = nullptr;
@@ -1136,17 +1142,12 @@ bool SRigVMEditorGraphExplorer::OnIsVariablePublic(const FString& InVariableName
 		return false;
 	}
 
-	URigVMBlueprint* Blueprint = Editor->GetRigVMBlueprint();
-	for (FRigVMGraphVariableDescription& Variable : Blueprint->GetMemberVariables())
+	TRigVMAssetInterfacePtr Blueprint = Editor->GetRigVMAsset();
+	for (FRigVMGraphVariableDescription& Variable : Blueprint->GetAssetVariables())
 	{
-		if (Variable.Name == InVariableName)
+		if (Variable.Name == *InVariableName)
 		{
-			FProperty* Property = Blueprint->SkeletonGeneratedClass->FindPropertyByName(*InVariableName);
-			if (Property)
-			{
-				return Property->HasAnyPropertyFlags(CPF_DisableEditOnInstance) ? true : false;
-			}
-			return true;
+			return Variable.bPublic;
 		}
 	}
 
@@ -1161,10 +1162,12 @@ bool SRigVMEditorGraphExplorer::OnToggleVariablePublic(const FString& InVariable
 		return false;
 	}
 
-	// Toggle the flag on the blueprint's version of the variable description, based on state
-	const bool bVariableIsExposed = OnIsVariablePublic(InVariableName);
-	FBlueprintEditorUtils::SetBlueprintOnlyEditableFlag(Editor->GetRigVMBlueprint(), *InVariableName, !bVariableIsExposed);
-	
+	if (TRigVMAssetInterfacePtr Blueprint = Editor->GetRigVMAsset())
+	{
+		// Toggle the flag on the blueprint's version of the variable description, based on state
+	   const bool bVariableIsExposed = OnIsVariablePublic(InVariableName);
+	   Blueprint->SetVariablePublic(*InVariableName, !bVariableIsExposed);
+	}
 	return true;
 }
 
@@ -1176,7 +1179,7 @@ bool SRigVMEditorGraphExplorer::IsFunctionFocused() const
 		return false;
 	}
 
-	if (URigVMGraph* Graph = Editor->GetRigVMBlueprint()->GetFocusedModel())
+	if (URigVMGraph* Graph = Editor->GetRigVMAsset()->GetFocusedModel())
 	{
 		if (URigVMCollapseNode* FunctionNode = Cast<URigVMCollapseNode>(Graph->GetOuter()))
 		{
@@ -1281,7 +1284,7 @@ FReply SRigVMEditorGraphExplorer::OnDragDetected(const FGeometry& MyGeometry, co
 	{
 		if (RigVMEditor.IsValid() && DraggedElements.Num() == 1)
 		{
-			TSharedRef<FRigVMGraphExplorerDragDropOp> DragDropOp = FRigVMGraphExplorerDragDropOp::New(MoveTemp(DraggedElements[0]), RigVMEditor.Pin()->GetRigVMBlueprint());
+			TSharedRef<FRigVMGraphExplorerDragDropOp> DragDropOp = FRigVMGraphExplorerDragDropOp::New(MoveTemp(DraggedElements[0]), RigVMEditor.Pin()->GetRigVMAsset());
 			return FReply::Handled().BeginDragDrop(DragDropOp);
 		}
 	}
@@ -1309,7 +1312,7 @@ void SRigVMEditorGraphExplorer::OnCopy()
 
 	FString OutputString;
 	const URigVMEdGraphSchema* Schema = GetDefault<URigVMEdGraphSchema>();
-	URigVMBlueprint* Blueprint = Editor->GetRigVMBlueprint();
+	TRigVMAssetInterfacePtr Blueprint = Editor->GetRigVMAsset();
 
 	switch (Selection[0].Type)
 	{
@@ -1319,7 +1322,7 @@ void SRigVMEditorGraphExplorer::OnCopy()
 			{
 				if (const URigVMLibraryNode* FunctionNode = FunctionLibrary->FindFunction(*Selection[0].Name))
 				{
-					if (UEdGraph* EdGraph = Editor->GetRigVMBlueprint()->GetEdGraph(FunctionNode->GetContainedGraph()))
+					if (UEdGraph* EdGraph = Editor->GetRigVMAsset()->GetEdGraph(FunctionNode->GetContainedGraph()))
 					{
 						Blueprint->ExportGraphToText(EdGraph, OutputString);
 					}
@@ -1329,30 +1332,7 @@ void SRigVMEditorGraphExplorer::OnCopy()
 		}
 		case ERigVMExplorerElementType::Variable:
 		{
-			int32 VarIndex = FBlueprintEditorUtils::FindNewVariableIndex(Blueprint, *Selection[0].Name);
-			if (VarIndex != INDEX_NONE)
-			{
-				// make a copy of the Variable description so we can set the default value
-				FBPVariableDescription Description = Blueprint->NewVariables[VarIndex];
-
-				//Grab property of blueprint's current CDO
-				UClass* GeneratedClass = Blueprint->GeneratedClass;
-				UObject* GeneratedCDO = GeneratedClass->GetDefaultObject();
-				FProperty* TargetProperty = FindFProperty<FProperty>(GeneratedClass, Description.VarName);
-
-				if (TargetProperty)
-				{
-					// Grab the address of where the property is actually stored (UObject* base, plus the offset defined in the property)
-					void* OldPropertyAddr = TargetProperty->ContainerPtrToValuePtr<void>(GeneratedCDO);
-					if (OldPropertyAddr)
-					{
-						TargetProperty->ExportTextItem_Direct(Description.DefaultValue, OldPropertyAddr, OldPropertyAddr, nullptr, PPF_SerializedAsImportText);
-					}
-				}
-
-				FBPVariableDescription::StaticStruct()->ExportText(OutputString, &Description, &Description, nullptr, 0, nullptr, false);
-				OutputString = TEXT("BPVar") + OutputString;
-			}
+			OutputString = Blueprint->OnCopyVariable(*Selection[0].Name);
 			break;
 		}
 		case ERigVMExplorerElementType::LocalVariable:
@@ -1495,7 +1475,7 @@ void SRigVMEditorGraphExplorer::OnPasteFunction()
 	FString ClipboardText;
 	FPlatformApplicationMisc::ClipboardPaste(ClipboardText);
 
-	Editor->GetRigVMBlueprint()->TryImportGraphFromText(ClipboardText);
+	Editor->GetRigVMAsset()->TryImportGraphFromText(ClipboardText);
 }
 
 bool SRigVMEditorGraphExplorer::CanPasteFunction() const
@@ -1509,7 +1489,7 @@ bool SRigVMEditorGraphExplorer::CanPasteFunction() const
 	FString ClipboardText;
 	FPlatformApplicationMisc::ClipboardPaste(ClipboardText);
 
-	return Editor->GetRigVMBlueprint()->CanImportGraphFromText(ClipboardText);
+	return Editor->GetRigVMAsset()->CanImportGraphFromText(ClipboardText);
 }
 
 void SRigVMEditorGraphExplorer::OnPasteVariable()
@@ -1519,33 +1499,11 @@ void SRigVMEditorGraphExplorer::OnPasteVariable()
 	{
 		return;
 	}
-	URigVMBlueprint* Blueprint = Editor->GetRigVMBlueprint();
+	TRigVMAssetInterfacePtr Blueprint = Editor->GetRigVMAsset();
 	
 	FString ClipboardText;
 	FPlatformApplicationMisc::ClipboardPaste(ClipboardText);
-	if (!ensure(ClipboardText.StartsWith(TEXT("BPVar"), ESearchCase::CaseSensitive)))
-	{
-		return;
-	}
-
-	FBPVariableDescription Description;
-	FStringOutputDevice Errors;
-	const TCHAR* Import = ClipboardText.GetCharArray().GetData() + FCString::Strlen(TEXT("BPVar"));
-	FBPVariableDescription::StaticStruct()->ImportText(Import, &Description, nullptr, PPF_None, &Errors, FBPVariableDescription::StaticStruct()->GetName());
-	if (Errors.IsEmpty())
-	{
-		FBPVariableDescription NewVar = FBlueprintEditorUtils::DuplicateVariableDescription(Blueprint, Description);
-		if (NewVar.VarGuid.IsValid())
-		{
-			FScopedTransaction Transaction(FText::Format(LOCTEXT("PasteVariable", "Paste Variable: {0}"), FText::FromName(NewVar.VarName)));
-			Blueprint->Modify();
-			Blueprint->NewVariables.Add(NewVar);
-
-			// Potentially adjust variable names for any child blueprints
-			FBlueprintEditorUtils::ValidateBlueprintChildVariables(Blueprint, NewVar.VarName);
-			FBlueprintEditorUtils::MarkBlueprintAsStructurallyModified(Blueprint);
-		}
-	}
+	Blueprint->OnPasteVariable(ClipboardText);
 }
 
 bool SRigVMEditorGraphExplorer::CanPasteVariable() const
@@ -1578,7 +1536,7 @@ void SRigVMEditorGraphExplorer::OnPasteLocalVariable()
 	{
 		return;
 	}
-	URigVMBlueprint* Blueprint = Editor->GetRigVMBlueprint();
+	TRigVMAssetInterfacePtr Blueprint = Editor->GetRigVMAsset();
 	
 	FString ClipboardText;
 	FPlatformApplicationMisc::ClipboardPaste(ClipboardText);
@@ -1629,7 +1587,7 @@ void SRigVMEditorGraphExplorer::OnOpenGraph(bool bOnNewTab)
 		if (Selection.Num() == 1 && Selection[0].Type == ERigVMExplorerElementType::Graph)
 		{
 			const URigVMEdGraphSchema* Schema = GetDefault<URigVMEdGraphSchema>();
-			GraphToOpen = Editor->GetRigVMBlueprint()->GetEdGraph(Selection[0].Name);
+			GraphToOpen = Editor->GetRigVMAsset()->GetEdGraph(Selection[0].Name);
 		}
 
 		if (GraphToOpen)
@@ -1738,11 +1696,11 @@ void SRigVMEditorGraphExplorer::OnAddFunctionVariant()
 		return;
 	}
 
-	if (URigVMFunctionLibrary* FunctionLibrary = Editor->GetRigVMBlueprint()->GetLocalFunctionLibrary())
+	if (URigVMFunctionLibrary* FunctionLibrary = Editor->GetRigVMAsset()->GetLocalFunctionLibrary())
 	{
 		if (const URigVMLibraryNode* FunctionNode = FunctionLibrary->FindFunction(*Selected[0].Name))
 		{
-			if (UEdGraph* EdGraph = Editor->GetRigVMBlueprint()->GetEdGraph(FunctionNode->GetContainedGraph()))
+			if (UEdGraph* EdGraph = Editor->GetRigVMAsset()->GetEdGraph(FunctionNode->GetContainedGraph()))
 			{
 				Editor->AddNewFunctionVariant(EdGraph);
 			}
@@ -1774,7 +1732,7 @@ void SRigVMEditorGraphExplorer::OnDeleteEntry()
 
 	if (Editor)
 	{
-		URigVMBlueprint* Blueprint = Editor->GetRigVMBlueprint();
+		TRigVMAssetInterfacePtr Blueprint = Editor->GetRigVMAsset();
 		TArray<FRigVMExplorerElementKey> Selected = TreeView->GetSelectedKeys();
 		const URigVMEdGraphSchema* Schema = GetDefault<URigVMEdGraphSchema>();
 		for (const FRigVMExplorerElementKey& Key : Selected)
@@ -1807,15 +1765,15 @@ void SRigVMEditorGraphExplorer::OnDeleteEntry()
 				{
 					const FScopedTransaction Transaction( LOCTEXT( "RemoveVariable", "Remove Variable" ) );
 
-					Blueprint->Modify();
-					FBlueprintEditorUtils::RemoveMemberVariable(Blueprint, *Key.Name);
+					Blueprint->GetObject()->Modify();
+					Blueprint->RemoveMemberVariable(*Key.Name);
 					break;
 				}
 				case ERigVMExplorerElementType::LocalVariable:
 				{
 					if (URigVMGraph* FocusedGraph = Editor->GetFocusedModel())
 					{
-						if (URigVMController* Controller = Blueprint->GetRigVMClient()->GetController(FocusedGraph))
+						if (URigVMController* Controller = Blueprint->GetRigVMClientHost()->GetRigVMClient()->GetController(FocusedGraph))
 						{
 							Controller->RemoveLocalVariable(*Key.Name, true, true);
 						}
@@ -1894,7 +1852,7 @@ bool SRigVMEditorGraphExplorer::OnRenameGraph(const FString& InOldPath, const FS
 	if (Editor)
 	{
 		const URigVMEdGraphSchema* Schema = GetDefault<URigVMEdGraphSchema>();
-		if (UEdGraph* EdGraph = Editor->GetRigVMBlueprint()->GetEdGraph(InOldPath))
+		if (UEdGraph* EdGraph = Editor->GetRigVMAsset()->GetEdGraph(InOldPath))
 		{
 			bResult = Schema->TryRenameGraph(EdGraph, *InNewPath);
 		}
@@ -1921,7 +1879,7 @@ bool SRigVMEditorGraphExplorer::OnCanRenameGraph(const FString& InOldPath, const
 		{
 			NewPath = FString::Printf(TEXT("%s%s"), *Prefix, *InNewPath);
 		}
-		if (Editor->GetRigVMBlueprint()->GetRigVMClient()->GetModel(NewPath))
+		if (Editor->GetRigVMAsset()->GetRigVMClientHost()->GetRigVMClient()->GetModel(NewPath))
 		{
 			OutErrorMessage = FText::FromString(TEXT("Name already in use."));
 			bResult = false;
@@ -1943,11 +1901,11 @@ bool SRigVMEditorGraphExplorer::OnRenameFunction(const FString& InOldPath, const
 	if (Editor)
 	{
 		const URigVMEdGraphSchema* Schema = GetDefault<URigVMEdGraphSchema>();
-		if (URigVMFunctionLibrary* FunctionLibrary = Editor->GetRigVMBlueprint()->GetLocalFunctionLibrary())
+		if (URigVMFunctionLibrary* FunctionLibrary = Editor->GetRigVMAsset()->GetLocalFunctionLibrary())
 		{
 			if (const URigVMLibraryNode* FunctionNode = FunctionLibrary->FindFunction(*InOldPath))
 			{
-				if (UEdGraph* EdGraph = Editor->GetRigVMBlueprint()->GetEdGraph(FunctionNode->GetContainedGraph()))
+				if (UEdGraph* EdGraph = Editor->GetRigVMAsset()->GetEdGraph(FunctionNode->GetContainedGraph()))
 				{
 					bResult = Schema->TryRenameGraph(EdGraph, *InNewPath);
 				}
@@ -1966,17 +1924,41 @@ bool SRigVMEditorGraphExplorer::OnCanRenameFunction(const FString& InOldPath, co
 
 	if (Editor)
 	{
-		if (URigVMFunctionLibrary* Library = Editor->GetRigVMBlueprint()->GetRigVMClient()->GetFunctionLibrary())
+		const URigVMEdGraphSchema* Schema = GetDefault<URigVMEdGraphSchema>();
+		TSharedPtr<INameValidatorInterface> NameValidator = Schema->GetNameValidator(Editor->GetRigVMAsset(), *InOldPath, nullptr);
+		//TSharedPtr<INameValidatorInterface> NameValidator = Schema->GetNameValidator(Cast<UBlueprint>(Editor->GetRigVMAsset().GetObject()), *InOldPath, nullptr, NAME_None);
+		const EValidatorResult Result = NameValidator->IsValid(InNewPath);
+		switch (Result)
 		{
-			if (Library->FindFunction(*InNewPath))
-			{
-				OutErrorMessage = FText::FromString(TEXT("Name already in use."));
-				bResult = false;
-			}
-			else
-			{
-				bResult = true;
-			}
+			case EValidatorResult::AlreadyInUse:
+			case EValidatorResult::LocallyInUse:
+				{
+					OutErrorMessage = FText::FromString(TEXT("Name already in use."));
+					bResult = false;
+					break;
+				}
+			case EValidatorResult::ContainsInvalidCharacters:
+				{
+					OutErrorMessage = FText::FromString(TEXT("Name contains invalid characters."));
+					bResult = false;
+					break;
+				}
+			case EValidatorResult::TooLong:
+				{
+					OutErrorMessage = FText::FromString(TEXT("Name is too long."));
+					bResult = false;
+					break;
+				}
+			case EValidatorResult::EmptyName:
+				{
+					OutErrorMessage = FText::FromString(TEXT("Name is empty."));
+					bResult = false;
+					break;
+				}
+			default:
+				{
+					bResult = true;
+				}
 		}
 	}
 
@@ -1999,18 +1981,18 @@ bool SRigVMEditorGraphExplorer::OnRenameVariable(const FRigVMExplorerElementKey&
 
 	const FScopedTransaction Transaction( LOCTEXT( "RenameVariable", "Rename Variable" ) );
 
-	URigVMBlueprint* Blueprint = Editor->GetRigVMBlueprint();
+	TRigVMAssetInterfacePtr Blueprint = Editor->GetRigVMAsset();
 
 	if (InOldKey.Type == ERigVMExplorerElementType::Variable)
 	{
-		Blueprint->Modify();
-		FBlueprintEditorUtils::RenameMemberVariable(Blueprint, *InOldKey.Name, *InNewName);
+		Blueprint->GetObject()->Modify();
+		Blueprint->RenameMemberVariable(*InOldKey.Name, *InNewName);
 	}
 	else if (InOldKey.Type == ERigVMExplorerElementType::LocalVariable)
 	{
 		if (URigVMGraph* Graph = Blueprint->GetFocusedModel())
 		{
-			if (URigVMController* Controller = Blueprint->GetRigVMClient()->GetController(Graph))
+			if (URigVMController* Controller = Blueprint->GetRigVMClientHost()->GetRigVMClient()->GetController(Graph))
 			{
 				Controller->RenameLocalVariable(*InOldKey.Name, *InNewName, true, true);
 			}
@@ -2029,12 +2011,12 @@ bool SRigVMEditorGraphExplorer::OnCanRenameVariable(const FRigVMExplorerElementK
 	TArray<FRigVMGraphVariableDescription> Variables = (InOldKey.Type == ERigVMExplorerElementType::Variable) ? GetVariables() : GetLocalVariables();
 	for (FRigVMGraphVariableDescription& Variable : Variables)
 	{
-		if (Variable.Name == InOldKey.Name)
+		if (Variable.Name == *InOldKey.Name)
 		{
 			continue;
 		}
 
-		if (Variable.Name == InNewName)
+		if (Variable.Name == *InNewName)
 		{
 			OutErrorMessage = FText::FromString(TEXT("Name already in use."));
 			return false;
@@ -2053,7 +2035,7 @@ void SRigVMEditorGraphExplorer::OnFindReference(bool bSearchAllBlueprints)
 	}
 
 	TSharedPtr<IRigVMEditor> Editor = RigVMEditor.Pin();
-	URigVMBlueprint* Blueprint = Editor->GetRigVMBlueprint();
+	TRigVMAssetInterfacePtr Blueprint = Editor->GetRigVMAsset();
 	const URigVMEdGraphSchema* Schema = GetDefault<URigVMEdGraphSchema>();
 
 	TArray<FRigVMExplorerElementKey> Selection = TreeView->GetSelectedKeys();
@@ -2066,7 +2048,7 @@ void SRigVMEditorGraphExplorer::OnFindReference(bool bSearchAllBlueprints)
 	{
 		case ERigVMExplorerElementType::Graph:
 		{
-			if (UEdGraph* EdGraph = Editor->GetRigVMBlueprint()->GetEdGraph(Selection[0].Name))
+			if (UEdGraph* EdGraph = Editor->GetRigVMAsset()->GetEdGraph(Selection[0].Name))
 			{
 				FGraphDisplayInfo DisplayInfo;
 				Schema->GetGraphDisplayInformation(*EdGraph, DisplayInfo);
