RigVM: Adding IRigVMAssetInterface as parent of RigVMBlueprint. First step to separation of RigVM and Blueprint.

diff --git a/Engine/Plugins/Runtime/RigVM/Source/RigVMEditor/Private/Widgets/SRigVMEditorSelectedDebugObjectWidget.cpp b/Engine/Plugins/Runtime/RigVM/Source/RigVMEditor/Private/Widgets/SRigVMEditorSelectedDebugObjectWidget.cpp
--- a/Engine/Plugins/Runtime/RigVM/Source/RigVMEditor/Private/Widgets/SRigVMEditorSelectedDebugObjectWidget.cpp
+++ b/Engine/Plugins/Runtime/RigVM/Source/RigVMEditor/Private/Widgets/SRigVMEditorSelectedDebugObjectWidget.cpp
@@ -6,7 +6,7 @@
 #include "PreviewScene.h"
 #include "IDocumentation.h"
 #include "PropertyCustomizationHelpers.h"
-#include "RigVMBlueprint.h"
+#include "RigVMBlueprintLegacy.h"
 #include "SLevelOfDetailBranchNode.h"
 #include "UnrealEdGlobals.h"
 #include "Components/Widget.h"
@@ -106,7 +106,7 @@ void SRigVMEditorSelectedDebugObjectWidget::Construct(const FArguments& InArgs,
 
 void SRigVMEditorSelectedDebugObjectWidget::Tick( const FGeometry& AllottedGeometry, const double InCurrentTime, const float InDeltaTime )
 {
-	if (UBlueprint* Blueprint = GetBlueprintObj())
+	if (TRigVMAssetInterfacePtr Blueprint = GetRigVMAsset())
 	{
 		if (UObject* Object = Blueprint->GetObjectBeingDebugged())
 		{
@@ -147,7 +147,7 @@ const FString& SRigVMEditorSelectedDebugObjectWidget::GetDebugAllWorldsString()
 
 void SRigVMEditorSelectedDebugObjectWidget::OnRefresh()
 {
-	if (GetBlueprintObj())
+	if (GetRigVMAsset())
 	{
 		GenerateDebugWorldNames(false);
 		GenerateDebugObjectInstances(false);
@@ -311,14 +311,12 @@ void SRigVMEditorSelectedDebugObjectWidget::GenerateDebugObjectInstances(bool bR
 
 	if (!EditorPtr->OnlyShowCustomDebugObjects())
 	{
-		UClass* BlueprintClass = GetBlueprintObj()->GeneratedClass;
-
-		if (BlueprintClass)
+		//if (BlueprintClass)
 		{
 			// Experimental new path for debug object discovery
-			TArray<UObject*> BlueprintInstances;
-			GetObjectsOfClass(BlueprintClass, BlueprintInstances, true);
-
+			bool bIncludeDerivedClasses = true;
+			TArray<UObject*> BlueprintInstances = GetRigVMAsset()->GetArchetypeInstances(false, bIncludeDerivedClasses);
+			
 			for (auto It = BlueprintInstances.CreateIterator(); It; ++It)
 			{
 				UObject* TestObject = *It;
@@ -376,60 +374,61 @@ void SRigVMEditorSelectedDebugObjectWidget::GenerateDebugObjectInstances(bool bR
 				}
 			}
 		}
-		else
-		{
-			for (TObjectIterator<UObject> It; It; ++It)
-			{
-				UObject* TestObject = *It;
-
-				// Skip Blueprint preview objects (don't allow them to be selected for debugging)
-				if (PreviewWorld != nullptr && TestObject->IsIn(PreviewWorld))
-				{
-					continue;
-				}
-
-				const bool bPassesFlags = !TestObject->HasAnyFlags(RF_ClassDefaultObject) && IsValid(TestObject);
-				const bool bGeneratedByAnyBlueprint = TestObject->GetClass()->ClassGeneratedBy != nullptr;
-				const bool bGeneratedByThisBlueprint = bGeneratedByAnyBlueprint && GetBlueprintObj()->GeneratedClass && TestObject->IsA(GetBlueprintObj()->GeneratedClass);
-
-				if (bPassesFlags && bGeneratedByThisBlueprint)
-				{
-					UObject *ObjOuter = TestObject;
-					UWorld *ObjWorld = nullptr;
-					do		// Run through at least once in case the TestObject is a UGameInstance
-					{
-						UGameInstance *ObjGameInstance = Cast<UGameInstance>(ObjOuter);
-
-						ObjOuter = ObjOuter->GetOuter();
-						ObjWorld = ObjGameInstance ? ObjGameInstance->GetWorld() : Cast<UWorld>(ObjOuter);
-					} while (ObjWorld == nullptr && ObjOuter != nullptr);
-
-					if (ObjWorld)
-					{
-						// Make check on owning level (not streaming level)
-						if (ObjWorld->PersistentLevel && ObjWorld->PersistentLevel->OwningWorld)
-						{
-							ObjWorld = ObjWorld->PersistentLevel->OwningWorld;
-						}
-
-						// We have a specific debug world and the object isn't in it
-						if (DebugWorld && ObjWorld != DebugWorld)
-						{
-							continue;
-						}
-
-						if ((ObjWorld->WorldType == EWorldType::Editor) && (GUnrealEd->GetPIEViewport() == nullptr))
-						{
-							AddDebugObject(TestObject);
-						}
-						else if (ObjWorld->WorldType == EWorldType::PIE)
-						{
-							AddDebugObject(TestObject);
-						}
-					}
-				}
-			}
-		}
+		// this else doesn't make sense to me, since BlueprintClass doesn't exist, but we check it anyway
+		// else
+		// {
+		// 	for (TObjectIterator<UObject> It; It; ++It)
+		// 	{
+		// 		UObject* TestObject = *It;
+		//
+		// 		// Skip Blueprint preview objects (don't allow them to be selected for debugging)
+		// 		if (PreviewWorld != nullptr && TestObject->IsIn(PreviewWorld))
+		// 		{
+		// 			continue;
+		// 		}
+		//
+		// 		const bool bPassesFlags = !TestObject->HasAnyFlags(RF_ClassDefaultObject) && IsValid(TestObject);
+		// 		const bool bGeneratedByAnyBlueprint = TestObject->GetClass()->ClassGeneratedBy != nullptr;
+		// 		const bool bGeneratedByThisBlueprint = bGeneratedByAnyBlueprint && GetBlueprintObj()->GetRigVMGeneratedClass() && TestObject->IsA(GetBlueprintObj()->GetRigVMGeneratedClass());
+		//
+		// 		if (bPassesFlags && bGeneratedByThisBlueprint)
+		// 		{
+		// 			UObject *ObjOuter = TestObject;
+		// 			UWorld *ObjWorld = nullptr;
+		// 			do		// Run through at least once in case the TestObject is a UGameInstance
+		// 			{
+		// 				UGameInstance *ObjGameInstance = Cast<UGameInstance>(ObjOuter);
+		//
+		// 				ObjOuter = ObjOuter->GetOuter();
+		// 				ObjWorld = ObjGameInstance ? ObjGameInstance->GetWorld() : Cast<UWorld>(ObjOuter);
+		// 			} while (ObjWorld == nullptr && ObjOuter != nullptr);
+		//
+		// 			if (ObjWorld)
+		// 			{
+		// 				// Make check on owning level (not streaming level)
+		// 				if (ObjWorld->PersistentLevel && ObjWorld->PersistentLevel->OwningWorld)
+		// 				{
+		// 					ObjWorld = ObjWorld->PersistentLevel->OwningWorld;
+		// 				}
+		//
+		// 				// We have a specific debug world and the object isn't in it
+		// 				if (DebugWorld && ObjWorld != DebugWorld)
+		// 				{
+		// 					continue;
+		// 				}
+		//
+		// 				if ((ObjWorld->WorldType == EWorldType::Editor) && (GUnrealEd->GetPIEViewport() == nullptr))
+		// 				{
+		// 					AddDebugObject(TestObject);
+		// 				}
+		// 				else if (ObjWorld->WorldType == EWorldType::PIE)
+		// 				{
+		// 					AddDebugObject(TestObject);
+		// 				}
+		// 			}
+		// 		}
+		// 	}
+		// }
 	}
 
 	if (DebugObjectsComboBox.IsValid())
@@ -470,8 +469,8 @@ void SRigVMEditorSelectedDebugObjectWidget::GenerateDebugObjectInstances(bool bR
 
 TSharedPtr<FRigVMDebugObjectInstance> SRigVMEditorSelectedDebugObjectWidget::GetDebugObjectInstance() const
 {
-	check(GetBlueprintObj());
-	const FString& PathToDebug = GetBlueprintObj()->GetObjectPathToDebug();
+	check(GetRigVMAsset());
+	const FString& PathToDebug = GetRigVMAsset()->GetObjectBeingDebuggedPath();
 	if (!PathToDebug.IsEmpty())
 	{
 		for (int32 ObjectIndex = 0; ObjectIndex < DebugObjects.Num(); ++ObjectIndex)
@@ -493,10 +492,10 @@ TSharedPtr<FRigVMDebugObjectInstance> SRigVMEditorSelectedDebugObjectWidget::Get
 
 TSharedPtr<FString> SRigVMEditorSelectedDebugObjectWidget::GetDebugWorldName() const
 {
-	check(GetBlueprintObj());
+	check(GetRigVMAsset());
 	if (ensure(DebugWorlds.Num() == DebugWorldNames.Num()))
 	{
-		UWorld* DebugWorld = GetBlueprintObj()->GetWorldBeingDebugged();
+		UWorld* DebugWorld = GetRigVMAsset()->GetWorldBeingDebugged();
 		if (DebugWorld != nullptr)
 		{
 			for (int32 WorldIndex = 0; WorldIndex < DebugWorlds.Num(); ++WorldIndex)
@@ -527,9 +526,9 @@ void SRigVMEditorSelectedDebugObjectWidget::DebugWorldSelectionChanged(TSharedPt
 		{
 			if (DebugWorldNames[WorldIdx] == NewSelection)
 			{
-				GetBlueprintObj()->SetWorldBeingDebugged(DebugWorlds[WorldIdx].Get());
+				GetRigVMAsset()->SetWorldBeingDebugged(DebugWorlds[WorldIdx].Get());
 
-				GetBlueprintObj()->SetObjectBeingDebugged(nullptr);
+				GetRigVMAsset()->SetObjectBeingDebugged(nullptr);
 				LastObjectObserved.Reset();
 
 				GenerateDebugObjectInstances(false);
@@ -544,7 +543,7 @@ void SRigVMEditorSelectedDebugObjectWidget::DebugObjectSelectionChanged(TSharedP
 	if (NewSelection != GetDebugObjectInstance() && NewSelection.IsValid())
 	{
 		UObject* DebugObj = NewSelection->ObjectPtr.Get();
-		GetBlueprintObj()->SetObjectBeingDebugged(DebugObj);
+		GetRigVMAsset()->SetObjectBeingDebugged(DebugObj);
 
 		if (TSharedPtr<IRigVMEditor> SharedBlueprintEditor = Editor.Pin())
 		{
@@ -557,8 +556,8 @@ void SRigVMEditorSelectedDebugObjectWidget::DebugObjectSelectionChanged(TSharedP
 
 bool SRigVMEditorSelectedDebugObjectWidget::IsDebugObjectSelected() const
 {
-	check(GetBlueprintObj());
-	if (UObject* DebugObj = GetBlueprintObj()->GetObjectBeingDebugged())
+	check(GetRigVMAsset());
+	if (UObject* DebugObj = GetRigVMAsset()->GetObjectBeingDebugged())
 	{
 		if (AActor* Actor = Cast<AActor>(DebugObj))
 		{
@@ -570,7 +569,7 @@ bool SRigVMEditorSelectedDebugObjectWidget::IsDebugObjectSelected() const
 
 void SRigVMEditorSelectedDebugObjectWidget::SelectedDebugObject_OnClicked()
 {
-	if (UObject* DebugObj = GetBlueprintObj()->GetObjectBeingDebugged())
+	if (UObject* DebugObj = GetRigVMAsset()->GetObjectBeingDebugged())
 	{
 		if (AActor* Actor = Cast<AActor>(DebugObj))
 		{
@@ -746,9 +745,9 @@ void SRigVMEditorSelectedDebugObjectWidget::AddDebugObject(UObject* TestObject,
 	DebugObjects.Add(NewInstance);
 }
 
-UBlueprint* SRigVMEditorSelectedDebugObjectWidget::GetBlueprintObj() const
+TRigVMAssetInterfacePtr SRigVMEditorSelectedDebugObjectWidget::GetRigVMAsset() const
 {
-	return Editor.Pin()->GetRigVMBlueprint();
+	return Editor.Pin()->GetRigVMAsset();
 }
 
 TSharedRef<SWidget> SRigVMEditorSelectedDebugObjectWidget::CreateDebugObjectItemWidget(TSharedPtr<FRigVMDebugObjectInstance> InItem)
@@ -776,7 +775,7 @@ FText SRigVMEditorSelectedDebugObjectWidget::GetSelectedDebugObjectTextLabel() c
 	{
 		Label = DebugInstance->ObjectLabel;
 
-		UBlueprint* Blueprint = GetBlueprintObj();
+		TRigVMAssetInterfacePtr Blueprint = GetRigVMAsset();
 		if (Blueprint != nullptr)
 		{
 			UObject* DebugObj = Blueprint->GetObjectBeingDebugged();
diff --git a/Engine/Plugins/Runtime/RigVM/Source/RigVMEditor/Private/Widgets/SRigVMExecutionStackView.cpp b/Engine/Plugins/Runtime/RigVM/Source/RigVMEditor/Private/Widgets/SRigVMExecutionStackView.cpp
--- a/Engine/Plugins/Runtime/RigVM/Source/RigVMEditor/Private/Widgets/SRigVMExecutionStackView.cpp
+++ b/Engine/Plugins/Runtime/RigVM/Source/RigVMEditor/Private/Widgets/SRigVMExecutionStackView.cpp
@@ -37,7 +37,7 @@ FRigStackEntry::FRigStackEntry(int32 InEntryIndex, ERigStackEntry::Type InEntryT
 
 }
 
-TSharedRef<ITableRow> FRigStackEntry::MakeTreeRowWidget(const TSharedRef<STableViewBase>& InOwnerTable, TSharedRef<FRigStackEntry> InEntry, TSharedRef<FUICommandList> InCommandList, TSharedPtr<SRigVMExecutionStackView> InStackView, TWeakObjectPtr<URigVMBlueprint> InBlueprint)
+TSharedRef<ITableRow> FRigStackEntry::MakeTreeRowWidget(const TSharedRef<STableViewBase>& InOwnerTable, TSharedRef<FRigStackEntry> InEntry, TSharedRef<FUICommandList> InCommandList, TSharedPtr<SRigVMExecutionStackView> InStackView, TWeakObjectPtr<UObject> InBlueprint)
 {
 	return SNew(SRigStackItem, InOwnerTable, InEntry, InCommandList, InBlueprint);
 }
@@ -45,7 +45,7 @@ TSharedRef<ITableRow> FRigStackEntry::MakeTreeRowWidget(const TSharedRef<STableV
 //////////////////////////////////////////////////////////////
 /// SRigStackItem
 ///////////////////////////////////////////////////////////
-void SRigStackItem::Construct(const FArguments& InArgs, const TSharedRef<STableViewBase>& OwnerTable, TSharedRef<FRigStackEntry> InStackEntry, TSharedRef<FUICommandList> InCommandList, TWeakObjectPtr<URigVMBlueprint> InBlueprint)
+void SRigStackItem::Construct(const FArguments& InArgs, const TSharedRef<STableViewBase>& OwnerTable, TSharedRef<FRigStackEntry> InStackEntry, TSharedRef<FUICommandList> InCommandList, TWeakObjectPtr<UObject> InBlueprint)
 {
 	WeakStackEntry = InStackEntry;
 	WeakBlueprint = InBlueprint;
@@ -192,46 +192,55 @@ FText SRigStackItem::GetTooltip() const
 
 FText SRigStackItem::GetVisitedCountText() const
 {
-	if(WeakStackEntry.IsValid() && WeakBlueprint.IsValid())
+	if (TSharedPtr<FRigStackEntry> StackEntry = WeakStackEntry.Pin())
 	{
-		if(WeakBlueprint->RigGraphDisplaySettings.bShowNodeRunCounts)
+		if (TStrongObjectPtr<UObject> StrongPtr = WeakBlueprint.Pin())
 		{
-			if(WeakStackEntry.Pin()->EntryType == ERigStackEntry::Operator)
+			IRigVMAssetInterface* Blueprint = Cast<IRigVMAssetInterface>(StrongPtr.Get());
+			if(Blueprint->GetRigGraphDisplaySettings().bShowNodeRunCounts)
 			{
-				if(URigVMHost* RigVMHost = Cast<URigVMHost>(WeakBlueprint->GetObjectBeingDebugged()))
+				if(WeakStackEntry.Pin()->EntryType == ERigStackEntry::Operator)
 				{
-					if(URigVM* VM = RigVMHost->GetVM())
+					if(URigVMHost* RigVMHost = Blueprint->GetDebuggedRigVMHost())
 					{
-						const int32 Count = VM->GetInstructionVisitedCount(RigVMHost->GetRigVMExtendedExecuteContext(), WeakStackEntry.Pin()->InstructionIndex);
-						if(Count > 0)
+						if(URigVM* VM = RigVMHost->GetVM())
 						{
-							return FText::FromString(FString::FromInt(Count));
+							const int32 Count = VM->GetInstructionVisitedCount(RigVMHost->GetRigVMExtendedExecuteContext(), WeakStackEntry.Pin()->InstructionIndex);
+							if(Count > 0)
+							{
+								return FText::FromString(FString::FromInt(Count));
+							}
 						}
 					}
 				}
 			}
 		}
 	}
+
 	return FText();
 }
 
 FText SRigStackItem::GetDurationText() const
 {
-	if(WeakStackEntry.IsValid() && WeakBlueprint.IsValid())
+	if (TSharedPtr<FRigStackEntry> StackEntry = WeakStackEntry.Pin())
 	{
-		if(WeakBlueprint->VMRuntimeSettings.bEnableProfiling)
+		if (TStrongObjectPtr<UObject> StrongPtr = WeakBlueprint.Pin())
 		{
-			if(WeakStackEntry.Pin()->EntryType == ERigStackEntry::Operator)
+			IRigVMAssetInterface* Blueprint = Cast<IRigVMAssetInterface>(StrongPtr.Get());
+			if(Blueprint->GetVMRuntimeSettings().bEnableProfiling)
 			{
-				if(URigVMHost* RigVMHost = Cast<URigVMHost>(WeakBlueprint->GetObjectBeingDebugged()))
+				if(WeakStackEntry.Pin()->EntryType == ERigStackEntry::Operator)
 				{
-					if(URigVM* VM = RigVMHost->GetVM())
+					if(URigVMHost* RigVMHost = Blueprint->GetDebuggedRigVMHost())
 					{
-						const double CurrentMicroSeconds = VM->GetInstructionMicroSeconds(RigVMHost->GetRigVMExtendedExecuteContext(), WeakStackEntry.Pin()->InstructionIndex);
-						MicroSeconds = WeakBlueprint->RigGraphDisplaySettings.AggregateAverage(MicroSecondsFrames, MicroSeconds, CurrentMicroSeconds);
-						if(MicroSeconds > 0.0)
+						if(URigVM* VM = RigVMHost->GetVM())
 						{
-							return FText::FromString(FString::Printf(TEXT("%.02f µs"), (float)MicroSeconds));
+							const double CurrentMicroSeconds = VM->GetInstructionMicroSeconds(RigVMHost->GetRigVMExtendedExecuteContext(), WeakStackEntry.Pin()->InstructionIndex);
+							MicroSeconds = Blueprint->GetRigGraphDisplaySettings().AggregateAverage(MicroSecondsFrames, MicroSeconds, CurrentMicroSeconds);
+							if(MicroSeconds > 0.0)
+							{
+								return FText::FromString(FString::Printf(TEXT("%.02f µs"), (float)MicroSeconds));
+							}
 						}
 					}
 				}
@@ -249,9 +258,9 @@ SRigVMExecutionStackView::~SRigVMExecutionStackView()
 {
 	if (TSharedPtr<IRigVMEditor> Editor = WeakEditor.Pin())
 	{
-		if (OnModelModified.IsValid() && Editor->GetRigVMBlueprint())
+		if (OnModelModified.IsValid() && Editor->GetRigVMAsset())
 		{
-			Editor->GetRigVMBlueprint()->OnModified().Remove(OnModelModified);
+			Editor->GetRigVMAsset()->OnModified().Remove(OnModelModified);
 		}
 		if (OnHostInitializedHandle.IsValid())
 		{
@@ -265,16 +274,19 @@ SRigVMExecutionStackView::~SRigVMExecutionStackView()
 			Editor->OnPreviewHostUpdated().Remove(OnPreviewHostUpdatedHandle);
 		}
 	}
-	if (RigVMBlueprint.IsValid() && OnVMCompiledHandle.IsValid())
+	if (TStrongObjectPtr<UObject> Blueprint = WeakRigVMBlueprint.Pin())
 	{
-		RigVMBlueprint->OnVMCompiled().Remove(OnVMCompiledHandle);
+		if (OnVMCompiledHandle.IsValid())
+		{
+			Cast<IRigVMAssetInterface>(Blueprint.Get())->OnVMCompiled().Remove(OnVMCompiledHandle);
+		}
 	}
 }
 
 void SRigVMExecutionStackView::Construct( const FArguments& InArgs, TSharedRef<IRigVMEditor> InRigVMEditor)
 {
 	WeakEditor = InRigVMEditor;
-	RigVMBlueprint = InRigVMEditor->GetRigVMBlueprint();
+	WeakRigVMBlueprint = InRigVMEditor->GetRigVMAsset()->GetObject();
 	CommandList = MakeShared<FUICommandList>();
 	bSuspendModelNotifications = false;
 	bSuspendControllerSelection = false;
@@ -324,7 +336,7 @@ void SRigVMExecutionStackView::Construct( const FArguments& InArgs, TSharedRef<I
 				SAssignNew(TreeView, STreeView<TSharedPtr<FRigStackEntry>>)
 				.TreeItemsSource(&Operators)
 				.SelectionMode(ESelectionMode::Multi)
-				.OnGenerateRow(this, &SRigVMExecutionStackView::MakeTableRowWidget, RigVMBlueprint)
+				.OnGenerateRow(this, &SRigVMExecutionStackView::MakeTableRowWidget, WeakRigVMBlueprint)
 				.OnGetChildren(this, &SRigVMExecutionStackView::HandleGetChildrenForTree)
 				.OnSelectionChanged(this, &SRigVMExecutionStackView::OnSelectionChanged)
 				.OnContextMenuOpening(this, &SRigVMExecutionStackView::CreateContextMenu)
@@ -335,8 +347,9 @@ void SRigVMExecutionStackView::Construct( const FArguments& InArgs, TSharedRef<I
 
 	RefreshTreeView(nullptr, nullptr);
 
-	if (RigVMBlueprint.IsValid())
+	if (TStrongObjectPtr<UObject> StrongPtr = WeakRigVMBlueprint.Pin())
 	{
+		IRigVMAssetInterface* RigVMBlueprint = Cast<IRigVMAssetInterface>(StrongPtr.Get());
 		if (OnVMCompiledHandle.IsValid())
 		{
 			RigVMBlueprint->OnVMCompiled().Remove(OnVMCompiledHandle);
@@ -352,7 +365,7 @@ void SRigVMExecutionStackView::Construct( const FArguments& InArgs, TSharedRef<I
 		{
 			if (URigVMHost* RigVMHost = Editor->GetRigVMHost())
 			{
-				OnVMCompiled(RigVMBlueprint.Get(),  RigVMHost->GetVM(),  RigVMHost->GetRigVMExtendedExecuteContext());
+				OnVMCompiled(RigVMBlueprint->GetObject(),  RigVMHost->GetVM(),  RigVMHost->GetRigVMExtendedExecuteContext());
 			}
 		}
 	}
@@ -374,61 +387,54 @@ void SRigVMExecutionStackView::OnSelectionChanged(TSharedPtr<FRigStackEntry> Sel
 	TArray<TSharedPtr<FRigStackEntry>> SelectedItems = TreeView->GetSelectedItems();
 	if (SelectedItems.Num() > 0)
 	{
-		if (!RigVMBlueprint.IsValid())
-		{
-			return;
-		}
-
-		URigVMBlueprintGeneratedClass* GeneratedClass = RigVMBlueprint->GetRigVMBlueprintGeneratedClass();
-		if (GeneratedClass == nullptr)
-		{
-			return;
-		}
-
-		URigVMHost* Host = nullptr;
-		if (TSharedPtr<IRigVMEditor> Editor = WeakEditor.Pin())
+		if (TStrongObjectPtr<UObject> StrongPtr = WeakRigVMBlueprint.Pin())
 		{
-			Host = Editor->GetRigVMHost();
-		}
-		if (Host == nullptr || Host->GetVM() == nullptr)
-		{
-			return;
-		}
+			IRigVMAssetInterface* RigVMBlueprint = Cast<IRigVMAssetInterface>(StrongPtr.Get());
+			URigVMHost* Host = nullptr;
+			if (TSharedPtr<IRigVMEditor> Editor = WeakEditor.Pin())
+			{
+				Host = Editor->GetRigVMHost();
+			}
+			if (Host == nullptr || Host->GetVM() == nullptr)
+			{
+				return;
+			}
 
-		const FRigVMByteCode& ByteCode = Host->GetVM()->GetByteCode();
+			const FRigVMByteCode& ByteCode = Host->GetVM()->GetByteCode();
 
-		TMap<URigVMGraph*, TArray<FName>> SelectedNodesPerGraph;
-		for (TSharedPtr<FRigStackEntry>& Entry : SelectedItems)
-		{
-			for(int32 StackIndex = 0; StackIndex < Entry->Callstack.Num(); StackIndex++)
+			TMap<URigVMGraph*, TArray<FName>> SelectedNodesPerGraph;
+			for (TSharedPtr<FRigStackEntry>& Entry : SelectedItems)
 			{
-				const UObject* Subject = Entry->Callstack[StackIndex];
-				URigVMGraph* SubjectGraph = nullptr;
-
-				FName NodeName = NAME_None;
-				if (const URigVMNode* Node = Cast<URigVMNode>(Subject))
+				for(int32 StackIndex = 0; StackIndex < Entry->Callstack.Num(); StackIndex++)
 				{
-					NodeName = Node->GetFName();
-					SubjectGraph = Node->GetGraph();
-				}
-				else if (const URigVMPin* Pin = Cast<URigVMPin>(Subject))
-				{
-					NodeName = Pin->GetNode()->GetFName();
-					SubjectGraph = Pin->GetGraph();
-				}
+					const UObject* Subject = Entry->Callstack[StackIndex];
+					URigVMGraph* SubjectGraph = nullptr;
 
-				if (NodeName.IsNone() || SubjectGraph == nullptr)
-				{
-					continue;
-				}
+					FName NodeName = NAME_None;
+					if (const URigVMNode* Node = Cast<URigVMNode>(Subject))
+					{
+						NodeName = Node->GetFName();
+						SubjectGraph = Node->GetGraph();
+					}
+					else if (const URigVMPin* Pin = Cast<URigVMPin>(Subject))
+					{
+						NodeName = Pin->GetNode()->GetFName();
+						SubjectGraph = Pin->GetGraph();
+					}
 
-				SelectedNodesPerGraph.FindOrAdd(SubjectGraph).AddUnique(NodeName);
+					if (NodeName.IsNone() || SubjectGraph == nullptr)
+					{
+						continue;
+					}
+
+					SelectedNodesPerGraph.FindOrAdd(SubjectGraph).AddUnique(NodeName);
+				}
 			}
-		}
 
-		for (const TPair< URigVMGraph*, TArray<FName> >& Pair : SelectedNodesPerGraph)
-		{
-			RigVMBlueprint->GetOrCreateController(Pair.Key)->SetNodeSelection(Pair.Value);
+			for (const TPair< URigVMGraph*, TArray<FName> >& Pair : SelectedNodesPerGraph)
+			{
+				RigVMBlueprint->GetOrCreateController(Pair.Key)->SetNodeSelection(Pair.Value);
+			}
 		}
 	}
 
@@ -444,7 +450,7 @@ void SRigVMExecutionStackView::BindCommands()
 	CommandList->MapAction(Commands.SelectTargetInstructions, FExecuteAction::CreateSP(this, &SRigVMExecutionStackView::HandleSelectTargetInstructions));
 }
 
-TSharedRef<ITableRow> SRigVMExecutionStackView::MakeTableRowWidget(TSharedPtr<FRigStackEntry> InItem, const TSharedRef<STableViewBase>& OwnerTable, TWeakObjectPtr<URigVMBlueprint> InBlueprint)
+TSharedRef<ITableRow> SRigVMExecutionStackView::MakeTableRowWidget(TSharedPtr<FRigStackEntry> InItem, const TSharedRef<STableViewBase>& OwnerTable, TWeakObjectPtr<UObject> InBlueprint)
 {
 	return InItem->MakeTreeRowWidget(OwnerTable, InItem.ToSharedRef(), CommandList.ToSharedRef(), SharedThis(this), InBlueprint);
 }
@@ -832,30 +838,34 @@ void SRigVMExecutionStackView::HandleFocusOnSelectedGraphNode()
 	{
 		if (TSharedPtr<IRigVMEditor> Editor = WeakEditor.Pin())
 		{
-			URigVMHost* Host = Editor->GetRigVMHost();
-			if (Host == nullptr || Host->GetVM() == nullptr)
+			if (TStrongObjectPtr<UObject> StrongPtr = WeakRigVMBlueprint.Pin())
 			{
-				return;
-			}
+				IRigVMAssetInterface* RigVMBlueprint = Cast<IRigVMAssetInterface>(StrongPtr.Get());
+				URigVMHost* Host = Editor->GetRigVMHost();
+				if (Host == nullptr || Host->GetVM() == nullptr)
+				{
+					return;
+				}
 
-			const FRigVMByteCode& ByteCode = Host->GetVM()->GetByteCode();
-			UObject* Subject = ByteCode.GetSubjectForInstruction(SelectedItems[0]->InstructionIndex);
-			if (URigVMNode* SelectedNode = Cast<URigVMNode>(Subject))
-			{
-				URigVMGraph* GraphToFocus = SelectedNode->GetGraph();
-				if (GraphToFocus && GraphToFocus->GetTypedOuter<URigVMAggregateNode>())
+				const FRigVMByteCode& ByteCode = Host->GetVM()->GetByteCode();
+				UObject* Subject = ByteCode.GetSubjectForInstruction(SelectedItems[0]->InstructionIndex);
+				if (URigVMNode* SelectedNode = Cast<URigVMNode>(Subject))
 				{
-					if(URigVMGraph* ParentGraph = GraphToFocus->GetParentGraph())
+					URigVMGraph* GraphToFocus = SelectedNode->GetGraph();
+					if (GraphToFocus && GraphToFocus->GetTypedOuter<URigVMAggregateNode>())
 					{
-						GraphToFocus = ParentGraph;
-					} 
-				}
+						if(URigVMGraph* ParentGraph = GraphToFocus->GetParentGraph())
+						{
+							GraphToFocus = ParentGraph;
+						} 
+					}
 				
-				if (UEdGraph* EdGraph = RigVMBlueprint->GetEdGraph(GraphToFocus))
-				{
-					Editor->OpenGraphAndBringToFront(EdGraph, true);
-					Editor->ZoomToSelection_Clicked();
-					Editor->HandleModifiedEvent(ERigVMGraphNotifType::NodeSelected, GraphToFocus, SelectedNode);
+					if (UEdGraph* EdGraph = RigVMBlueprint->GetEdGraph(GraphToFocus))
+					{
+						Editor->OpenGraphAndBringToFront(EdGraph, true);
+						Editor->ZoomToSelection_Clicked();
+						Editor->HandleModifiedEvent(ERigVMGraphNotifType::NodeSelected, GraphToFocus, SelectedNode);
+					}
 				}
 			}
 		}
@@ -1183,7 +1193,7 @@ void SRigVMExecutionStackView::HandleHostInitializedEvent(URigVMHost* InHost, co
 		InHost->OnInitialized_AnyThread().Remove(OnHostInitializedHandle);
 		OnHostInitializedHandle.Reset();
 
-		if (URigVMBlueprint* RigBlueprint = Editor->GetRigVMBlueprint())
+		if (TRigVMAssetInterfacePtr RigBlueprint = Editor->GetRigVMAsset())
 		{
 			TArray<URigVMGraph*> Models = RigBlueprint->GetAllModels();
 			for (URigVMGraph* Model : Models)
@@ -1212,34 +1222,33 @@ void SRigVMExecutionStackView::HandleItemMouseDoubleClick(TSharedPtr<FRigStackEn
 {
 	if (TSharedPtr<IRigVMEditor> Editor = WeakEditor.Pin())
 	{
-		if (!RigVMBlueprint.IsValid())
+		if (TStrongObjectPtr<UObject> StrongPtr = WeakRigVMBlueprint.Pin())
 		{
-			return;
-		}
-	
-		URigVMHost* Host = Cast<URigVMHost>(RigVMBlueprint->GetObjectBeingDebugged());
-		if (!Host || !Host->GetVM())
-		{
-			return;
-		}
+			IRigVMAssetInterface* RigVMBlueprint = Cast<IRigVMAssetInterface>(StrongPtr.Get());
+			URigVMHost* Host = RigVMBlueprint->GetDebuggedRigVMHost();
+			if (!Host || !Host->GetVM())
+			{
+				return;
+			}
 
-		const FRigVMByteCode& ByteCode = Host->GetVM()->GetByteCode();
-		if (URigVMNode* Subject = Cast<URigVMNode>(ByteCode.GetSubjectForInstruction(InItem->InstructionIndex)))
-		{
-			URigVMGraph* GraphToFocus = Subject->GetGraph();
-			if (GraphToFocus && GraphToFocus->GetTypedOuter<URigVMAggregateNode>())
+			const FRigVMByteCode& ByteCode = Host->GetVM()->GetByteCode();
+			if (URigVMNode* Subject = Cast<URigVMNode>(ByteCode.GetSubjectForInstruction(InItem->InstructionIndex)))
 			{
-				if(URigVMGraph* ParentGraph = GraphToFocus->GetParentGraph())
+				URigVMGraph* GraphToFocus = Subject->GetGraph();
+				if (GraphToFocus && GraphToFocus->GetTypedOuter<URigVMAggregateNode>())
 				{
-					GraphToFocus = ParentGraph;
-				} 
-			}
+					if(URigVMGraph* ParentGraph = GraphToFocus->GetParentGraph())
+					{
+						GraphToFocus = ParentGraph;
+					} 
+				}
 		
-			if(URigVMEdGraph* EdGraph = Cast<URigVMEdGraph>(RigVMBlueprint->GetEdGraph(GraphToFocus)))
-			{
-				if(const UEdGraphNode* Node = EdGraph->FindNodeForModelNodeName(Subject->GetFName()))
+				if(URigVMEdGraph* EdGraph = Cast<URigVMEdGraph>(RigVMBlueprint->GetEdGraph(GraphToFocus)))
 				{
-					Editor->JumpToHyperlink(Node, false);
+					if(const UEdGraphNode* Node = EdGraph->FindNodeForModelNodeName(Subject->GetFName()))
+					{
+						Editor->JumpToHyperlink(Node, false);
+					}
 				}
 			}
 		}
diff --git a/Engine/Plugins/Runtime/RigVM/Source/RigVMEditor/Private/Widgets/SRigVMGraphFunctionBulkEditWidget.cpp b/Engine/Plugins/Runtime/RigVM/Source/RigVMEditor/Private/Widgets/SRigVMGraphFunctionBulkEditWidget.cpp
--- a/Engine/Plugins/Runtime/RigVM/Source/RigVMEditor/Private/Widgets/SRigVMGraphFunctionBulkEditWidget.cpp
+++ b/Engine/Plugins/Runtime/RigVM/Source/RigVMEditor/Private/Widgets/SRigVMGraphFunctionBulkEditWidget.cpp
@@ -23,7 +23,7 @@
 /// SRigVMGraphFunctionBulkEditWidget
 ///////////////////////////////////////////////////////////
 
-void SRigVMGraphFunctionBulkEditWidget::Construct(const FArguments& InArgs, URigVMBlueprint* InBlueprint, URigVMController* InController, URigVMLibraryNode* InFunction, ERigVMControllerBulkEditType InEditType)
+void SRigVMGraphFunctionBulkEditWidget::Construct(const FArguments& InArgs, TRigVMAssetInterfacePtr InBlueprint, URigVMController* InController, URigVMLibraryNode* InFunction, ERigVMControllerBulkEditType InEditType)
 {
 	Blueprint = InBlueprint;
 	Controller = InController;
diff --git a/Engine/Plugins/Runtime/RigVM/Source/RigVMEditor/Private/Widgets/SRigVMGraphNode.cpp b/Engine/Plugins/Runtime/RigVM/Source/RigVMEditor/Private/Widgets/SRigVMGraphNode.cpp
--- a/Engine/Plugins/Runtime/RigVM/Source/RigVMEditor/Private/Widgets/SRigVMGraphNode.cpp
+++ b/Engine/Plugins/Runtime/RigVM/Source/RigVMEditor/Private/Widgets/SRigVMGraphNode.cpp
@@ -22,7 +22,7 @@
 #include "Kismet2/BlueprintEditorUtils.h"
 #include "PropertyPathHelpers.h"
 #include "UObject/PropertyPortFlags.h"
-#include "RigVMBlueprint.h"
+#include "RigVMBlueprintLegacy.h"
 #include "RigVMCore/RigVM.h"
 #include "RigVMModel/RigVMController.h"
 #include "RigVMModel/RigVMNode.h"
@@ -30,6 +30,7 @@
 #include "RigVMCompiler/RigVMCompiler.h"
 #include "IDocumentation.h"
 #include "DetailLayoutBuilder.h"
+#include "RigVMBlueprintUtils.h"
 #include "Styling/AppStyle.h"
 #include "Widgets/SRigVMGraphPinVariableBinding.h"
 #include "RigVMModel/Nodes/RigVMAggregateNode.h"
@@ -308,7 +309,7 @@ void SRigVMGraphNode::Construct( const FArguments& InArgs )
 		return;
 	}
 
-	Blueprint = Cast<URigVMBlueprint>(FBlueprintEditorUtils::FindBlueprintForNode(this->GraphNode));
+	Blueprint = FRigVMBlueprintUtils::FindAssetForNode(this->GraphNode);
 
 	NodeErrorType = int32(EMessageSeverity::Info) + 1;
 	this->UpdateGraphNode();
@@ -1043,19 +1044,19 @@ FSlateColor SRigVMGraphNode::GetVariableLabelTextColor(
 	return FLinearColor::White;
 }
 
-FText SRigVMGraphNode::GetVariableLabelTooltipText(TWeakObjectPtr<URigVMBlueprint> InBlueprint,
+FText SRigVMGraphNode::GetVariableLabelTooltipText(TWeakInterfacePtr<IRigVMAssetInterface> InBlueprint,
 	FName InVariableName) const
 {
 	if(InBlueprint.IsValid())
 	{
-		for(const FBPVariableDescription& Variable : InBlueprint->NewVariables)
+		for(const FRigVMGraphVariableDescription& Variable : InBlueprint->GetAssetVariables())
 		{
-			if(Variable.VarName == InVariableName)
+			if(Variable.Name == InVariableName)
 			{
-				FString Message = FString::Printf(TEXT("Variable from %s"), *InBlueprint->GetPathName()); 
-				if(Variable.HasMetaData(FBlueprintMetadata::MD_Tooltip))
+				FString Message = FString::Printf(TEXT("Variable from %s"), *InBlueprint->GetObject()->GetPathName()); 
+				if(Variable.Tooltip.IsEmpty())
 				{
-					const FString Tooltip = Variable.GetMetaData(FBlueprintMetadata::MD_Tooltip);
+					const FString Tooltip = Variable.Tooltip.ToString();
 					Message = FString::Printf(TEXT("%s\n%s"), *Message, *Tooltip);
 				}
 				return FText::FromString(Message);
@@ -1128,7 +1129,7 @@ void SRigVMGraphNode::GetNodeInfoPopups(FNodeInfoContext* Context, TArray<FGraph
 
 	URigVMHost* ActiveObject = Cast<URigVMHost>(K2Context->ActiveObjectBeingDebugged);
 	URigVMEdGraphNode* RigGraphNode = Cast<URigVMEdGraphNode>(GraphNode);
-	URigVMBlueprint* RigBlueprint = Cast<URigVMBlueprint>(K2Context->SourceBlueprint);
+	TRigVMAssetInterfacePtr RigBlueprint = K2Context->SourceBlueprint;
 
 	// Display any pending latent actions
 	if (ActiveObject && RigBlueprint && RigGraphNode)
@@ -1338,9 +1339,9 @@ TArray<FOverlayWidgetInfo> SRigVMGraphNode::GetOverlayWidgets(bool bSelected, co
 
 		if(Blueprint.IsValid())
 		{
-			const bool bShowInstructionIndex = Blueprint->RigGraphDisplaySettings.bShowNodeInstructionIndex;
-			const bool bShowNodeCounts = Blueprint->RigGraphDisplaySettings.bShowNodeRunCounts;
-			const bool bEnableProfiling = Blueprint->VMRuntimeSettings.bEnableProfiling;
+			const bool bShowInstructionIndex = Blueprint->GetRigGraphDisplaySettings().bShowNodeInstructionIndex;
+			const bool bShowNodeCounts = Blueprint->GetRigGraphDisplaySettings().bShowNodeRunCounts;
+			const bool bEnableProfiling = Blueprint->GetVMRuntimeSettings().bEnableProfiling;
 			
 			if(bShowNodeCounts || bShowInstructionIndex || bEnableProfiling)
 			{
@@ -1349,7 +1350,7 @@ TArray<FOverlayWidgetInfo> SRigVMGraphNode::GetOverlayWidgets(bool bSelected, co
 					if(bShowNodeCounts || bShowInstructionIndex)
 					{
 						const int32 Count = ModelNode->GetInstructionVisitedCount(DebuggedHost->GetRigVMExtendedExecuteContext(), DebuggedHost->GetVM(), FRigVMASTProxy());
-						if((Count > Blueprint->RigGraphDisplaySettings.NodeRunLowerBound) || bShowInstructionIndex)
+						if((Count > Blueprint->GetRigGraphDisplaySettings().NodeRunLowerBound) || bShowInstructionIndex)
 						{
 							const float VOffset = bSelected ? -2.f : 2.f;
 							const FVector2f TextSize = InstructionCountTextBlockWidget->GetDesiredSize();
@@ -1623,8 +1624,8 @@ FText SRigVMGraphNode::GetInstructionCountText() const
 {
 	if(Blueprint.IsValid())
 	{
-		bool bShowInstructionIndex = Blueprint->RigGraphDisplaySettings.bShowNodeInstructionIndex;
-		bool bShowNodeRunCount = Blueprint->RigGraphDisplaySettings.bShowNodeRunCounts;
+		bool bShowInstructionIndex = Blueprint->GetRigGraphDisplaySettings().bShowNodeInstructionIndex;
+		bool bShowNodeRunCount = Blueprint->GetRigGraphDisplaySettings().bShowNodeRunCounts;
 		if(bShowInstructionIndex || bShowNodeRunCount)
 		{
 			if (ModelNode.IsValid())
@@ -1636,7 +1637,7 @@ FText SRigVMGraphNode::GetInstructionCountText() const
 					if(bShowNodeRunCount)
 					{
 						RunCount = ModelNode->GetInstructionVisitedCount(DebuggedHost->GetRigVMExtendedExecuteContext(), DebuggedHost->GetVM(), FRigVMASTProxy());
-						bShowNodeRunCount = RunCount > Blueprint->RigGraphDisplaySettings.NodeRunLowerBound;
+						bShowNodeRunCount = RunCount > Blueprint->GetRigGraphDisplaySettings().NodeRunLowerBound;
 
 						// toodoo here we want to differentiate between primary and secondary instructions. (call extern vs copy).
 						// const TArray<int32>& Instructions = ModelNode->GetInstructionsForVM(DebuggedHost->GetRigVMExtendedExecuteContext(), DebuggedHost->GetVM());
@@ -1687,7 +1688,7 @@ FText SRigVMGraphNode::GetInstructionDurationText() const
 {
 	if(Blueprint.IsValid())
 	{
-		if(Blueprint->VMRuntimeSettings.bEnableProfiling)
+		if(Blueprint->GetVMRuntimeSettings().bEnableProfiling)
 		{
 			if(const URigVMEdGraphNode* RigGraphNode = Cast<URigVMEdGraphNode>(GraphNode))
 			{
@@ -1916,7 +1917,7 @@ FReply SRigVMGraphNode::OnExpanderArrowClicked(int32 InPinInfoIndex)
 				Algo::Reverse(PinInfoIndicesToModify);
 			}
 
-			FGuardSkipDirtyBlueprintStatus GuardDirtyBlueprintStatus(Blueprint, true);
+			FGuardSkipDirtyBlueprintStatus GuardDirtyBlueprintStatus(Blueprint.GetObject(), true);
 
 			Controller->OpenUndoBracket(PinInfo.bExpanded ? TEXT("Collapsing Pin") : TEXT("Expanding Pin"));
 			for(const int32& PinInfoIndexToModify : PinInfoIndicesToModify)
@@ -2534,7 +2535,7 @@ void SRigVMGraphNode::UpdatePinTreeView()
 	ERigVMTagDisplayMode TagDisplayMode = ERigVMTagDisplayMode::All;
 	if(Blueprint.IsValid())
 	{
-		TagDisplayMode = Blueprint->RigGraphDisplaySettings.TagDisplayMode;
+		TagDisplayMode = Blueprint->GetRigGraphDisplaySettings().TagDisplayMode;
 	}
 	if(ModelNode.IsValid() && TagDisplayMode != ERigVMTagDisplayMode::None)
 	{
@@ -2857,7 +2858,7 @@ void SRigVMGraphNode::UpdatePinTreeView()
 	if(URigVMFunctionReferenceNode* FunctionReferenceNode = Cast<URigVMFunctionReferenceNode>(ModelNode))
 	{
 		TWeakObjectPtr<URigVMFunctionReferenceNode> WeakFunctionReferenceNode = FunctionReferenceNode;
-		TWeakObjectPtr<URigVMBlueprint> WeakRigVMBlueprint = Blueprint;
+		TWeakInterfacePtr<IRigVMAssetInterface> WeakRigVMBlueprint = Blueprint;
 
 		// add the entries for the variable remapping
 		for(const TSharedPtr<FRigVMExternalVariable>& ExternalVariable : RigGraphNode->ExternalVariables)
@@ -2890,7 +2891,7 @@ void SRigVMGraphNode::UpdatePinTreeView()
 				.Padding(PinWidgetSidePadding, TopPadding, PinWidgetSidePadding, 0.f)
 				[
 					SNew(SRigVMGraphVariableBinding)
-					.Blueprint(Blueprint.Get())
+					.Asset(Blueprint.GetObject())
 					.FunctionReferenceNode(FunctionReferenceNode)
 					.InnerVariableName(ExternalVariable->Name)
 				]
diff --git a/Engine/Plugins/Runtime/RigVM/Source/RigVMEditor/Private/Widgets/SRigVMGraphNodeComment.cpp b/Engine/Plugins/Runtime/RigVM/Source/RigVMEditor/Private/Widgets/SRigVMGraphNodeComment.cpp
--- a/Engine/Plugins/Runtime/RigVM/Source/RigVMEditor/Private/Widgets/SRigVMGraphNodeComment.cpp
+++ b/Engine/Plugins/Runtime/RigVM/Source/RigVMEditor/Private/Widgets/SRigVMGraphNodeComment.cpp
@@ -6,7 +6,7 @@
 #include "Framework/Application/SlateApplication.h"
 #include "PropertyPathHelpers.h"
 #include "UObject/PropertyPortFlags.h"
-#include "RigVMBlueprint.h"
+#include "RigVMBlueprintLegacy.h"
 #include "EdGraph/RigVMEdGraph.h"
 #include "EdGraph/RigVMEdGraphSchema.h"
 #include "RigVMModel/RigVMController.h"
@@ -40,7 +40,7 @@ FReply SRigVMGraphNodeComment::OnMouseButtonUp(const FGeometry& MyGeometry, cons
 			UEdGraphNode_Comment* CommentNode = CastChecked<UEdGraphNode_Comment>(GraphNode);
 			if (URigVMEdGraph* Graph = Cast<URigVMEdGraph>(CommentNode->GetOuter()))
 			{
-				if (URigVMBlueprint* Blueprint = Cast<URigVMBlueprint>(Graph->GetOuter()))
+				if (TRigVMAssetInterfacePtr Blueprint = Graph->GetOuter())
 				{
 					FVector2D Position(CommentNode->NodePosX, CommentNode->NodePosY);
 					FVector2D Size(CommentNode->NodeWidth, CommentNode->NodeHeight);
@@ -104,7 +104,7 @@ void SRigVMGraphNodeComment::Tick(const FGeometry& AllottedGeometry, const doubl
 		{
 			if (URigVMEdGraph* Graph = Cast<URigVMEdGraph>(CommentNode->GetOuter()))
 			{
-				if (URigVMBlueprint* Blueprint = Cast<URigVMBlueprint>(Graph->GetOuter()))
+				if (TRigVMAssetInterfacePtr Blueprint = Graph->GetOuter())
 				{
 					if (URigVMController* Controller = Blueprint->GetController(Graph))
 					{
@@ -126,7 +126,7 @@ void SRigVMGraphNodeComment::Tick(const FGeometry& AllottedGeometry, const doubl
 			{
 				if (URigVMEdGraph* Graph = Cast<URigVMEdGraph>(CommentNode->GetOuter()))
 				{
-					if (URigVMBlueprint* Blueprint = Cast<URigVMBlueprint>(Graph->GetOuter()))
+					if (TRigVMAssetInterfacePtr Blueprint = Graph->GetOuter())
 					{
 						if (URigVMController* Controller = Blueprint->GetController(Graph))
 						{
diff --git a/Engine/Plugins/Runtime/RigVM/Source/RigVMEditor/Private/Widgets/SRigVMGraphPinEnumPicker.cpp b/Engine/Plugins/Runtime/RigVM/Source/RigVMEditor/Private/Widgets/SRigVMGraphPinEnumPicker.cpp
--- a/Engine/Plugins/Runtime/RigVM/Source/RigVMEditor/Private/Widgets/SRigVMGraphPinEnumPicker.cpp
+++ b/Engine/Plugins/Runtime/RigVM/Source/RigVMEditor/Private/Widgets/SRigVMGraphPinEnumPicker.cpp
@@ -5,7 +5,7 @@
 
 #include "DetailLayoutBuilder.h"
 #include "EdGraph/RigVMEdGraph.h"
-#include "RigVMBlueprint.h"
+#include "RigVMBlueprintLegacy.h"
 #include "Widgets/Layout/SBox.h"
 #include "ScopedTransaction.h"
 #include "SSearchableComboBox.h"
@@ -118,7 +118,7 @@ void SRigVMGraphPinEnumPicker::HandleControlEnumChanged(TSharedPtr<FString> InIt
 {
 	if(ModelPin)
 	{
-		if(URigVMBlueprint* Blueprint = ModelPin->GetTypedOuter<URigVMBlueprint>())
+		if(TRigVMAssetInterfacePtr Blueprint = IRigVMAssetInterface::GetInterfaceOuter(ModelPin))
 		{
 			if(URigVMController* Controller = Blueprint->GetOrCreateController(ModelPin->GetGraph()))
 			{
diff --git a/Engine/Plugins/Runtime/RigVM/Source/RigVMEditor/Private/Widgets/SRigVMGraphPinQuat.cpp b/Engine/Plugins/Runtime/RigVM/Source/RigVMEditor/Private/Widgets/SRigVMGraphPinQuat.cpp
--- a/Engine/Plugins/Runtime/RigVM/Source/RigVMEditor/Private/Widgets/SRigVMGraphPinQuat.cpp
+++ b/Engine/Plugins/Runtime/RigVM/Source/RigVMEditor/Private/Widgets/SRigVMGraphPinQuat.cpp
@@ -4,7 +4,7 @@
 #include "Widgets/SRigVMGraphPinQuat.h"
 #include "EdGraph/RigVMEdGraph.h"
 #include "EdGraph/RigVMEdGraphNode.h"
-#include "RigVMBlueprint.h"
+#include "RigVMBlueprintLegacy.h"
 #include "RigVMModel/RigVMController.h"
 #include "Logging/LogScopedVerbosityOverride.h"
 #include "Widgets/Layout/SBox.h"
@@ -80,7 +80,7 @@ void SRigVMGraphPinQuat::OnRotatorCommitted(FRotator InRotator, ETextCommit::Typ
 		FString NewDefaultValue;	
 		TBaseStructure<FRotator>::Get()->ExportText(NewDefaultValue, &InRotator, &InRotator, nullptr, PPF_None, nullptr);
 
-		if(URigVMBlueprint* Blueprint = ModelPin->GetTypedOuter<URigVMBlueprint>())
+		if(TRigVMAssetInterfacePtr Blueprint = IRigVMAssetInterface::GetInterfaceOuter(ModelPin))
 		{
 			if(URigVMController* Controller = Blueprint->GetOrCreateController(ModelPin->GetGraph()))
 			{
diff --git a/Engine/Plugins/Runtime/RigVM/Source/RigVMEditor/Private/Widgets/SRigVMGraphPinUserDataPath.cpp b/Engine/Plugins/Runtime/RigVM/Source/RigVMEditor/Private/Widgets/SRigVMGraphPinUserDataPath.cpp
--- a/Engine/Plugins/Runtime/RigVM/Source/RigVMEditor/Private/Widgets/SRigVMGraphPinUserDataPath.cpp
+++ b/Engine/Plugins/Runtime/RigVM/Source/RigVMEditor/Private/Widgets/SRigVMGraphPinUserDataPath.cpp
@@ -270,7 +270,7 @@ void SRigVMUserDataPath::FillUserDataPathMenu(FMenuBuilder& InMenuBuilder, FStri
 
 void SRigVMUserDataPath::HandleSetUserDataPath(FString InUserDataPath)
 {
-	URigVMBlueprint* Blueprint = GetBlueprint();
+	TRigVMAssetInterfacePtr Blueprint = GetRigVMAsset();
 	if(ModelPins.IsEmpty() || (Blueprint == nullptr))
 	{
 		return;
@@ -321,11 +321,11 @@ void SRigVMUserDataPath::HandleSetUserDataPath(FString InUserDataPath)
 	}
 }
 
-URigVMBlueprint* SRigVMUserDataPath::GetBlueprint() const
+TRigVMAssetInterfacePtr SRigVMUserDataPath::GetRigVMAsset() const
 {
 	if(!ModelPins.IsEmpty())
 	{
-		return ModelPins[0]->GetTypedOuter<URigVMBlueprint>();
+		return IRigVMAssetInterface::GetInterfaceOuter(ModelPins[0]);
 	}
 	return nullptr;
 }
@@ -368,7 +368,7 @@ const UNameSpacedUserData* SRigVMUserDataPath::GetUserDataObject() const
 		return nullptr;
 	}
 	
-	if(const UBlueprint* Blueprint = GetBlueprint())
+	if(const TRigVMAssetInterfacePtr Blueprint = GetRigVMAsset())
 	{
 		if(const UObject* DebuggedObject = Blueprint->GetObjectBeingDebugged())
 		{
diff --git a/Engine/Plugins/Runtime/RigVM/Source/RigVMEditor/Private/Widgets/SRigVMGraphPinVariableBinding.cpp b/Engine/Plugins/Runtime/RigVM/Source/RigVMEditor/Private/Widgets/SRigVMGraphPinVariableBinding.cpp
--- a/Engine/Plugins/Runtime/RigVM/Source/RigVMEditor/Private/Widgets/SRigVMGraphPinVariableBinding.cpp
+++ b/Engine/Plugins/Runtime/RigVM/Source/RigVMEditor/Private/Widgets/SRigVMGraphPinVariableBinding.cpp
@@ -24,7 +24,7 @@ void SRigVMGraphVariableBinding::Construct(const FArguments& InArgs)
 	this->ModelPins = InArgs._ModelPins;
 	this->FunctionReferenceNode = InArgs._FunctionReferenceNode;
 	this->InnerVariableName = InArgs._InnerVariableName;
-	this->Blueprint = InArgs._Blueprint;
+	this->Blueprint = InArgs._Asset;
 	this->bCanRemoveBinding = InArgs._CanRemoveBinding;
 
 	IPropertyAccessEditor& PropertyAccessEditor = IModularFeatures::Get().GetModularFeature<IPropertyAccessEditor>("PropertyAccessEditor");
@@ -57,9 +57,11 @@ void SRigVMGraphVariableBinding::Construct(const FArguments& InArgs)
 		nullptr,
 		FMenuExtensionDelegate::CreateSP(this, &SRigVMGraphVariableBinding::FillLocalVariableMenu));
 
+	TArray<FBindingContextStruct> BindingContexts;
+	BindingContexts.Add(Blueprint->GetRigVMGeneratedClass());
 	this->ChildSlot
 	[
-		PropertyAccessEditor.MakePropertyBindingWidget(Blueprint, BindingArgs)
+		PropertyAccessEditor.MakePropertyBindingWidget(BindingContexts, BindingArgs)
 	];
 }
 
@@ -133,11 +135,11 @@ FLinearColor SRigVMGraphVariableBinding::GetBindingColor() const
 			}
 		}
 
-		for (const FBPVariableDescription& VariableDescription : Blueprint->NewVariables)
+		for (const FRigVMGraphVariableDescription& VariableDescription : Blueprint->GetAssetVariables())
 		{
-			if (VariableDescription.VarName == BoundVariable)
+			if (VariableDescription.Name == BoundVariable)
 			{
-				return Schema->GetPinTypeColor(VariableDescription.VarType);
+				return Schema->GetPinTypeColor(VariableDescription.ToPinType());
 			}
 		}
 
@@ -219,7 +221,7 @@ bool SRigVMGraphVariableBinding::OnCanBindToClass(UClass* InClass) const
 {
 	if (InClass)
 	{
-		return InClass->ClassGeneratedBy == Blueprint;
+		return InClass->ClassGeneratedBy == Blueprint.GetObject();
 	}
 	return true;
 }
@@ -231,8 +233,11 @@ void SRigVMGraphVariableBinding::OnAddBinding(FName InPropertyName, const TArray
 		TArray<FString> Parts;
 		for (const FBindingChainElement& ChainElement : InBindingChain)
 		{
-			ensure(ChainElement.Field);
-			Parts.Add(ChainElement.Field.GetName());
+			//ensure(ChainElement.Field);
+			if (!ChainElement.Field.GetFName().IsNone())
+			{
+				Parts.Add(ChainElement.Field.GetName());
+			}
 		}
 
 		if(ModelPins.Num() > 0)
@@ -390,7 +395,7 @@ void SRigVMGraphVariableBinding::HandleBindToLocalVariable(FRigVMGraphVariableDe
 void SRigVMGraphPinVariableBinding::Construct(const FArguments& InArgs, UEdGraphPin* InGraphPinObj)
 {
 	this->ModelPins = InArgs._ModelPins;
-	this->Blueprint = InArgs._Blueprint;
+	this->Blueprint = InArgs._Asset;
 
 	SGraphPin::Construct(SGraphPin::FArguments(), InGraphPinObj);
 }
@@ -398,7 +403,7 @@ void SRigVMGraphPinVariableBinding::Construct(const FArguments& InArgs, UEdGraph
 TSharedRef<SWidget>	SRigVMGraphPinVariableBinding::GetDefaultValueWidget()
 {
 	return SNew(SRigVMGraphVariableBinding)
-		.Blueprint(Blueprint)
+		.Asset(Blueprint)
 		.ModelPins(ModelPins);
 }
 
diff --git a/Engine/Plugins/Runtime/RigVM/Source/RigVMEditor/Private/Widgets/SRigVMGraphPinVariableName.cpp b/Engine/Plugins/Runtime/RigVM/Source/RigVMEditor/Private/Widgets/SRigVMGraphPinVariableName.cpp
--- a/Engine/Plugins/Runtime/RigVM/Source/RigVMEditor/Private/Widgets/SRigVMGraphPinVariableName.cpp
+++ b/Engine/Plugins/Runtime/RigVM/Source/RigVMEditor/Private/Widgets/SRigVMGraphPinVariableName.cpp
@@ -8,7 +8,7 @@
 #include "EdGraphSchema_K2.h"
 #include "ScopedTransaction.h"
 #include "DetailLayoutBuilder.h"
-#include "RigVMBlueprint.h"
+#include "RigVMBlueprintLegacy.h"
 
 void SRigVMGraphPinVariableName::Construct(const FArguments& InArgs, UEdGraphPin* InGraphPinObj)
 {
@@ -115,12 +115,12 @@ TArray<TSharedPtr<FString>>& SRigVMGraphPinVariableName::GetVariableNames()
 				}
 			}
 
-			if (URigVMBlueprint* Blueprint = Cast<URigVMBlueprint>(ModelNode->GetGraph()->GetOuter()))
+			if (TRigVMAssetInterfacePtr Blueprint = ModelNode->GetGraph()->GetOuter())
 			{
-				for (FBPVariableDescription& BPVariable : Blueprint->NewVariables)
+				for (FRigVMGraphVariableDescription& BPVariable : Blueprint->GetAssetVariables())
 				{
 					FRigVMGraphVariableDescription* Found = VariableDescriptions.FindByPredicate([&BPVariable](const FRigVMGraphVariableDescription& Description) {
-						if (Description.Name == BPVariable.VarName)
+						if (Description.Name == BPVariable.Name)
 						{
 							return true;
 						}
@@ -129,7 +129,7 @@ TArray<TSharedPtr<FString>>& SRigVMGraphPinVariableName::GetVariableNames()
 
 					if (!Found)
 					{
-						VariableNames.Add(MakeShared<FString>(BPVariable.VarName.ToString()));
+						VariableNames.Add(MakeShared<FString>(BPVariable.Name.ToString()));
 					}
 				}
 			}
diff --git a/Engine/Plugins/Runtime/RigVM/Source/RigVMEditor/Private/Widgets/SRigVMLogWidget.cpp b/Engine/Plugins/Runtime/RigVM/Source/RigVMEditor/Private/Widgets/SRigVMLogWidget.cpp
--- a/Engine/Plugins/Runtime/RigVM/Source/RigVMEditor/Private/Widgets/SRigVMLogWidget.cpp
+++ b/Engine/Plugins/Runtime/RigVM/Source/RigVMEditor/Private/Widgets/SRigVMLogWidget.cpp
@@ -87,16 +87,54 @@ void SRigVMLogWidget::UnbindLog(TSharedPtr<IMessageLogListing> InListing)
 	}
 }
 
-void SRigVMLogWidget::BindLog(const URigVMBlueprint* InRigVMBlueprint)
+void SRigVMLogWidget::BindLog(const TRigVMAssetInterfacePtr InRigVMBlueprint)
 {
-	TSharedRef<IMessageLogListing> Listing = FCompilerResultsLog::GetBlueprintMessageLog(const_cast<URigVMBlueprint*>(InRigVMBlueprint));
-	BindLog(Listing);
+	TSharedPtr<IMessageLogListing> Listing = nullptr;
+	FMessageLogModule& MessageLogModule = FModuleManager::LoadModuleChecked<FMessageLogModule>("MessageLog");
+
+	const FName LogName = InRigVMBlueprint->GetCompileLog().GetLogName();
+
+	// Reuse any existing log, or create a new one (that is not held onto bey the message log system)
+	if(MessageLogModule.IsRegisteredLogListing(LogName))
+	{
+		Listing = MessageLogModule.GetLogListing(LogName);
+	}
+	else
+	{
+		FMessageLogInitializationOptions LogInitOptions;
+		LogInitOptions.bShowInLogWindow = false;
+		Listing = MessageLogModule.CreateLogListing(LogName, LogInitOptions);
+	}
+	
+	if (Listing)
+	{
+		BindLog(Listing);
+	}
 }
 
-void SRigVMLogWidget::UnbindLog(const URigVMBlueprint* InRigVMBlueprint)
+void SRigVMLogWidget::UnbindLog(const TRigVMAssetInterfacePtr InRigVMBlueprint)
 {
-	TSharedRef<IMessageLogListing> Listing = FCompilerResultsLog::GetBlueprintMessageLog(const_cast<URigVMBlueprint*>(InRigVMBlueprint));
-	UnbindLog(Listing);
+	TSharedPtr<IMessageLogListing> Listing = nullptr;
+	FMessageLogModule& MessageLogModule = FModuleManager::LoadModuleChecked<FMessageLogModule>("MessageLog");
+
+	const FName LogName = InRigVMBlueprint->GetCompileLog().GetLogName();
+
+	// Reuse any existing log, or create a new one (that is not held onto bey the message log system)
+	if(MessageLogModule.IsRegisteredLogListing(LogName))
+	{
+		Listing = MessageLogModule.GetLogListing(LogName);
+	}
+	else
+	{
+		FMessageLogInitializationOptions LogInitOptions;
+		LogInitOptions.bShowInLogWindow = false;
+		Listing = MessageLogModule.CreateLogListing(LogName, LogInitOptions);
+	}
+	
+	if (Listing)
+	{
+		UnbindLog(Listing);
+	}
 }
 
 void SRigVMLogWidget::OnBoundListingChanged(TWeakPtr<IMessageLogListing> InWeakListing)
diff --git a/Engine/Plugins/Runtime/RigVM/Source/RigVMEditor/Private/Widgets/SRigVMSwapAssetReferencesWidget.cpp b/Engine/Plugins/Runtime/RigVM/Source/RigVMEditor/Private/Widgets/SRigVMSwapAssetReferencesWidget.cpp
--- a/Engine/Plugins/Runtime/RigVM/Source/RigVMEditor/Private/Widgets/SRigVMSwapAssetReferencesWidget.cpp
+++ b/Engine/Plugins/Runtime/RigVM/Source/RigVMEditor/Private/Widgets/SRigVMSwapAssetReferencesWidget.cpp
@@ -3,7 +3,7 @@
 #include "Widgets/SRigVMSwapAssetReferencesWidget.h"
 #include "Widgets/SRigVMBulkEditDialog.h"
 
-#include "RigVMBlueprint.h"
+#include "RigVMBlueprintLegacy.h"
 #include "Dialogs/Dialogs.h"
 #include "AssetRegistry/AssetRegistryModule.h"
 #include "Styling/AppStyle.h"
@@ -68,8 +68,8 @@ bool FRigVMTreeAssetVariantFilter::Filters(TSharedRef<FRigVMTreeNode>& InNode, c
 		{
 			if (!SourceVariants.Contains(SwapContext->GetSourceAsset().GetFullName()))
 			{
-				static const FName AssetVariantPropertyName = GET_MEMBER_NAME_CHECKED(URigVMBlueprint, AssetVariant);
-				const FProperty* AssetVariantProperty = CastField<FProperty>(URigVMBlueprint::StaticClass()->FindPropertyByName(AssetVariantPropertyName));
+				static const FName AssetVariantPropertyName = TEXT("AssetVariant");
+				const FProperty* AssetVariantProperty = CastField<FProperty>(SwapContext->GetSourceAsset().GetClass()->FindPropertyByName(AssetVariantPropertyName));
 				const FString VariantStr = SwapContext->GetSourceAsset().GetTagValueRef<FString>(AssetVariantPropertyName);
 				FRigVMVariant AssetVariant;
 				AssetVariantProperty->ImportText_Direct(*VariantStr, &AssetVariant, nullptr, EPropertyPortFlags::PPF_None);
@@ -308,7 +308,7 @@ TArray<TSharedRef<FRigVMTreeNode>> SRigVMSwapAssetReferencesWidget::GetAssetNode
 {
 	TArray<TSharedRef<FRigVMTreeNode>> Nodes;
 
-	const TArray<FAssetData> Assets = URigVMEditorBlueprintLibrary::GetAssetsWithFilter(URigVMBlueprint::StaticClass(),
+	const TArray<FAssetData> Assets = URigVMEditorBlueprintLibrary::GetAssetsWithFilter(nullptr,
 		FRigVMAssetDataFilter::CreateLambda([this, InPhase] (const FAssetData& AssetData)
 		{
 			TArray<FRigVMAssetDataFilter>& Filters = (InPhase == PHASE_PICKSOURCE) ? SourceAssetFilters : TargetAssetFilters;
diff --git a/Engine/Plugins/Runtime/RigVM/Source/RigVMEditor/Private/Widgets/SRigVMSwapFunctionsWidget.cpp b/Engine/Plugins/Runtime/RigVM/Source/RigVMEditor/Private/Widgets/SRigVMSwapFunctionsWidget.cpp
--- a/Engine/Plugins/Runtime/RigVM/Source/RigVMEditor/Private/Widgets/SRigVMSwapFunctionsWidget.cpp
+++ b/Engine/Plugins/Runtime/RigVM/Source/RigVMEditor/Private/Widgets/SRigVMSwapFunctionsWidget.cpp
@@ -3,7 +3,7 @@
 #include "Widgets/SRigVMSwapFunctionsWidget.h"
 #include "Widgets/SRigVMBulkEditDialog.h"
 
-#include "RigVMBlueprint.h"
+#include "RigVMBlueprintLegacy.h"
 #include "Dialogs/Dialogs.h"
 #include "AssetRegistry/AssetRegistryModule.h"
 #include "Styling/AppStyle.h"
@@ -463,14 +463,14 @@ URigVMFunctionReferenceNode* FRigVMSwapFunctionTask::GetReferenceNode(const TSha
 	return ReferenceNode;
 }
 
-URigVMBlueprint* FRigVMSwapFunctionTask::GetBlueprint(const TSharedRef<FRigVMTreePhase>& InPhase) const
+TRigVMAssetInterfacePtr FRigVMSwapFunctionTask::GetRigVMAsset(const TSharedRef<FRigVMTreePhase>& InPhase) const
 {
 	URigVMFunctionReferenceNode* ReferenceNode = GetReferenceNode(InPhase);
 	if(ReferenceNode == nullptr)
 	{
 		return nullptr;
 	}
-	return ReferenceNode->GetTypedOuter<URigVMBlueprint>();
+	return IRigVMAssetInterface::GetInterfaceOuter(ReferenceNode);
 }
 
 void SRigVMSwapFunctionsWidget::Construct(const FArguments& InArgs)
@@ -828,7 +828,7 @@ FReply SRigVMSwapFunctionsWidget::OnNodeDoubleClicked(TSharedRef<FRigVMTreeNode>
 			// force load
 			if(UObject* TopLevelObject = AssetData.GetSoftObjectPath().GetWithoutSubPath().TryLoad())
 			{
-				if(TopLevelObject->GetClass()->IsChildOf(URigVMBlueprint::StaticClass()))
+				if(TopLevelObject->GetClass()->ImplementsInterface(URigVMAssetInterface::StaticClass()))
 				{
 					if(const UObject* ObjectReference = StaticLoadObject(UObject::StaticClass(), nullptr, *Node->GetPath(), nullptr))
 					{
diff --git a/Engine/Plugins/Runtime/RigVM/Source/RigVMEditor/Public/EdGraph/NodeSpawners/RigVMEdGraphEnumNodeSpawner.h b/Engine/Plugins/Runtime/RigVM/Source/RigVMEditor/Public/EdGraph/NodeSpawners/RigVMEdGraphEnumNodeSpawner.h
--- a/Engine/Plugins/Runtime/RigVM/Source/RigVMEditor/Public/EdGraph/NodeSpawners/RigVMEdGraphEnumNodeSpawner.h
+++ b/Engine/Plugins/Runtime/RigVM/Source/RigVMEditor/Public/EdGraph/NodeSpawners/RigVMEdGraphEnumNodeSpawner.h
@@ -3,38 +3,29 @@
 #pragma once
 
 #include "CoreMinimal.h"
-#include "UObject/ObjectMacros.h"
-#include "UObject/Object.h"
-#include "Templates/SubclassOf.h"
-#include "UObject/UnrealType.h"
-#include "EdGraph/EdGraphPin.h"
-#include "Engine/Blueprint.h"
 #include "EdGraph/NodeSpawners/RigVMEdGraphNodeSpawner.h"
-#include "RigVMEdGraphEnumNodeSpawner.generated.h"
 
 #define UE_API RIGVMEDITOR_API
 
 class URigVMEdGraphNode;
 
-UCLASS(MinimalAPI, Transient)
 class URigVMEdGraphEnumNodeSpawner : public URigVMEdGraphNodeSpawner
 {
-	GENERATED_BODY()
-
 public:
 
+	virtual ~URigVMEdGraphEnumNodeSpawner() {}
+
 	/**
 	 * Creates a new URigVMEdGraphEnumNodeSpawner
 	 * 
 	 * @return A newly allocated instance of this class.
 	 */
-	static UE_API URigVMEdGraphEnumNodeSpawner* CreateForEnum(UEnum* InEnum, const FText& InMenuDesc, const FText& InCategory, const FText& InTooltip);
+	UE_API URigVMEdGraphEnumNodeSpawner(UEnum* InEnum, const FText& InMenuDesc, const FText& InCategory, const FText& InTooltip);
 
-	// UBlueprintNodeSpawner interface
-	UE_API virtual FBlueprintNodeSignature GetSpawnerSignature() const override;
-	UE_API virtual FBlueprintActionUiSpec GetUiSpec(FBlueprintActionContext const& Context, FBindingSet const& Bindings) const override;
-	UE_API virtual UEdGraphNode* Invoke(UEdGraph* ParentGraph, FBindingSet const& Bindings, FVector2D const Location) const override;
-	// End UBlueprintNodeSpawner interface
+	// URigVMEdGraphNodeSpawner interface
+	UE_API virtual FString GetSpawnerSignature() const override;
+	UE_API virtual URigVMEdGraphNode* Invoke(URigVMEdGraph* ParentGraph, FVector2D const Location) const override;
+	// End URigVMEdGraphNodeSpawner interface
 
 private:
 
diff --git a/Engine/Plugins/Runtime/RigVM/Source/RigVMEditor/Public/EdGraph/NodeSpawners/RigVMEdGraphFunctionRefNodeSpawner.h b/Engine/Plugins/Runtime/RigVM/Source/RigVMEditor/Public/EdGraph/NodeSpawners/RigVMEdGraphFunctionRefNodeSpawner.h
--- a/Engine/Plugins/Runtime/RigVM/Source/RigVMEditor/Public/EdGraph/NodeSpawners/RigVMEdGraphFunctionRefNodeSpawner.h
+++ b/Engine/Plugins/Runtime/RigVM/Source/RigVMEditor/Public/EdGraph/NodeSpawners/RigVMEdGraphFunctionRefNodeSpawner.h
@@ -10,45 +10,48 @@
 #include "EdGraph/EdGraphPin.h"
 #include "Engine/Blueprint.h"
 #include "EdGraph/NodeSpawners/RigVMEdGraphNodeSpawner.h"
-#include "RigVMBlueprint.h"
+#include "RigVMBlueprintLegacy.h"
 #include "EdGraph/RigVMEdGraph.h"
 #include "EdGraph/RigVMEdGraphNode.h"
 #include "RigVMModel/Nodes/RigVMUnitNode.h"
-#include "RigVMEdGraphFunctionRefNodeSpawner.generated.h"
 
 #define UE_API RIGVMEDITOR_API
 
 class URigVMEdGraphNode;
 
-UCLASS(MinimalAPI, Transient)
 class URigVMEdGraphFunctionRefNodeSpawner : public URigVMEdGraphNodeSpawner
 {
-	GENERATED_BODY()
 
 public:
 
+	virtual ~URigVMEdGraphFunctionRefNodeSpawner() {}
+	
 	/**
 	 * Creates a new URigVMEdGraphFunctionRefNodeSpawner, charged with spawning a function reference
 	 * 
 	 * @return A newly allocated instance of this class.
 	 */
-	static UE_API URigVMEdGraphFunctionRefNodeSpawner* CreateFromFunction(URigVMLibraryNode* InFunction);
+	UE_DEPRECATED(5.7, "Plase use URigVMEdGraphFunctionRefNodeSpawner(URigVMLibraryNode* InFunction)")
+	static UE_API URigVMEdGraphFunctionRefNodeSpawner* CreateFromFunction(URigVMLibraryNode* InFunction) { return nullptr; }
+	UE_API URigVMEdGraphFunctionRefNodeSpawner(URigVMLibraryNode* InFunction);
 
 	/**
 	 * Creates a new URigVMEdGraphFunctionRefNodeSpawner, charged with spawning a function reference
 	 * 
 	 * @return A newly allocated instance of this class.
 	 */
-	static UE_API URigVMEdGraphFunctionRefNodeSpawner* CreateFromAssetData(const FAssetData& InAssetData, const FRigVMGraphFunctionHeader& InPublicFunction);
-	static UE_API URigVMEdGraphFunctionRefNodeSpawner* CreateFromAssetData(const FAssetData& InAssetData, const FRigVMOldPublicFunctionData& InPublicFunction);
-
-	// UBlueprintNodeSpawner interface
-	UE_API virtual void Prime() override;
-	UE_API virtual FBlueprintNodeSignature GetSpawnerSignature() const override;
-	UE_API virtual FBlueprintActionUiSpec GetUiSpec(FBlueprintActionContext const& Context, FBindingSet const& Bindings) const override;
-	UE_API virtual UEdGraphNode* Invoke(UEdGraph* ParentGraph, FBindingSet const& Bindings, FVector2D const Location) const override;
-	UE_API virtual bool IsTemplateNodeFilteredOut(FBlueprintActionFilter const& Filter) const override;
-	// End UBlueprintNodeSpawner interface
+	UE_DEPRECATED(5.7, "Please use URigVMEdGraphFunctionRefNodeSpawner(const FAssetData& InAssetData, const FRigVMGraphFunctionHeader& InPublicFunction)")
+	static UE_API URigVMEdGraphFunctionRefNodeSpawner* CreateFromAssetData(const FAssetData& InAssetData, const FRigVMGraphFunctionHeader& InPublicFunction) { return nullptr; }
+	UE_DEPRECATED(5.7, "Please use URigVMEdGraphFunctionRefNodeSpawner(const FAssetData& InAssetData, const FRigVMOldPublicFunctionData& InPublicFunction)")
+	static UE_API URigVMEdGraphFunctionRefNodeSpawner* CreateFromAssetData(const FAssetData& InAssetData, const FRigVMOldPublicFunctionData& InPublicFunction) { return nullptr; }
+	UE_API URigVMEdGraphFunctionRefNodeSpawner(const FAssetData& InAssetData, const FRigVMGraphFunctionHeader& InPublicFunction);
+	UE_API URigVMEdGraphFunctionRefNodeSpawner(const FAssetData& InAssetData, const FRigVMOldPublicFunctionData& InPublicFunction);
+
+	// URigVMEdGraphNodeSpawner interface
+	UE_API virtual FString GetSpawnerSignature() const override;
+	UE_API virtual URigVMEdGraphNode* Invoke(URigVMEdGraph* ParentGraph, FVector2D const Location) const override;
+	UE_API virtual bool IsTemplateNodeFilteredOut(TArray<UObject*>& InAssets, TArray<UEdGraph*> InGraphs, TArray<UEdGraphPin*> InPins) const override;
+	// End URigVMEdGraphNodeSpawner interface
 
 private:
 	/** Returns a resolved asset path. This is directly from the asset data if CreateFromAssetData was used to create the spawner, or if
@@ -57,20 +60,15 @@ private:
 	FSoftObjectPath GetResolvedAssetPath() const;
 	
 	/** The public function definition we will spawn from [optional] */
-	UPROPERTY(Transient)
 	mutable FRigVMGraphFunctionHeader ReferencedPublicFunctionHeader;
 
 	/** Marked as true for local function definitions */
-	UPROPERTY(Transient)
 	bool bIsLocalFunction;
 
-	UPROPERTY()
 	mutable FSoftObjectPath AssetPath;
-
-	UPROPERTY()
 	mutable FString FunctionPath;
 
-	static UE_API URigVMEdGraphNode* SpawnNode(UEdGraph* ParentGraph, UBlueprint* Blueprint, FRigVMGraphFunctionHeader& InFunction, FVector2D const Location);
+	static UE_API URigVMEdGraphNode* SpawnNode(UEdGraph* ParentGraph, TRigVMAssetInterfacePtr Blueprint, FRigVMGraphFunctionHeader& InFunction, FVector2D const Location);
 
 	friend class UEngineTestControlRig;
 	friend class FRigVMEditorBase;
diff --git a/Engine/Plugins/Runtime/RigVM/Source/RigVMEditor/Public/EdGraph/NodeSpawners/RigVMEdGraphInvokeEntryNodeSpawner.h b/Engine/Plugins/Runtime/RigVM/Source/RigVMEditor/Public/EdGraph/NodeSpawners/RigVMEdGraphInvokeEntryNodeSpawner.h
--- a/Engine/Plugins/Runtime/RigVM/Source/RigVMEditor/Public/EdGraph/NodeSpawners/RigVMEdGraphInvokeEntryNodeSpawner.h
+++ b/Engine/Plugins/Runtime/RigVM/Source/RigVMEditor/Public/EdGraph/NodeSpawners/RigVMEdGraphInvokeEntryNodeSpawner.h
@@ -5,34 +5,30 @@
 #include "CoreMinimal.h"
 #include "EdGraph/NodeSpawners/RigVMEdGraphNodeSpawner.h"
 #include "RigVMModel/RigVMGraph.h"
-#include "RigVMEdGraphInvokeEntryNodeSpawner.generated.h"
 
 #define UE_API RIGVMEDITOR_API
 
 class URigVMEdGraphNode;
-class URigVMBlueprint;
 
-UCLASS(MinimalAPI, Transient)
 class URigVMEdGraphInvokeEntryNodeSpawner : public URigVMEdGraphNodeSpawner
 {
-	GENERATED_BODY()
-
 public:
+	virtual ~URigVMEdGraphInvokeEntryNodeSpawner() {}
 
-	static UE_API URigVMEdGraphInvokeEntryNodeSpawner* CreateForEntry(URigVMBlueprint* InBlueprint, const FName& InEntryName, const FText& InMenuDesc, const FText& InCategory, const FText& InTooltip);
+	UE_DEPRECATED(5.7, "Plase use URigVMEdGraphInvokeEntryNodeSpawner(TRigVMAssetInterfacePtr InBlueprint, const FName& InEntryName, const FText& InMenuDesc, const FText& InCategory, const FText& InTooltip)")
+	static UE_API URigVMEdGraphInvokeEntryNodeSpawner* CreateForEntry(URigVMBlueprint* InBlueprint, const FName& InEntryName, const FText& InMenuDesc, const FText& InCategory, const FText& InTooltip) { return nullptr; }
+	UE_API URigVMEdGraphInvokeEntryNodeSpawner(TRigVMAssetInterfacePtr InBlueprint, const FName& InEntryName, const FText& InMenuDesc, const FText& InCategory, const FText& InTooltip);
 
-	// UBlueprintNodeSpawner interface
-	UE_API virtual void Prime() override;
-	UE_API virtual FBlueprintNodeSignature GetSpawnerSignature() const override;
-	UE_API virtual FBlueprintActionUiSpec GetUiSpec(FBlueprintActionContext const& Context, FBindingSet const& Bindings) const override;
-	UE_API virtual UEdGraphNode* Invoke(UEdGraph* ParentGraph, FBindingSet const& Bindings, FVector2D const Location) const override;
-	UE_API virtual bool IsTemplateNodeFilteredOut(FBlueprintActionFilter const& Filter) const override;
-	// End UBlueprintNodeSpawner interface
+	// URigVMEdGraphNodeSpawner interface
+	UE_API virtual FString GetSpawnerSignature() const override;
+	UE_API virtual URigVMEdGraphNode* Invoke(URigVMEdGraph* ParentGraph, FVector2D const Location) const override;
+	UE_API virtual bool IsTemplateNodeFilteredOut(TArray<UObject*>& InAssets, TArray<UEdGraph*> InGraphs, TArray<UEdGraphPin*> InPins) const override;
+	// End URigVMEdGraphNodeSpawner interface
 
 private:
 
 	/** The pin type we will spawn */
-	TWeakObjectPtr<URigVMBlueprint> Blueprint;
+	TWeakInterfacePtr<IRigVMAssetInterface> Blueprint;
 	TWeakObjectPtr<URigVMGraph> GraphOwner;
 	FName EntryName;
 
diff --git a/Engine/Plugins/Runtime/RigVM/Source/RigVMEditor/Public/EdGraph/NodeSpawners/RigVMEdGraphNodeSpawner.h b/Engine/Plugins/Runtime/RigVM/Source/RigVMEditor/Public/EdGraph/NodeSpawners/RigVMEdGraphNodeSpawner.h
--- a/Engine/Plugins/Runtime/RigVM/Source/RigVMEditor/Public/EdGraph/NodeSpawners/RigVMEdGraphNodeSpawner.h
+++ b/Engine/Plugins/Runtime/RigVM/Source/RigVMEditor/Public/EdGraph/NodeSpawners/RigVMEdGraphNodeSpawner.h
@@ -3,7 +3,7 @@
 #pragma once
 
 #include "CoreMinimal.h"
-#include "RigVMBlueprint.h"
+#include "RigVMBlueprintLegacy.h"
 #include "BlueprintNodeSpawner.h"
 #include "BlueprintNodeSignature.h"
 #include "BlueprintActionFilter.h"
@@ -13,16 +13,16 @@
 
 class URigVMEdGraphNode;
 
-UCLASS(MinimalAPI, Transient)
-class URigVMEdGraphNodeSpawner : public UBlueprintNodeSpawner
+class URigVMEdGraphNodeSpawner : public TSharedFromThis<URigVMEdGraphNodeSpawner>
 {
-	GENERATED_BODY()
-
+protected:
+	virtual ~URigVMEdGraphNodeSpawner() {}
+	
 public:
 
-	UE_API virtual bool IsTemplateNodeFilteredOut(FBlueprintActionFilter const& Filter) const override;
-
-	UE_API void SetRelatedBlueprintClass(TSubclassOf<URigVMBlueprint> InClass);
+	UE_DEPRECATED(5.7, "Plase use SetRelatedBlueprintClass(UClass* InClass)")
+	UE_API void SetRelatedBlueprintClass(TSubclassOf<URigVMBlueprint> InClass){}
+	UE_API void SetRelatedBlueprintClass(UClass* InClass);
 
 	struct FPinInfo
 	{
@@ -45,12 +45,62 @@ public:
 		, CPPTypeObject(InCPPTypeObject)
 		{}
 	};
+
+	UE_API virtual bool IsTemplateNodeFilteredOut(TArray<UObject*>& InAssets, TArray<UEdGraph*> InGraphs, TArray<UEdGraphPin*> InPins) const;
+	UE_API virtual FString GetSpawnerSignature() const { return FString(); }
+	UE_API virtual URigVMEdGraphNode* Invoke(URigVMEdGraph* ParentGraph, FVector2D const Location) const { return nullptr; }
+	static UE_API URigVMEdGraphNode* SpawnTemplateNode(URigVMEdGraph* InParentGraph, const TArray<FPinInfo>& InPins, const FName& InNodeName = NAME_None);
 	
-	static UE_API URigVMEdGraphNode* SpawnTemplateNode(UEdGraph* InParentGraph, const TArray<FPinInfo>& InPins, const FName& InNodeName = NAME_None);
-	
+protected:
+
+	UClass* RelatedBlueprintClass;
+	UClass* NodeClass;
+
+	// Menu signature
+	FText MenuName;
+	FText MenuTooltip;
+	FText MenuCategory;
+	FText MenuKeywords;
+	FSlateIcon MenuIcon;
+	FLinearColor MenuIconTint = FLinearColor::White;
+
+	friend class URigVMEdGraphNodeBlueprintSpawner;
+};
+
+UCLASS(MinimalAPI, Transient)
+class URigVMEdGraphNodeBlueprintSpawner : public UBlueprintNodeSpawner
+{
+	GENERATED_BODY()
+
+public:
+
+	static URigVMEdGraphNodeBlueprintSpawner* CreateFromRigVMSpawner(TSharedPtr<URigVMEdGraphNodeSpawner> InNodeSpawner)
+	{
+		URigVMEdGraphNodeBlueprintSpawner* NewSpawner = NewObject<URigVMEdGraphNodeBlueprintSpawner>(GetTransientPackage());
+		NewSpawner->RigVMSpawner = InNodeSpawner;
+		NewSpawner->NodeClass = InNodeSpawner->NodeClass;
+		NewSpawner->DefaultMenuSignature.MenuName = InNodeSpawner->MenuName;
+		NewSpawner->DefaultMenuSignature.Tooltip = InNodeSpawner->MenuTooltip;
+		NewSpawner->DefaultMenuSignature.Category = InNodeSpawner->MenuCategory;
+		NewSpawner->DefaultMenuSignature.Keywords = InNodeSpawner->MenuKeywords;
+		NewSpawner->DefaultMenuSignature.Icon = InNodeSpawner->MenuIcon;
+		NewSpawner->DefaultMenuSignature.IconTint = InNodeSpawner->MenuIconTint;
+		return NewSpawner;
+	}
+
+	// UBlueprintNodeSpawner interface
+	UE_API virtual bool IsTemplateNodeFilteredOut(FBlueprintActionFilter const& Filter) const override;
+	UE_API virtual void Prime() override;
+	UE_API virtual FBlueprintNodeSignature GetSpawnerSignature() const override;
+	UE_API virtual FBlueprintActionUiSpec GetUiSpec(FBlueprintActionContext const& Context, FBindingSet const& Bindings) const override;
+	UE_API virtual UEdGraphNode* Invoke(UEdGraph* ParentGraph, FBindingSet const& Bindings, FVector2D const Location) const override;
+	// End UBlueprintNodeSpawner interface
+
 private:
 
-	TSubclassOf<URigVMBlueprint> RelatedBlueprintClass;
+	TSharedPtr<URigVMEdGraphNodeSpawner> RigVMSpawner;
 };
 
+
+
 #undef UE_API
diff --git a/Engine/Plugins/Runtime/RigVM/Source/RigVMEditor/Public/EdGraph/NodeSpawners/RigVMEdGraphTemplateNodeSpawner.h b/Engine/Plugins/Runtime/RigVM/Source/RigVMEditor/Public/EdGraph/NodeSpawners/RigVMEdGraphTemplateNodeSpawner.h
--- a/Engine/Plugins/Runtime/RigVM/Source/RigVMEditor/Public/EdGraph/NodeSpawners/RigVMEdGraphTemplateNodeSpawner.h
+++ b/Engine/Plugins/Runtime/RigVM/Source/RigVMEditor/Public/EdGraph/NodeSpawners/RigVMEdGraphTemplateNodeSpawner.h
@@ -13,33 +13,31 @@
 #include "EdGraph/RigVMEdGraph.h"
 #include "EdGraph/RigVMEdGraphNode.h"
 #include "RigVMModel/Nodes/RigVMTemplateNode.h"
-#include "RigVMEdGraphTemplateNodeSpawner.generated.h"
 
 #define UE_API RIGVMEDITOR_API
 
 class URigVMEdGraphNode;
 
-UCLASS(MinimalAPI, Transient)
 class URigVMEdGraphTemplateNodeSpawner : public URigVMEdGraphNodeSpawner
 {
-	GENERATED_BODY()
-
 public:
+	
+	virtual ~URigVMEdGraphTemplateNodeSpawner() {}
 	/**
 	 * Creates a new URigVMEdGraphVariableNodeSpawner, charged with spawning 
 	 * a new member-variable node
 	 * 
 	 * @return A newly allocated instance of this class.
 	 */
-	static UE_API URigVMEdGraphTemplateNodeSpawner* CreateFromNotation(const FName& InNotation, const FText& InMenuDesc, const FText& InCategory, const FText& InTooltip);
+	UE_DEPRECATED(5.7, "Please use URigVMEdGraphTemplateNodeSpawner(const FName& InNotation, const FText& InMenuDesc, const FText& InCategory, const FText& InTooltip)")
+	static UE_API URigVMEdGraphTemplateNodeSpawner* CreateFromNotation(const FName& InNotation, const FText& InMenuDesc, const FText& InCategory, const FText& InTooltip) { return nullptr; }
+	UE_API URigVMEdGraphTemplateNodeSpawner(const FName& InNotation, const FText& InMenuDesc, const FText& InCategory, const FText& InTooltip);
 
-	// UBlueprintNodeSpawner interface
-	UE_API virtual void Prime() override;
-	UE_API virtual FBlueprintNodeSignature GetSpawnerSignature() const override;
-	UE_API bool IsTemplateNodeFilteredOut(FBlueprintActionFilter const& Filter) const override;
-	UE_API virtual FBlueprintActionUiSpec GetUiSpec(FBlueprintActionContext const& Context, FBindingSet const& Bindings) const override;
-	UE_API virtual UEdGraphNode* Invoke(UEdGraph* ParentGraph, FBindingSet const& Bindings, FVector2D const Location) const override;
-	// End UBlueprintNodeSpawner interface
+	// URigVMEdGraphNodeSpawner interface
+	UE_API virtual FString GetSpawnerSignature() const override;
+	UE_API bool IsTemplateNodeFilteredOut(TArray<UObject*>& InAssets, TArray<UEdGraph*> InGraphs, TArray<UEdGraphPin*> InPins) const override;
+	UE_API virtual URigVMEdGraphNode* Invoke(URigVMEdGraph* ParentGraph, FVector2D const Location) const override;
+	// End URigVMEdGraphNodeSpawner interface
 
 private:
 	/** The unit type we will spawn */
@@ -48,7 +46,7 @@ private:
 
 	const FRigVMTemplate* Template = nullptr;
 
-	static UE_API URigVMEdGraphNode* SpawnNode(UEdGraph* ParentGraph, UBlueprint* Blueprint, const FRigVMTemplate* Template, FVector2D const Location);
+	static UE_API URigVMEdGraphNode* SpawnNode(URigVMEdGraph* ParentGraph, TRigVMAssetInterfacePtr Blueprint, const FRigVMTemplate* Template, FVector2D const Location);
 
 	friend class UEngineTestControlRig;
 	friend class FRigVMEditorBase;
diff --git a/Engine/Plugins/Runtime/RigVM/Source/RigVMEditor/Public/EdGraph/NodeSpawners/RigVMEdGraphUnitNodeSpawner.h b/Engine/Plugins/Runtime/RigVM/Source/RigVMEditor/Public/EdGraph/NodeSpawners/RigVMEdGraphUnitNodeSpawner.h
--- a/Engine/Plugins/Runtime/RigVM/Source/RigVMEditor/Public/EdGraph/NodeSpawners/RigVMEdGraphUnitNodeSpawner.h
+++ b/Engine/Plugins/Runtime/RigVM/Source/RigVMEditor/Public/EdGraph/NodeSpawners/RigVMEdGraphUnitNodeSpawner.h
@@ -10,47 +10,46 @@
 #include "EdGraph/EdGraphPin.h"
 #include "Engine/Blueprint.h"
 #include "EdGraph/NodeSpawners/RigVMEdGraphNodeSpawner.h"
-#include "RigVMBlueprint.h"
+#include "RigVMBlueprintLegacy.h"
 #include "EdGraph/RigVMEdGraph.h"
 #include "EdGraph/RigVMEdGraphNode.h"
 #include "RigVMModel/Nodes/RigVMUnitNode.h"
-#include "RigVMEdGraphUnitNodeSpawner.generated.h"
 
 #define UE_API RIGVMEDITOR_API
 
 class URigVMEdGraphNode;
 
-UCLASS(MinimalAPI, Transient)
 class URigVMEdGraphUnitNodeSpawner : public URigVMEdGraphNodeSpawner
 {
-	GENERATED_BODY()
-
 public:
+	
+	virtual ~URigVMEdGraphUnitNodeSpawner() {}
+	
 	/**
 	 * Creates a new URigVMEdGraphVariableNodeSpawner, charged with spawning 
-	 * a new member-variable node
+	 * a new unit node
 	 * 
 	 * @return A newly allocated instance of this class.
 	 */
-	static UE_API URigVMEdGraphUnitNodeSpawner* CreateFromStruct(UScriptStruct* InStruct, const FName& InMethodName, const FText& InMenuDesc, const FText& InCategory, const FText& InTooltip);
+	UE_DEPRECATED(5.7, "Please use URigVMEdGraphUnitNodeSpawner(UScriptStruct* InStruct, const FName& InMethodName, const FText& InMenuDesc, const FText& InCategory, const FText& InTooltip)")
+	static UE_API URigVMEdGraphUnitNodeSpawner* CreateFromStruct(UScriptStruct* InStruct, const FName& InMethodName, const FText& InMenuDesc, const FText& InCategory, const FText& InTooltip) { return nullptr; }
+	UE_API URigVMEdGraphUnitNodeSpawner(UScriptStruct* InStruct, const FName& InMethodName, const FText& InMenuDesc, const FText& InCategory, const FText& InTooltip);
 
-	// UBlueprintNodeSpawner interface
-	UE_API virtual void Prime() override;
-	UE_API virtual FBlueprintNodeSignature GetSpawnerSignature() const override;
-	UE_API virtual FBlueprintActionUiSpec GetUiSpec(FBlueprintActionContext const& Context, FBindingSet const& Bindings) const override;
-	UE_API virtual UEdGraphNode* Invoke(UEdGraph* ParentGraph, FBindingSet const& Bindings, FVector2D const Location) const override;
-	static UE_API void HookupMutableNode(URigVMNode* InModelNode, URigVMBlueprint* InRigBlueprint);
-	// End UBlueprintNodeSpawner interface
+	// URigVMEdGraphNodeSpawner interface
+	UE_API virtual FString GetSpawnerSignature() const override;
+	UE_API virtual URigVMEdGraphNode* Invoke(URigVMEdGraph* ParentGraph, FVector2D const Location) const override;
+	// URigVMEdGraphNodeSpawner interface
+	
+	static UE_API void HookupMutableNode(URigVMNode* InModelNode, TRigVMAssetInterfacePtr InRigBlueprint);
 
 private:
 	/** The unit type we will spawn */
-	UPROPERTY(Transient)
 	TObjectPtr<UScriptStruct> StructTemplate;
 
-	UPROPERTY(Transient)
 	FName MethodName;
 
-	static UE_API URigVMEdGraphNode* SpawnNode(UEdGraph* ParentGraph, UBlueprint* Blueprint, UScriptStruct* StructTemplate, const FName& InMethodName, FVector2D const Location);
+	static UE_API URigVMEdGraphNode* SpawnNode(URigVMEdGraph* ParentGraph, TRigVMAssetInterfacePtr Blueprint, UScriptStruct* StructTemplate, const
+	                                           FName& InMethodName, FVector2D Location);
 
 	friend class UEngineTestControlRig;
 	friend class FRigVMEditorBase;
diff --git a/Engine/Plugins/Runtime/RigVM/Source/RigVMEditor/Public/EdGraph/NodeSpawners/RigVMEdGraphVariableNodeSpawner.h b/Engine/Plugins/Runtime/RigVM/Source/RigVMEditor/Public/EdGraph/NodeSpawners/RigVMEdGraphVariableNodeSpawner.h
--- a/Engine/Plugins/Runtime/RigVM/Source/RigVMEditor/Public/EdGraph/NodeSpawners/RigVMEdGraphVariableNodeSpawner.h
+++ b/Engine/Plugins/Runtime/RigVM/Source/RigVMEditor/Public/EdGraph/NodeSpawners/RigVMEdGraphVariableNodeSpawner.h
@@ -12,19 +12,17 @@
 #include "EdGraph/NodeSpawners/RigVMEdGraphNodeSpawner.h"
 #include "RigVMCore/RigVMExternalVariable.h"
 #include "RigVMModel/RigVMVariableDescription.h"
-#include "RigVMEdGraphVariableNodeSpawner.generated.h"
 
 #define UE_API RIGVMEDITOR_API
 
 class URigVMEdGraphNode;
-class URigVMBlueprint;
 
-UCLASS(MinimalAPI, Transient)
 class URigVMEdGraphVariableNodeSpawner : public URigVMEdGraphNodeSpawner
 {
-	GENERATED_BODY()
 
 public:
+	
+	virtual ~URigVMEdGraphVariableNodeSpawner() {}
 
 	/**
 	 * Creates a new URigVMEdGraphVariableNodeSpawner, charged with spawning 
@@ -32,21 +30,24 @@ public:
 	 * 
 	 * @return A newly allocated instance of this class.
 	 */
-	static UE_API URigVMEdGraphVariableNodeSpawner* CreateFromExternalVariable(URigVMBlueprint* InBlueprint, const FRigVMExternalVariable& InExternalVariable, bool bInIsGetter, const FText& InMenuDesc, const FText& InCategory, const FText& InTooltip);
-	static UE_API URigVMEdGraphVariableNodeSpawner* CreateFromLocalVariable(URigVMBlueprint* InBlueprint, URigVMGraph* InGraphOwner, const FRigVMGraphVariableDescription& InLocalVariable, bool bInIsGetter, const FText& InMenuDesc, const FText& InCategory, const FText& InTooltip);
+	UE_DEPRECATED(5.7, "Please use URigVMEdGraphVariableNodeSpawner(TRigVMAssetInterfacePtr InBlueprint, const FRigVMExternalVariable& InExternalVariable, bool bInIsGetter, const FText& InMenuDesc, const FText& InCategory, const FText& InTooltip)")
+	static UE_API URigVMEdGraphVariableNodeSpawner* CreateFromExternalVariable(URigVMBlueprint* InBlueprint, const FRigVMExternalVariable& InExternalVariable, bool bInIsGetter, const FText& InMenuDesc, const FText& InCategory, const FText& InTooltip) { return nullptr; }
+	UE_DEPRECATED(5.7, "Please use URigVMEdGraphVariableNodeSpawner(TRigVMAssetInterfacePtr InBlueprint, URigVMGraph* InGraphOwner, const FRigVMGraphVariableDescription& InLocalVariable, bool bInIsGetter, const FText& InMenuDesc, const FText& InCategory, const FText& InTooltip)")
+	static UE_API URigVMEdGraphVariableNodeSpawner* CreateFromLocalVariable(URigVMBlueprint* InBlueprint, URigVMGraph* InGraphOwner, const FRigVMGraphVariableDescription& InLocalVariable, bool bInIsGetter, const FText& InMenuDesc, const FText& InCategory, const FText& InTooltip) { return nullptr; }
 
-	// UBlueprintNodeSpawner interface
-	UE_API virtual void Prime() override;
-	UE_API virtual FBlueprintNodeSignature GetSpawnerSignature() const override;
-	UE_API virtual FBlueprintActionUiSpec GetUiSpec(FBlueprintActionContext const& Context, FBindingSet const& Bindings) const override;
-	UE_API virtual UEdGraphNode* Invoke(UEdGraph* ParentGraph, FBindingSet const& Bindings, FVector2D const Location) const override;
-	UE_API virtual bool IsTemplateNodeFilteredOut(FBlueprintActionFilter const& Filter) const override;
-	// End UBlueprintNodeSpawner interface
+	UE_API URigVMEdGraphVariableNodeSpawner(TRigVMAssetInterfacePtr InBlueprint, const FRigVMExternalVariable& InExternalVariable, bool bInIsGetter, const FText& InMenuDesc, const FText& InCategory, const FText& InTooltip);
+	UE_API URigVMEdGraphVariableNodeSpawner(TRigVMAssetInterfacePtr InBlueprint, URigVMGraph* InGraphOwner, const FRigVMGraphVariableDescription& InLocalVariable, bool bInIsGetter, const FText& InMenuDesc, const FText& InCategory, const FText& InTooltip);
+
+	// URigVMEdGraphNodeSpawner interface
+	UE_API virtual FString GetSpawnerSignature() const override;
+	UE_API virtual URigVMEdGraphNode* Invoke(URigVMEdGraph* ParentGraph, FVector2D Location) const override;
+	UE_API virtual bool IsTemplateNodeFilteredOut(TArray<UObject*>& InAssets, TArray<UEdGraph*> InGraphs, TArray<UEdGraphPin*> InPins) const override;
+	// End URigVMEdGraphNodeSpawner interface
 
 private:
 
 	/** The pin type we will spawn */
-	TWeakObjectPtr<URigVMBlueprint> Blueprint;
+	TWeakInterfacePtr<IRigVMAssetInterface> Blueprint;
 	TWeakObjectPtr<URigVMGraph> GraphOwner;
 	FRigVMExternalVariable ExternalVariable;
 	bool bIsGetter;
diff --git a/Engine/Plugins/Runtime/RigVM/Source/RigVMEditor/Public/Editor/RigVMCommentNodeDetailCustomization.h b/Engine/Plugins/Runtime/RigVM/Source/RigVMEditor/Public/Editor/RigVMCommentNodeDetailCustomization.h
--- a/Engine/Plugins/Runtime/RigVM/Source/RigVMEditor/Public/Editor/RigVMCommentNodeDetailCustomization.h
+++ b/Engine/Plugins/Runtime/RigVM/Source/RigVMEditor/Public/Editor/RigVMCommentNodeDetailCustomization.h
@@ -6,7 +6,7 @@
 #include "IDetailCustomization.h"
 #include "IPropertyTypeCustomization.h"
 #include "RigVMHost.h"
-#include "RigVMBlueprint.h"
+#include "RigVMAsset.h"
 #include "Styling/SlateTypes.h"
 #include "IPropertyUtilities.h"
 
@@ -35,7 +35,7 @@ public:
 
 private:
 
-	URigVMBlueprint* BlueprintBeingCustomized;
+	IRigVMAssetInterface* BlueprintBeingCustomized;
 	TArray<TWeakObjectPtr<URigVMCommentNode>> ObjectsBeingCustomized;
 
 	UE_API void GetValuesFromNode(TWeakObjectPtr<URigVMCommentNode> Node);
diff --git a/Engine/Plugins/Runtime/RigVM/Source/RigVMEditor/Public/Editor/RigVMEditor.h b/Engine/Plugins/Runtime/RigVM/Source/RigVMEditor/Public/Editor/RigVMEditor.h
--- a/Engine/Plugins/Runtime/RigVM/Source/RigVMEditor/Public/Editor/RigVMEditor.h
+++ b/Engine/Plugins/Runtime/RigVM/Source/RigVMEditor/Public/Editor/RigVMEditor.h
@@ -25,6 +25,7 @@
 
 #define LOCTEXT_NAMESPACE "RigVMEditor"
 
+class IRigVMAssetInterface;
 class FDocumentTracker;
 class SRigVMEditorGraphExplorer;
 class URigVM;
@@ -33,11 +34,10 @@ class FRigVMEditorBase;
 class URigVMController;
 class URigVMHost;
 class FTransaction;
-class URigVMBlueprint;
 class SGraphEditor;
 class FPreviewScene;
 
-DECLARE_MULTICAST_DELEGATE_TwoParams(FRigVMEditorClosed, const IRigVMEditor*, URigVMBlueprint*);
+DECLARE_MULTICAST_DELEGATE_TwoParams(FRigVMEditorClosed, const IRigVMEditor*, TRigVMAssetInterfacePtr);
 
 
 /* Enums to use when grouping the blueprint members in the list panel. The order here will determine the order in the list */
@@ -104,7 +104,9 @@ class IRigVMEditor
 public:
 
 	virtual TSharedPtr<FAssetEditorToolkit> GetHostingApp() = 0;
-	virtual URigVMBlueprint* GetRigVMBlueprint() const = 0;
+	UE_DEPRECATED(5.7, "Plase use GetRigVMAsset")
+	virtual URigVMBlueprint* GetRigVMBlueprint() const { return Cast<URigVMBlueprint>(GetRigVMAsset().GetObject()); }
+	virtual TRigVMAssetInterfacePtr GetRigVMAsset() const = 0;
 	virtual const TArray< UObject* >* GetObjectsCurrentlyBeingEdited() const = 0;
 	virtual URigVMHost* GetRigVMHost() const = 0;
 
@@ -184,7 +186,7 @@ protected:
 	virtual const TSharedPtr<FAssetEditorToolkit> GetHostingApp() const = 0;
 	virtual void InitAssetEditor(const EToolkitMode::Type Mode, const TSharedPtr<IToolkitHost>& InitToolkitHost, const FName AppIdentifier, const TSharedRef<FTabManager::FLayout>& StandaloneDefaultLayout, const bool bCreateDefaultStandaloneMenu, const bool bCreateDefaultToolbar, const TArray<UObject*>& ObjectsToEdit, const bool bInIsToolbarFocusable = false, const bool bInUseSmallToolbarIcons = false, const TOptional<EAssetOpenMethod>& InOpenMethod = TOptional<EAssetOpenMethod>()) = 0;
 	virtual void CreateEditorToolbar() = 0;
-	virtual void CommonInitialization(const TArray<UBlueprint*>& InitBlueprints, bool bShouldOpenInDefaultsMode) = 0;
+	virtual void CommonInitialization(const TArray<TRigVMAssetInterfacePtr>& InitBlueprints, bool bShouldOpenInDefaultsMode) = 0;
 	virtual TSharedPtr<FDocumentTracker> GetDocumentManager() const = 0;
 	virtual void AddApplicationMode(FName ModeName, TSharedRef<FApplicationMode> Mode) = 0;
 	virtual void RegenerateMenusAndToolbars() = 0;
@@ -207,7 +209,7 @@ protected:
 	virtual void RenameNewlyAddedAction(FName InActionName) = 0;
 	virtual FGraphPanelSelectionSet GetSelectedNodes() const = 0;
 	virtual void SetUISelectionState(FName SelectionOwner) = 0;
-	virtual void AnalyticsTrackNodeEvent(UBlueprint* Blueprint, UEdGraphNode *GraphNode, bool bNodeDelete = false) const = 0;
+	virtual void AnalyticsTrackNodeEvent(IRigVMAssetInterface* Blueprint, UEdGraphNode *GraphNode, bool bNodeDelete = false) const = 0;
 	virtual void PostUndo(bool bSuccess) = 0;
 	virtual void PostRedo(bool bSuccess) = 0;
 	virtual UEdGraphPin* GetCurrentlySelectedPin() const = 0;
@@ -225,7 +227,7 @@ protected:
 	virtual void AddMenuExtender(TSharedPtr<FExtender> Extender) = 0;
 	virtual void RemoveMenuExtender(TSharedPtr<FExtender> Extender) = 0;
 	virtual TSharedPtr<class IMessageLogListing> GetCompilerResultsListing() = 0;
-	virtual void OnBlueprintChangedInnerImpl(UBlueprint* InBlueprint, bool bIsJustBeingCompiled = false) = 0;
+	virtual void OnBlueprintChangedInnerImpl(IRigVMAssetInterface* InBlueprint, bool bIsJustBeingCompiled = false) = 0;
 	virtual void RefreshEditorsImpl(ERefreshRigVMEditorReason::Type Reason = ERefreshRigVMEditorReason::UnknownReason) = 0;
 	virtual void SetupGraphEditorEventsImpl(UEdGraph* InGraph, SGraphEditor::FGraphEditorEvents& InEvents) = 0;
 	virtual FActionMenuContent OnCreateGraphActionMenu(UEdGraph* InGraph, const FVector2f& InNodePosition, const TArray<UEdGraphPin*>& InDraggedPins, bool bAutoExpand, SGraphEditor::FActionMenuClosed InOnMenuClosed) = 0;
@@ -260,11 +262,13 @@ public:
 	 * @param	InitToolkitHost			When Mode is WorldCentric, this is the level editor instance to spawn this editor within
 	 * @param	InRigVMBlueprint	The blueprint object to start editing.
 	 */
-	UE_API virtual void InitRigVMEditor(const EToolkitMode::Type Mode, const TSharedPtr< class IToolkitHost >& InitToolkitHost, class URigVMBlueprint* InRigVMBlueprint);
+	UE_API virtual void InitRigVMEditor(const EToolkitMode::Type Mode, const TSharedPtr< class IToolkitHost >& InitToolkitHost, TRigVMAssetInterfacePtr InRigVMBlueprint);
 
 	static UE_API FRigVMEditorBase* GetFromAssetEditorInstance(IAssetEditorInstance* Instance);
 	// returns the blueprint being edited
+	UE_DEPRECATED(5.7, "Please use GetRigVMAsset")
 	UE_API virtual URigVMBlueprint* GetRigVMBlueprint() const override;
+	UE_API virtual TRigVMAssetInterfacePtr GetRigVMAsset() const override;
 	
 	UE_API void HandleJumpToHyperlink(const UObject* InSubject);
 
@@ -331,7 +335,6 @@ protected:
 	UE_API virtual const FName GetEditorModeName() const;
 
 	// FBlueprintEditor interface
-	UE_API UBlueprint* GetBlueprintObj() const;
 	UE_API virtual bool InEditingMode() const override;
 	UE_API TSubclassOf<UEdGraphSchema> GetDefaultSchemaClass() const;
 	UE_API void OnGraphEditorFocused(const TSharedRef<SGraphEditor>& InGraphEditor);
@@ -415,14 +418,11 @@ protected:
 	bool IsInAScriptingMode() const { return true; }
 	UE_API virtual void OnNewDocumentClicked(ECreatedDocumentType GraphType) override;
 	UE_API bool IsSectionVisibleImpl(RigVMNodeSectionID::Type InSectionID) const;
-	UE_API bool AreEventGraphsAllowed() const;
-	UE_API bool AreMacrosAllowed() const;
-	UE_API bool AreDelegatesAllowed() const;
 	UE_API bool NewDocument_IsVisibleForTypeImpl(ECreatedDocumentType GraphType) const;
 	UE_API FGraphAppearanceInfo GetGraphAppearance(class UEdGraph* InGraph) const;
 	UE_API void OnActiveTabChanged( TSharedPtr<SDockTab> PreviouslyActive, TSharedPtr<SDockTab> NewlyActivated );
 	UE_API void OnSelectedNodesChangedImpl(const TSet<class UObject*>& NewSelection);
-	UE_API void OnBlueprintChangedImpl(UBlueprint* InBlueprint, bool bIsJustBeingCompiled);
+	UE_API void OnBlueprintChangedImpl(IRigVMAssetInterface* InBlueprint, bool bIsJustBeingCompiled);
 	UE_API virtual void ForceEditorRefresh(ERefreshRigVMEditorReason::Type Reason = ERefreshRigVMEditorReason::UnknownReason) override;
 	UE_API void SetupGraphEditorEvents(UEdGraph* InGraph, SGraphEditor::FGraphEditorEvents& InEvents);
 	UE_API void FocusInspectorOnGraphSelection(const TSet<class UObject*>& NewSelection, bool bForceRefresh = false);
@@ -446,11 +446,11 @@ protected:
 	UE_API void OnPropertyChanged(UObject* InObject, FPropertyChangedEvent& InEvent);
 	UE_API virtual void OnWrappedPropertyChangedChainEvent(URigVMDetailsViewWrapperObject* InWrapperObject, const FString& InPropertyPath, FPropertyChangedChainEvent& InPropertyChangedChainEvent);
 	UE_API void OnRequestLocalizeFunctionDialog(FRigVMGraphFunctionIdentifier& InFunction, URigVMController* InTargetController, IRigVMGraphFunctionHost* InTargetFunctionHost, bool bForce);
-	UE_API FRigVMController_BulkEditResult OnRequestBulkEditDialog(URigVMBlueprint* InBlueprint, URigVMController* InController, URigVMLibraryNode* InFunction, ERigVMControllerBulkEditType InEditType);
+	UE_API FRigVMController_BulkEditResult OnRequestBulkEditDialog(TRigVMAssetInterfacePtr InBlueprint, URigVMController* InController, URigVMLibraryNode* InFunction, ERigVMControllerBulkEditType InEditType);
 	UE_API bool OnRequestBreakLinksDialog(TArray<URigVMLink*> InLinks);
 	UE_API TRigVMTypeIndex OnRequestPinTypeSelectionDialog(const TArray<TRigVMTypeIndex>& InTypes);
 
-	UE_API bool UpdateDefaultValueForVariable(FBPVariableDescription& InVariable, bool bUseCDO);
+	UE_API bool UpdateDefaultValueForVariable(FRigVMGraphVariableDescription& InVariable, bool bUseDebuggedObject);
 
 	URigVMController* ActiveController;
 
@@ -542,18 +542,14 @@ protected:
 	UE_API void HandleSetObjectBeingDebugged(UObject* InObject);
 	UE_API virtual FString GetCustomDebugObjectLabel(UObject* ObjectBeingDebugged) const override;
 
-	/** Handle hiding items in the graph */
-	UE_API void HandleHideItem();
-	UE_API bool CanHideItem() const;
-
 	/** Update stale watch pins */
 	UE_API void UpdateStaleWatchedPins();
 	
-	UE_API virtual void HandleRefreshEditorFromBlueprint(URigVMBlueprint* InBlueprint);
+	UE_API virtual void HandleRefreshEditorFromBlueprint(TRigVMAssetInterfacePtr InBlueprint);
 	UE_API void HandleVariableDroppedFromBlueprint(UObject* InSubject, FProperty* InVariableToDrop, const FVector2D& InDropPosition, const FVector2D& InScreenPosition);
 	UE_API void HandleBreakpointAdded();
 	UE_API void OnGraphNodeClicked(URigVMEdGraphNode* InNode);
-	UE_API void OnNodeDoubleClicked(URigVMBlueprint* InBlueprint, URigVMNode* InNode);
+	UE_API void OnNodeDoubleClicked(TRigVMAssetInterfacePtr InBlueprint, URigVMNode* InNode);
 	UE_API void OnGraphImported(UEdGraph* InEdGraph);
 	UE_API bool OnActionMatchesName(FEdGraphSchemaAction* InAction, const FName& InName) const;
 	UE_API void HandleShowCurrentStatement();
diff --git a/Engine/Plugins/Runtime/RigVM/Source/RigVMEditor/Public/Editor/RigVMEditorMenuContext.h b/Engine/Plugins/Runtime/RigVM/Source/RigVMEditor/Public/Editor/RigVMEditorMenuContext.h
--- a/Engine/Plugins/Runtime/RigVM/Source/RigVMEditor/Public/Editor/RigVMEditorMenuContext.h
+++ b/Engine/Plugins/Runtime/RigVM/Source/RigVMEditor/Public/Editor/RigVMEditorMenuContext.h
@@ -3,6 +3,7 @@
 #pragma once
 
 #include "CoreMinimal.h"
+#include "RigVMAsset.h"
 #include "EdGraph/EdGraph.h"
 
 #include "RigVMEditorMenuContext.generated.h"
@@ -10,7 +11,6 @@
 #define UE_API RIGVMEDITOR_API
 
 class URigVMHost;
-class URigVMBlueprint;
 class IRigVMEditor;
 class URigVMGraph;
 class URigVMNode;
@@ -62,8 +62,11 @@ public:
 	UE_API void Init(TWeakPtr<IRigVMEditor> InRigVMEditor, const FRigVMEditorGraphMenuContext& InGraphMenuContext = FRigVMEditorGraphMenuContext());
 
 	/** Get the rigvm blueprint that we are editing */
-	UFUNCTION(BlueprintCallable, Category = RigVMEditor)
+	UFUNCTION(BlueprintCallable, Category = RigVMEditor, meta=(DeprecatedFunction, DeprecationMessage="GetRigVMBlueprint is deprecated. Please use GetRigVMAssetInterface instead." ))
     UE_API URigVMBlueprint* GetRigVMBlueprint() const;
+
+	UFUNCTION(BlueprintCallable, Category = RigVMEditor)
+	UE_API TScriptInterface<IRigVMAssetInterface> GetRigVMAsset() const;
 	
 	/** Get the active rigvm host instance in the viewport */
 	UFUNCTION(BlueprintCallable, Category = RigVMEditor)
diff --git a/Engine/Plugins/Runtime/RigVM/Source/RigVMEditor/Public/Editor/RigVMFindReferences.h b/Engine/Plugins/Runtime/RigVM/Source/RigVMEditor/Public/Editor/RigVMFindReferences.h
--- a/Engine/Plugins/Runtime/RigVM/Source/RigVMEditor/Public/Editor/RigVMFindReferences.h
+++ b/Engine/Plugins/Runtime/RigVM/Source/RigVMEditor/Public/Editor/RigVMFindReferences.h
@@ -3,6 +3,7 @@
 #pragma once
 
 #include "CoreMinimal.h"
+#include "RigVMAsset.h"
 #include "Input/Reply.h"
 #include "Widgets/Views/STreeView.h"
 #include "EdGraph/EdGraphPin.h"
@@ -11,11 +12,11 @@
 
 #define UE_API RIGVMEDITOR_API
 
+class IRigVMAssetInterface;
 class SWidget;
 class SDockTab;
 class UBlueprint;
 class FUICommandList;
-class URigVMBlueprint;
 
 
 class FRigVMFindResult : public TSharedFromThis< FRigVMFindResult >
@@ -25,8 +26,12 @@ public:
 	virtual ~FRigVMFindResult() = default;
 
 	/* Create a root */
+	UE_DEPRECATED(5.7, "Please use FRigVMFindResult(TRigVMAssetInterfacePtr InBlueprint)")
 	UE_API explicit FRigVMFindResult(TWeakObjectPtr<URigVMBlueprint> InBlueprint);
+	UE_DEPRECATED(5.7, "Please use FRigVMFindResult(TRigVMAssetInterfacePtr InBlueprint, const FText& InDisplayText)")
 	UE_API explicit FRigVMFindResult(TWeakObjectPtr<URigVMBlueprint> InBlueprint, const FText& InDisplayText);
+	UE_API explicit FRigVMFindResult(TWeakInterfacePtr<IRigVMAssetInterface> InBlueprint);
+	UE_API explicit FRigVMFindResult(TWeakInterfacePtr<IRigVMAssetInterface> InBlueprint, const FText& InDisplayText);
 
 	/* Called when user clicks on the search item */
 	UE_API virtual FReply OnClick();
@@ -41,7 +46,10 @@ public:
 	virtual void FinalizeSearchData() {};
 
 	/** gets the blueprint housing all these search results */
+	UE_DEPRECATED(5.7, "Please use TRigVMAssetInterfacePtr GetRigVMAsset()")
 	UE_API URigVMBlueprint* GetBlueprint() const;
+	UE_API TRigVMAssetInterfacePtr GetRigVMAsset();
+	UE_API const TRigVMAssetInterfacePtr GetRigVMAsset() const;
 
 	/**
 	* Parses search info for specific data important for displaying the search result in an easy to understand format
@@ -52,7 +60,7 @@ public:
 	virtual void ParseSearchInfo(FText InKey, FText InValue) {};
 
 	/** Returns the Object represented by this search information give the Blueprint it can be found in */
-	UE_API virtual UObject* GetObject(URigVMBlueprint* InBlueprint) const;
+	UE_API virtual UObject* GetObject(TRigVMAssetInterfacePtr InBlueprint) const;
 
 	/** Returns the display string for the row */
 	UE_API FText GetDisplayString() const;
@@ -65,7 +73,7 @@ public:
 	TWeakPtr<FRigVMFindResult> Parent;
 
 	/*If it exists it is the blueprint*/
-	TWeakObjectPtr<URigVMBlueprint> WeakBlueprint;
+	TWeakInterfacePtr<IRigVMAssetInterface> WeakBlueprint;
 
 	/*The display text for this item */
 	FText DisplayText;
@@ -106,7 +114,9 @@ namespace RigVMFindReferencesHelpers
 class FRigVMFindReferencesGraphNode : public FRigVMFindResult
 {
 public:
+	UE_DEPRECATED(5.7, "Please use FRigVMFindReferencesGraphNode(TRigVMAssetInterfacePtr InBlueprint)")
 	FRigVMFindReferencesGraphNode(TWeakObjectPtr<URigVMBlueprint> InBlueprint);
+	FRigVMFindReferencesGraphNode(TRigVMAssetInterfacePtr InBlueprint);
 	virtual ~FRigVMFindReferencesGraphNode() {}
 
 	/** FRigVMFindResult Interface */
@@ -114,7 +124,7 @@ public:
 	virtual TSharedRef<SWidget>	CreateIcon() const override;
 	virtual void ParseSearchInfo(FText InKey, FText InValue) override;
 	virtual void FinalizeSearchData() override;
-	virtual UObject* GetObject(URigVMBlueprint* InBlueprint) const override;
+	virtual UObject* GetObject(TRigVMAssetInterfacePtr InBlueprint) const override;
 	virtual FText GetCategory() const override;
 	/** End FRigVMFindResult Interface */
 
@@ -139,7 +149,9 @@ private:
 class FRigVMFindReferencesPin : public FRigVMFindResult
 {
 public:
+	UE_DEPRECATED(5.7, "Plase use FRigVMFindReferencesPin(TRigVMAssetInterfacePtr InBlueprint, FString InSchemaName)")
 	FRigVMFindReferencesPin(TWeakObjectPtr<URigVMBlueprint> InBlueprint, FString InSchemaName);
+	FRigVMFindReferencesPin(TRigVMAssetInterfacePtr InBlueprint, FString InSchemaName);
 	virtual ~FRigVMFindReferencesPin() {}
 
 	/** FRigVMFindResult Interface */
@@ -164,7 +176,9 @@ private:
 class FRigVMFindReferencesVariable : public FRigVMFindResult
 {
 public:
+	UE_DEPRECATED(5.7, "Plase use FRigVMFindReferencesVariable(TRigVMAssetInterfacePtr InBlueprint)")
 	FRigVMFindReferencesVariable(TWeakObjectPtr<URigVMBlueprint> InBlueprint);
+	FRigVMFindReferencesVariable(TRigVMAssetInterfacePtr InBlueprint);
 	virtual ~FRigVMFindReferencesVariable() {}
 
 	/** FRigVMFindResult Interface */
@@ -186,7 +200,9 @@ private:
 class FRigVMFindReferencesGraph : public FRigVMFindResult
 {
 public:
+	UE_DEPRECATED(5.7, "Please use FRigVMFindReferencesGraph(TRigVMAssetInterfacePtr InBlueprint, EGraphType InGraphType)")
 	FRigVMFindReferencesGraph(TWeakObjectPtr<URigVMBlueprint> InBlueprint, EGraphType InGraphType);
+	FRigVMFindReferencesGraph(TRigVMAssetInterfacePtr InBlueprint, EGraphType InGraphType);
 	virtual ~FRigVMFindReferencesGraph() {}
 
 	/** FRigVMFindResult Interface */
