RigVM: Adding IRigVMAssetInterface as parent of RigVMBlueprint. First step to separation of RigVM and Blueprint.

diff --git a/Engine/Plugins/Runtime/RigVM/Source/RigVMEditor/Public/Editor/RigVMGraphDetailCustomization.h b/Engine/Plugins/Runtime/RigVM/Source/RigVMEditor/Public/Editor/RigVMGraphDetailCustomization.h
--- a/Engine/Plugins/Runtime/RigVM/Source/RigVMEditor/Public/Editor/RigVMGraphDetailCustomization.h
+++ b/Engine/Plugins/Runtime/RigVM/Source/RigVMEditor/Public/Editor/RigVMGraphDetailCustomization.h
@@ -7,7 +7,7 @@
 #include "IDetailCustomNodeBuilder.h"
 #include "UObject/WeakObjectPtr.h"
 #include "EdGraph/RigVMEdGraph.h"
-#include "RigVMBlueprint.h"
+#include "RigVMBlueprintLegacy.h"
 #include "Editor/RigVMNewEditor.h"
 #include "Editor/RigVMLegacyEditor.h"
 #include "SGraphPin.h"
@@ -120,7 +120,7 @@ public:
 		, GraphPtr(InGraph)
 		, WeakRigVMClientHost(InRigVMClientHost)
 		, RigVMEditorPtr(InEditor)
-		, NameValidator(nullptr, InGraph.Get(), InPin->GetFName()) // TODO check if missing Blueprint affects name validation
+		, NameValidator((const TRigVMAssetInterfacePtr)nullptr, InGraph.Get(), InPin->GetFName()) // TODO check if missing Blueprint affects name validation
 	{}
 
 private:
@@ -212,10 +212,10 @@ class FRigVMGraphDetailCustomization : public IDetailCustomization
 public:
 	/** Makes a new instance of this detail layout class for a specific detail view requesting it */
 	static UE_API TSharedRef<IDetailCustomization> MakeInstance(TSharedPtr<IRigVMEditor> InBlueprintEditor, const UClass* InExpectedBlueprintClass);
-	UE_API FRigVMGraphDetailCustomization(TSharedPtr<IRigVMEditor> RigVMigEditor, URigVMBlueprint* RigVMBlueprint);
+	UE_API FRigVMGraphDetailCustomization(TSharedPtr<IRigVMEditor> RigVMigEditor, TRigVMAssetInterfacePtr RigVMBlueprint);
 #if WITH_RIGVMLEGACYEDITOR
 	static UE_API TSharedPtr<IDetailCustomization> MakeLegacyInstance(TSharedPtr<IBlueprintEditor> InBlueprintEditor, const UClass* InExpectedBlueprintClass);
-	UE_API FRigVMGraphDetailCustomization(TSharedPtr<IBlueprintEditor> RigVMigEditor, URigVMBlueprint* RigVMBlueprint);
+	UE_API FRigVMGraphDetailCustomization(TSharedPtr<IBlueprintEditor> RigVMigEditor, TRigVMAssetInterfacePtr RigVMBlueprint);
 #endif
 
 
@@ -230,7 +230,7 @@ private:
 	TWeakObjectPtr<URigVMEdGraph> GraphPtr;
 
 	/** The blueprint we are editing */
-	TWeakObjectPtr<URigVMBlueprint> RigVMBlueprintPtr;
+	TWeakInterfacePtr<IRigVMAssetInterface> RigVMBlueprintPtr;
 
 	TSharedPtr<FRigVMGraphDetailCustomizationImpl> RigVMGraphDetailCustomizationImpl;
 };
@@ -343,7 +343,7 @@ public:
 	UE_API void OnNameListComboBox(const FNameProperty* InProperty, const TArray<TSharedPtr<FRigVMStringWithTag>>* InNameList);
 	UE_API void CustomizeLiveValues(IDetailLayoutBuilder& DetailLayout);
 
-	URigVMBlueprint* BlueprintBeingCustomized;
+	TRigVMAssetInterfacePtr BlueprintBeingCustomized;
 	TArray<TWeakObjectPtr<URigVMDetailsViewWrapperObject>> ObjectsBeingCustomized;
 	TArray<TWeakObjectPtr<URigVMNode>> NodesBeingCustomized;
 	TMap<FName, TSharedPtr<SRigVMGraphPinNameListValueWidget>> NameListWidgets;
@@ -472,7 +472,7 @@ protected:
 
 	UE_API void HandleControlEnumChanged(TSharedPtr<FString> InEnumPath, ESelectInfo::Type InSelectType, TSharedRef<IPropertyHandle> InPropertyHandle);
 
-	URigVMBlueprint* BlueprintBeingCustomized;
+	TRigVMAssetInterfacePtr BlueprintBeingCustomized;
 	URigVMGraph* GraphBeingCustomized;
 	TArray<TWeakObjectPtr<UObject>> ObjectsBeingCustomized;
 	TArray<TSharedPtr<FStructOnScope>> StructsBeingCustomized;
@@ -610,7 +610,7 @@ protected:
 	void OnVectorComponentChanged(TSharedRef<class IPropertyHandle> InPropertyHandle, int32 InComponent, NumericType InValue, bool bIsCommit, ETextCommit::Type InCommitType = ETextCommit::Default)
 	{
 		URigVMController* Controller = nullptr;
-		if (TStrongObjectPtr<URigVMBlueprint> Blueprint = BlueprintBeingCustomized.Pin(); GraphBeingCustomized.IsValid())
+		if (TStrongObjectPtr<UObject> Blueprint = BlueprintBeingCustomized.GetWeakObjectPtr().Pin(); GraphBeingCustomized.IsValid())
 		{
 			Controller = BlueprintBeingCustomized->GetController(GraphBeingCustomized.Get());
 			if (bIsCommit)
@@ -691,7 +691,7 @@ protected:
         }
 	
 		URigVMController* Controller = nullptr;
-		if (TStrongObjectPtr<URigVMBlueprint> Blueprint = BlueprintBeingCustomized.Pin(); GraphBeingCustomized.IsValid())
+		if (TStrongObjectPtr<UObject> Blueprint = BlueprintBeingCustomized.GetWeakObjectPtr().Pin(); GraphBeingCustomized.IsValid())
 		{
 			Controller = BlueprintBeingCustomized->GetController(GraphBeingCustomized.Get());
 			if(bIsCommit)
@@ -733,7 +733,7 @@ protected:
 		return TEXT("Scale3D");
 	}
 
-	TWeakObjectPtr<URigVMBlueprint> BlueprintBeingCustomized;
+	TWeakInterfacePtr<IRigVMAssetInterface> BlueprintBeingCustomized;
 	TWeakObjectPtr<URigVMGraph> GraphBeingCustomized;
 };
 
diff --git a/Engine/Plugins/Runtime/RigVM/Source/RigVMEditor/Public/Editor/RigVMLegacyEditor.h b/Engine/Plugins/Runtime/RigVM/Source/RigVMEditor/Public/Editor/RigVMLegacyEditor.h
--- a/Engine/Plugins/Runtime/RigVM/Source/RigVMEditor/Public/Editor/RigVMLegacyEditor.h
+++ b/Engine/Plugins/Runtime/RigVM/Source/RigVMEditor/Public/Editor/RigVMLegacyEditor.h
@@ -8,7 +8,7 @@
 #include "BlueprintEditor.h"
 #include "BlueprintEditorTabs.h"
 #include "RigVMHost.h"
-#include "RigVMBlueprint.h"
+#include "RigVMBlueprintLegacy.h"
 #include "RigVMLegacyEditorMode.h"
 #include "Editor/RigVMDetailsViewWrapperObject.h"
 #include "AssetRegistry/AssetData.h"
@@ -34,7 +34,7 @@ protected:
 	virtual const TSharedPtr<FAssetEditorToolkit> GetHostingApp() const override { return ConstCastSharedRef<FAssetEditorToolkit>(AsShared()); }
 	RIGVMEDITOR_API virtual void InitAssetEditor(const EToolkitMode::Type Mode, const TSharedPtr<IToolkitHost>& InitToolkitHost, const FName AppIdentifier, const TSharedRef<FTabManager::FLayout>& StandaloneDefaultLayout, const bool bCreateDefaultStandaloneMenu, const bool bCreateDefaultToolbar, const TArray<UObject*>& ObjectsToEdit, const bool bInIsToolbarFocusable, const bool bInUseSmallToolbarIcons, const TOptional<EAssetOpenMethod>& InOpenMethod) override;
 	RIGVMEDITOR_API virtual void CreateEditorToolbar() override;
-	virtual void CommonInitialization(const TArray<UBlueprint*>& InitBlueprints, bool bShouldOpenInDefaultsMode) override { FBlueprintEditor::CommonInitialization(InitBlueprints, bShouldOpenInDefaultsMode); }
+	RIGVMEDITOR_API virtual void CommonInitialization(const TArray<TRigVMAssetInterfacePtr>& InitBlueprints, bool bShouldOpenInDefaultsMode) override;
 	virtual TSharedPtr<FDocumentTracker> GetDocumentManager() const override { return DocumentManager; }
 	virtual void AddApplicationMode(FName ModeName, TSharedRef<FApplicationMode> Mode) override { FBlueprintEditor::AddApplicationMode(ModeName, Mode); }
 	virtual void RegenerateMenusAndToolbars() override { FBlueprintEditor::RegenerateMenusAndToolbars(); }
@@ -67,6 +67,7 @@ protected:
 	virtual FGraphPanelSelectionSet GetSelectedNodes() const override { return FBlueprintEditor::GetSelectedNodes(); }
 	virtual void SetUISelectionState(FName SelectionOwner) override { FBlueprintEditor::SetUISelectionState(SelectionOwner); }
 	virtual void AnalyticsTrackNodeEvent(UBlueprint* Blueprint, UEdGraphNode* GraphNode, bool bNodeDelete) const override { FBlueprintEditor::AnalyticsTrackNodeEvent(Blueprint, GraphNode, bNodeDelete); }
+	virtual void AnalyticsTrackNodeEvent(IRigVMAssetInterface* Blueprint, UEdGraphNode* GraphNode, bool bNodeDelete) const override { FBlueprintEditor::AnalyticsTrackNodeEvent(Cast<UBlueprint>(Blueprint->GetObject()), GraphNode, bNodeDelete); }
 	RIGVMEDITOR_API virtual void JumpToHyperlink(const UObject* ObjectReference, bool bRequestRename) override;
 	RIGVMEDITOR_API virtual void PostUndo(bool bSuccess) override;
 	RIGVMEDITOR_API virtual void PostRedo(bool bSuccess) override;
@@ -86,7 +87,7 @@ protected:
 	virtual void AddMenuExtender(TSharedPtr<FExtender> Extender) override { FBlueprintEditor::AddMenuExtender(Extender); }
 	virtual void RemoveMenuExtender(TSharedPtr<FExtender> Extender) override { FBlueprintEditor::RemoveMenuExtender(Extender); }
 	virtual TSharedPtr<IMessageLogListing> GetCompilerResultsListing() override { return CompilerResultsListing; }
-	virtual void OnBlueprintChangedInnerImpl(UBlueprint* InBlueprint, bool bIsJustBeingCompiled) override { FBlueprintEditor::OnBlueprintChangedImpl(InBlueprint, bIsJustBeingCompiled); }
+	virtual void OnBlueprintChangedInnerImpl(IRigVMAssetInterface* InBlueprint, bool bIsJustBeingCompiled) override { FBlueprintEditor::OnBlueprintChangedImpl(Cast<UBlueprint>(InBlueprint->GetObject()), bIsJustBeingCompiled); }
 	RIGVMEDITOR_API virtual void RefreshEditorsImpl(ERefreshRigVMEditorReason::Type Reason) override;
 	virtual void SetupGraphEditorEvents(UEdGraph* InGraph, SGraphEditor::FGraphEditorEvents& InEvents) override { FRigVMEditorBase::SetupGraphEditorEvents(InGraph, InEvents); }
 	virtual void SetupGraphEditorEventsImpl(UEdGraph* InGraph, SGraphEditor::FGraphEditorEvents& InEvents) override { FBlueprintEditor::SetupGraphEditorEvents(InGraph, InEvents); }
@@ -133,7 +134,7 @@ protected:
 
 	///////////////////////////////////////////////////////////////////////////////////////////////
 	// FBlueprintEditor overrides
-	virtual UBlueprint* GetBlueprintObj() const override { return FRigVMEditorBase::GetBlueprintObj(); }
+	virtual UBlueprint* GetBlueprintObj() const override { return Cast<UBlueprint>(FRigVMEditorBase::GetRigVMAsset()->GetObject()); }
 	virtual TSubclassOf<UEdGraphSchema> GetDefaultSchemaClass() const override { return FRigVMEditorBase::GetDefaultSchemaClass(); }
 	virtual bool InEditingMode() const override { return FRigVMEditorBase::InEditingMode(); }
 	virtual bool CanAddNewLocalVariable() const override { return FRigVMEditorBase::CanAddNewLocalVariable(); }
@@ -158,16 +159,16 @@ protected:
 	virtual void SaveAssetAs_Execute() override { return FRigVMEditorBase::SaveAssetAs_Execute(); }
 	virtual bool IsInAScriptingMode() const override { return FRigVMEditorBase::IsInAScriptingMode(); }
 	virtual void NewDocument_OnClicked(FBlueprintEditor::ECreatedDocumentType GraphType) override { return FRigVMEditorBase::OnNewDocumentClicked((FRigVMEditorBase::ECreatedDocumentType)GraphType); }
-	virtual bool AreEventGraphsAllowed() const override { return FRigVMEditorBase::AreEventGraphsAllowed(); }
-	virtual bool AreMacrosAllowed() const override { return FRigVMEditorBase::AreMacrosAllowed(); }
-	virtual bool AreDelegatesAllowed() const override { return FRigVMEditorBase::AreDelegatesAllowed(); }
+	virtual bool AreEventGraphsAllowed() const override { return true; }
+	virtual bool AreMacrosAllowed() const override { return false; }
+	virtual bool AreDelegatesAllowed() const override { return false; }
 	virtual FGraphAppearanceInfo GetGraphAppearance(class UEdGraph* InGraph) const override { return FRigVMEditorBase::GetGraphAppearance(InGraph); }
 	virtual bool IsEditable(UEdGraph* InGraph) const override { return FRigVMEditorBase::IsEditable(InGraph); }
 	virtual bool IsCompilingEnabled() const override { return FRigVMEditorBase::IsCompilingEnabled(); }
 	virtual FText GetGraphDecorationString(UEdGraph* InGraph) const override { return FRigVMEditorBase::GetGraphDecorationString(InGraph); }
 	virtual void OnActiveTabChanged( TSharedPtr<SDockTab> PreviouslyActive, TSharedPtr<SDockTab> NewlyActivated ) override { return FRigVMEditorBase::OnActiveTabChanged(PreviouslyActive, NewlyActivated); }
 	virtual void OnSelectedNodesChangedImpl(const TSet<class UObject*>& NewSelection) override { FRigVMEditorBase::OnSelectedNodesChangedImpl(NewSelection); }
-	virtual void OnBlueprintChangedImpl(UBlueprint* InBlueprint, bool bIsJustBeingCompiled) override { return FRigVMEditorBase::OnBlueprintChangedImpl(InBlueprint, bIsJustBeingCompiled); }
+	virtual void OnBlueprintChangedImpl(UBlueprint* InBlueprint, bool bIsJustBeingCompiled) override { return FRigVMEditorBase::OnBlueprintChangedImpl(Cast<URigVMBlueprint>(InBlueprint), bIsJustBeingCompiled); }
 	virtual void RefreshEditors(ERefreshBlueprintEditorReason::Type Reason = ERefreshBlueprintEditorReason::UnknownReason) override { return ForceEditorRefresh((ERefreshRigVMEditorReason::Type)Reason); }
 	virtual void FocusInspectorOnGraphSelection(const TSet<class UObject*>& NewSelection, bool bForceRefresh = false) override { return FRigVMEditorBase::FocusInspectorOnGraphSelection(NewSelection, bForceRefresh); }
 	virtual void NotifyPreChange(FProperty* PropertyAboutToChange) override { return FRigVMEditorBase::NotifyPreChange(PropertyAboutToChange); }
diff --git a/Engine/Plugins/Runtime/RigVM/Source/RigVMEditor/Public/Editor/RigVMLocalVariableDetailCustomization.h b/Engine/Plugins/Runtime/RigVM/Source/RigVMEditor/Public/Editor/RigVMLocalVariableDetailCustomization.h
--- a/Engine/Plugins/Runtime/RigVM/Source/RigVMEditor/Public/Editor/RigVMLocalVariableDetailCustomization.h
+++ b/Engine/Plugins/Runtime/RigVM/Source/RigVMEditor/Public/Editor/RigVMLocalVariableDetailCustomization.h
@@ -6,7 +6,7 @@
 #include "IDetailCustomization.h"
 #include "IPropertyTypeCustomization.h"
 #include "RigVMHost.h"
-#include "RigVMBlueprint.h"
+#include "RigVMBlueprintLegacy.h"
 #include "Editor/RigVMDetailsViewWrapperObject.h"
 #include "EdGraph/RigVMEdGraph.h"
 #include "Widgets/SRigVMGraphPinNameListValueWidget.h"
@@ -23,7 +23,7 @@ class FRigVMLocalVariableDetailCustomization : public IDetailCustomization
 	FRigVMLocalVariableDetailCustomization()
 	: GraphBeingCustomized(nullptr)
 	, BlueprintBeingCustomized(nullptr)
-	, NameValidator(nullptr, nullptr, NAME_None)
+	, NameValidator((const TRigVMAssetInterfacePtr)nullptr, nullptr, NAME_None)
 	{}
 
 	
@@ -41,7 +41,7 @@ public:
 private:
 
 	URigVMGraph* GraphBeingCustomized;
-	URigVMBlueprint* BlueprintBeingCustomized;
+	TRigVMAssetInterfacePtr BlueprintBeingCustomized;
 	FRigVMGraphVariableDescription VariableDescription;
 	TArray<TWeakObjectPtr<URigVMDetailsViewWrapperObject>> ObjectsBeingCustomized;
 
diff --git a/Engine/Plugins/Runtime/RigVM/Source/RigVMEditor/Public/Editor/RigVMNewEditor.h b/Engine/Plugins/Runtime/RigVM/Source/RigVMEditor/Public/Editor/RigVMNewEditor.h
--- a/Engine/Plugins/Runtime/RigVM/Source/RigVMEditor/Public/Editor/RigVMNewEditor.h
+++ b/Engine/Plugins/Runtime/RigVM/Source/RigVMEditor/Public/Editor/RigVMNewEditor.h
@@ -32,8 +32,8 @@ protected:
 	
 	UE_API virtual void InitAssetEditor(const EToolkitMode::Type Mode, const TSharedPtr<IToolkitHost>& InitToolkitHost, const FName AppIdentifier, const TSharedRef<FTabManager::FLayout>& StandaloneDefaultLayout, const bool bCreateDefaultStandaloneMenu, const bool bCreateDefaultToolbar, const TArray<UObject*>& ObjectsToEdit, const bool bInIsToolbarFocusable, const bool bInUseSmallToolbarIcons, const TOptional<EAssetOpenMethod>& InOpenMethod) override;
 	virtual void CreateEditorToolbar() override {}
-	UE_API virtual void CommonInitialization(const TArray<UBlueprint*>& InitBlueprints, bool bShouldOpenInDefaultsMode) override;
-	UE_API void OnBlueprintChanged(UBlueprint* InBlueprint);
+	UE_API virtual void CommonInitialization(const TArray<TRigVMAssetInterfacePtr>& InitBlueprints, bool bShouldOpenInDefaultsMode) override;
+	UE_API void OnBlueprintChanged(UObject* InBlueprint);
 	UE_API void SaveEditedObjectState();
 	virtual TSharedPtr<FDocumentTracker> GetDocumentManager() const override { return DocumentManager; }
 	virtual void AddApplicationMode(FName ModeName, TSharedRef<FApplicationMode> Mode) override { FWorkflowCentricApplication::AddApplicationMode(ModeName, Mode); }
@@ -71,7 +71,7 @@ protected:
 	UE_API virtual void RenameNewlyAddedAction(FName InActionName) override;
 	UE_API virtual FGraphPanelSelectionSet GetSelectedNodes() const override;
 	UE_API virtual void SetUISelectionState(FName SelectionOwner) override;
-	UE_API virtual void AnalyticsTrackNodeEvent(UBlueprint* Blueprint, UEdGraphNode* GraphNode, bool bNodeDelete) const override;
+	UE_API virtual void AnalyticsTrackNodeEvent(IRigVMAssetInterface* Blueprint, UEdGraphNode* GraphNode, bool bNodeDelete) const override;
 	UE_API virtual void PostUndo(bool bSuccess) override;
 	UE_API virtual void PostRedo(bool bSuccess) override;
 	UE_API virtual UEdGraphPin* GetCurrentlySelectedPin() const override;
@@ -89,7 +89,7 @@ protected:
 	virtual void RemoveToolbarExtender(TSharedPtr<FExtender> Extender) override { FWorkflowCentricApplication::RemoveToolbarExtender(Extender); };
 	virtual void AddMenuExtender(TSharedPtr<FExtender> Extender) override { FWorkflowCentricApplication::AddMenuExtender(Extender); }
 	virtual void RemoveMenuExtender(TSharedPtr<FExtender> Extender) override { FWorkflowCentricApplication::RemoveMenuExtender(Extender); }
-	UE_API virtual void OnBlueprintChangedInnerImpl(UBlueprint* InBlueprint, bool bIsJustBeingCompiled) override;
+	UE_API virtual void OnBlueprintChangedInnerImpl(IRigVMAssetInterface* InBlueprint, bool bIsJustBeingCompiled) override;
 	UE_API virtual void RefreshEditorsImpl(ERefreshRigVMEditorReason::Type Reason) override;
 	UE_API virtual void SetupGraphEditorEventsImpl(UEdGraph* InGraph, SGraphEditor::FGraphEditorEvents& InEvents) override;
 	UE_API virtual FActionMenuContent OnCreateGraphActionMenu(UEdGraph* InGraph, const FVector2f& InNodePosition, const TArray<UEdGraphPin*>& InDraggedPins, bool bAutoExpand, SGraphEditor::FActionMenuClosed InOnMenuClosed) override;
@@ -154,7 +154,9 @@ protected:
 	/** Dumps messages to the compiler log, with an option to force it to display/come to front */
 	UE_API void DumpMessagesToCompilerLog(const TArray<TSharedRef<class FTokenizedMessage>>& Messages, bool bForceMessageDisplay);
 public:
-	UE_API void CreateDefaultTabContents(const TArray<UBlueprint*> InBlueprints);
+	UE_DEPRECATED(5.7, "Please use void CreateDefaultTabContents(const TArray<TRigVMAssetInterfacePtr> InBlueprints)")
+	UE_API void CreateDefaultTabContents(const TArray<UBlueprint*> InBlueprints){}
+	UE_API void CreateDefaultTabContents(const TArray<TRigVMAssetInterfacePtr> InBlueprints);
 
 	UE_API TSharedRef<SWidget> GenerateCompileOptionsMenu();
 	UE_API void MakeSaveOnCompileSubMenu(FMenuBuilder& InMenu);
diff --git a/Engine/Plugins/Runtime/RigVM/Source/RigVMEditor/Public/Editor/RigVMNewEditorMode.h b/Engine/Plugins/Runtime/RigVM/Source/RigVMEditor/Public/Editor/RigVMNewEditorMode.h
--- a/Engine/Plugins/Runtime/RigVM/Source/RigVMEditor/Public/Editor/RigVMNewEditorMode.h
+++ b/Engine/Plugins/Runtime/RigVM/Source/RigVMEditor/Public/Editor/RigVMNewEditorMode.h
@@ -8,7 +8,6 @@
 
 #define UE_API RIGVMEDITOR_API
 
-class URigVMBlueprint;
 class FRigVMNewEditor;
 
 struct FRigVMNewEditorApplicationModes
diff --git a/Engine/Plugins/Runtime/RigVM/Source/RigVMEditor/Public/Editor/RigVMVariableDetailCustomization.h b/Engine/Plugins/Runtime/RigVM/Source/RigVMEditor/Public/Editor/RigVMVariableDetailCustomization.h
--- a/Engine/Plugins/Runtime/RigVM/Source/RigVMEditor/Public/Editor/RigVMVariableDetailCustomization.h
+++ b/Engine/Plugins/Runtime/RigVM/Source/RigVMEditor/Public/Editor/RigVMVariableDetailCustomization.h
@@ -5,6 +5,7 @@
 #include "CoreMinimal.h"
 #include "Engine/Blueprint.h"
 #include "IDetailCustomization.h"
+#include "RigVMAsset.h"
 #include "Styling/SlateTypes.h"
 #include "UObject/WeakFieldPtr.h"
 #include "Layout/Visibility.h"
@@ -14,7 +15,6 @@
 
 #define UE_API RIGVMEDITOR_API
 
-class URigVMBlueprint;
 class SRigVMEditorGraphExplorer;
 class IDetailLayoutBuilder;
 class IBlueprintEditor;
@@ -31,7 +31,7 @@ class FRigVMVariableDetailCustomization : public IDetailCustomization
 public:
 	/** Makes a new instance of this detail layout class for a specific detail view requesting it */
 	static UE_API TSharedRef<IDetailCustomization> MakeInstance(TSharedPtr<IRigVMEditor> InEditor);
-	UE_API FRigVMVariableDetailCustomization(TSharedPtr<IRigVMEditor> InEditor, URigVMBlueprint* Blueprint);
+	UE_API FRigVMVariableDetailCustomization(TSharedPtr<IRigVMEditor> InEditor, TRigVMAssetInterfacePtr Blueprint);
 
 #if WITH_RIGVMLEGACYEDITOR
 	static UE_API TSharedPtr<IDetailCustomization> MakeLegacyInstance(TSharedPtr<IBlueprintEditor> InBlueprintEditor);
@@ -92,7 +92,7 @@ private:
 	TWeakPtr<IRigVMEditor> EditorPtr;
 
 	/** The blueprint we are editing */
-	TWeakObjectPtr<URigVMBlueprint> BlueprintPtr;
+	TWeakInterfacePtr<IRigVMAssetInterface> BlueprintPtr;
 
 	/** The widget used when in variable name editing mode */ 
 	TSharedPtr<SEditableTextBox> VarNameEditableTextBox;
diff --git a/Engine/Plugins/Runtime/RigVM/Source/RigVMEditor/Public/Editor/RigVMVariantDetailCustomization.h b/Engine/Plugins/Runtime/RigVM/Source/RigVMEditor/Public/Editor/RigVMVariantDetailCustomization.h
--- a/Engine/Plugins/Runtime/RigVM/Source/RigVMEditor/Public/Editor/RigVMVariantDetailCustomization.h
+++ b/Engine/Plugins/Runtime/RigVM/Source/RigVMEditor/Public/Editor/RigVMVariantDetailCustomization.h
@@ -4,7 +4,7 @@
 
 #include "CoreMinimal.h"
 #include "IPropertyTypeCustomization.h"
-#include "RigVMBlueprint.h"
+#include "RigVMBlueprintLegacy.h"
 #include "RigVMCore/RigVMVariant.h"
 #include "Widgets/SRigVMLogWidget.h"
 
@@ -35,7 +35,7 @@ protected:
 	void OnAddTag(const FName& InTagName);
 	void OnRemoveTag(const FName& InTagName);
 
-	URigVMBlueprint* BlueprintBeingCustomized = nullptr;
+	TRigVMAssetInterfacePtr BlueprintBeingCustomized = nullptr;
 
 	/** The log widget used for function variants */
 	TSharedPtr<SRigVMLogWidget> VariantLog;
diff --git a/Engine/Plugins/Runtime/RigVM/Source/RigVMEditor/Public/Editor/SRigVMActionMenu.cpp b/Engine/Plugins/Runtime/RigVM/Source/RigVMEditor/Public/Editor/SRigVMActionMenu.cpp
--- a/Engine/Plugins/Runtime/RigVM/Source/RigVMEditor/Public/Editor/SRigVMActionMenu.cpp
+++ b/Engine/Plugins/Runtime/RigVM/Source/RigVMEditor/Public/Editor/SRigVMActionMenu.cpp
@@ -15,7 +15,8 @@
 #include "BlueprintNodeSpawner.h"
 #include "K2Node_Variable.h"
 #include "RigVMActionMenuItem.h"
-#include "RigVMBlueprint.h"
+#include "RigVMBlueprintLegacy.h"
+#include "RigVMBlueprintUtils.h"
 #include "Widgets/Input/SCheckBox.h"
 #include "Framework/Application/SlateApplication.h"
 #include "Widgets/SToolTip.h"
@@ -40,8 +41,7 @@ struct FRigVMAction_PromoteVariable : public FEdGraphSchemaAction
 	{
 		if( ( ParentGraph != NULL ) && ( FromPin != NULL ) )
 		{
-			UBlueprint* Blueprint = FBlueprintEditorUtils::FindBlueprintForGraph(ParentGraph);
-			if (URigVMBlueprint* RigVMBlueprint = Cast<URigVMBlueprint>(Blueprint))
+			if (TRigVMAssetInterfacePtr RigVMBlueprint = FRigVMBlueprintUtils::FindAssetForGraph(ParentGraph))
 			{
 				if (URigVMController* Controller = RigVMBlueprint->GetController(ParentGraph))
 				{
@@ -297,15 +297,17 @@ void SRigVMActionMenu::ConstructActionContext(FBlueprintActionContext& ContextDe
 	// with a valid node, then fix it up in filtering
 	ContextDescOut.Graphs.Add(GraphObj);
 
-	UBlueprint* Blueprint = BlueprintEditor->GetRigVMBlueprint();
-	const bool bBlueprintIsValid = IsValid(Blueprint) && Blueprint->GeneratedClass && (Blueprint->GeneratedClass->ClassGeneratedBy == Blueprint);
-	if (!ensure(bBlueprintIsValid))  // to track UE-11597 and UE-11595
+	if (UBlueprint* Blueprint = Cast<UBlueprint>(BlueprintEditor->GetRigVMAsset().GetObject()))
 	{
-		return;
-	}
+		const bool bBlueprintIsValid = IsValid(Blueprint) && Blueprint->GeneratedClass && (Blueprint->GeneratedClass->ClassGeneratedBy == Blueprint);
+		if (!ensure(bBlueprintIsValid))  // to track UE-11597 and UE-11595
+		{
+			return;
+		}
 
-	//ContextDescOut.EditorPtr = EditorPtr;
-	ContextDescOut.Blueprints.Add(Blueprint);
+		//ContextDescOut.EditorPtr = EditorPtr;
+		ContextDescOut.Blueprints.Add(Blueprint);
+	}
 
 	if (bIsContextSensitive)
 	{
diff --git a/Engine/Plugins/Runtime/RigVM/Source/RigVMEditor/Public/Editor/SRigVMDetailsInspector.h b/Engine/Plugins/Runtime/RigVM/Source/RigVMEditor/Public/Editor/SRigVMDetailsInspector.h
--- a/Engine/Plugins/Runtime/RigVM/Source/RigVMEditor/Public/Editor/SRigVMDetailsInspector.h
+++ b/Engine/Plugins/Runtime/RigVM/Source/RigVMEditor/Public/Editor/SRigVMDetailsInspector.h
@@ -11,7 +11,7 @@
 #include "Internationalization/Text.h"
 #include "Layout/Visibility.h"
 #include "PropertyEditorDelegates.h"
-#include "RigVMBlueprint.h"
+#include "RigVMAsset.h"
 #include "SlateFwd.h"
 #include "Styling/SlateTypes.h"
 #include "Templates/SharedPointer.h"
@@ -123,7 +123,7 @@ public:
 	/** returns the list of selected objects */
 	UE_API const TArray< TWeakObjectPtr<UObject> >& GetSelectedObjects() const;
 
-	UE_API void OnEditorClose(const IRigVMEditor* RigVMEditorBase, URigVMBlueprint* RigVMBlueprint);
+	UE_API void OnEditorClose(const IRigVMEditor* RigVMEditorBase, TRigVMAssetInterfacePtr RigVMBlueprint);
 
 	UE_API virtual void AddReferencedObjects(FReferenceCollector& Collector) override;
 	UE_API virtual FString GetReferencerName() const override;
diff --git a/Engine/Plugins/Runtime/RigVM/Source/RigVMEditor/Public/RigVMEditorModule.h b/Engine/Plugins/Runtime/RigVM/Source/RigVMEditor/Public/RigVMEditorModule.h
--- a/Engine/Plugins/Runtime/RigVM/Source/RigVMEditor/Public/RigVMEditorModule.h
+++ b/Engine/Plugins/Runtime/RigVM/Source/RigVMEditor/Public/RigVMEditorModule.h
@@ -7,22 +7,25 @@
 #pragma once
 
 #include "CoreMinimal.h"
+#include "RigVMAsset.h"
 #include "Modules/ModuleInterface.h"
 #include "Logging/LogMacros.h"
 #include "EdGraph/RigVMEdGraphPanelNodeFactory.h"
 #include "EdGraph/RigVMEdGraphPanelPinFactory.h"
+#include "EdGraph/RigVMEdGraphSchema.h"
 #include "Toolkits/AssetEditorToolkit.h"
 #include "Kismet2/StructureEditorUtils.h"
 #include "Modules/ModuleManager.h"
 
 #define UE_API RIGVMEDITOR_API
 
+class URigVMAssetInterface;
+class IRigVMAssetInterface;
 class URigVMUserWorkflowOptions;
 class URigVMPin;
 class URigVMNode;
 class URigVMEdGraphNode;
 class IRigVMClientHost;
-class URigVMBlueprint;
 class URigVMEdGraphSchema;
 class FRigVMEditorBase;
 class FBlueprintActionDatabaseRegistrar;
@@ -54,10 +57,10 @@ public:
 	}
 
 	virtual void GetContextMenuActions(const URigVMEdGraphSchema* Schema, class UToolMenu* Menu, class UGraphNodeContextMenuContext* Context) const = 0;
-	virtual void GetTypeActions(URigVMBlueprint* RigVMBlueprint, FBlueprintActionDatabaseRegistrar& ActionRegistrar) = 0;
-	virtual void GetInstanceActions(URigVMBlueprint* RigVMBlueprint, FBlueprintActionDatabaseRegistrar& ActionRegistrar) = 0;
-	virtual void GetNodeContextMenuActions(IRigVMClientHost* RigVMClientHost, const URigVMEdGraphNode* EdGraphNode, URigVMNode* ModelNode, class UToolMenu* Menu) const = 0;
-	virtual void GetPinContextMenuActions(IRigVMClientHost* RigVMClientHost, const UEdGraphPin* EdGraphPin, URigVMPin* ModelPin, class UToolMenu* Menu) const = 0;
+	virtual void GetTypeActions(TRigVMAssetInterfacePtr RigVMBlueprint, FBlueprintActionDatabaseRegistrar& ActionRegistrar) = 0;
+	virtual void GetInstanceActions(TRigVMAssetInterfacePtr RigVMBlueprint, FBlueprintActionDatabaseRegistrar& ActionRegistrar) = 0;
+	virtual void GetNodeContextMenuActions(TScriptInterface<IRigVMClientHost> RigVMClientHost, const URigVMEdGraphNode* EdGraphNode, URigVMNode* ModelNode, class UToolMenu* Menu) const = 0;
+	virtual void GetPinContextMenuActions(TScriptInterface<IRigVMClientHost> RigVMClientHost, const UEdGraphPin* EdGraphPin, URigVMPin* ModelPin, class UToolMenu* Menu) const = 0;
 	virtual FConnectionDrawingPolicy* CreateConnectionDrawingPolicy(int32 InBackLayerID, int32 InFrontLayerID, float InZoomFactor, const FSlateRect& InClippingRect, class FSlateWindowElementList& InDrawElements, class UEdGraph* InGraphObj) const = 0;
 	virtual bool AssetsPublicFunctionsAllowed(const FAssetData& InAssetData) const = 0;
 };
@@ -76,12 +79,14 @@ public:
 	UE_API void ShutdownModuleCommon();
 
 	/** RigVMEditorModule interface */
-	UE_API virtual UClass* GetRigVMBlueprintClass() const;
-	UE_API const URigVMBlueprint* GetRigVMBlueprintCDO() const;
-	UE_API virtual void GetTypeActions(URigVMBlueprint* RigVMBlueprint, FBlueprintActionDatabaseRegistrar& ActionRegistrar);
-	UE_API virtual void GetInstanceActions(URigVMBlueprint* RigVMBlueprint, FBlueprintActionDatabaseRegistrar& ActionRegistrar);
-	UE_API virtual void GetNodeContextMenuActions(IRigVMClientHost* RigVMClientHost, const URigVMEdGraphNode* EdGraphNode, URigVMNode* ModelNode, class UToolMenu* Menu) const;
-	UE_API virtual void GetPinContextMenuActions(IRigVMClientHost* RigVMClientHost, const UEdGraphPin* EdGraphPin, URigVMPin* ModelPin, class UToolMenu* Menu) const;
+	UE_DEPRECATED(5.7, "Please update your code")
+	UE_API virtual UClass* GetRigVMBlueprintClass() const { return nullptr; }
+	UE_DEPRECATED(5.7, "Please update your code")
+	UE_API const URigVMBlueprint* GetRigVMBlueprintCDO() const { return nullptr; }
+	UE_API virtual void GetTypeActions(TRigVMAssetInterfacePtr RigVMBlueprint, FBlueprintActionDatabaseRegistrar& ActionRegistrar) override;
+	UE_API virtual void GetInstanceActions(TRigVMAssetInterfacePtr RigVMBlueprint, FBlueprintActionDatabaseRegistrar& ActionRegistrar) override;
+	UE_API virtual void GetNodeContextMenuActions(TScriptInterface<IRigVMClientHost> RigVMClientHost, const URigVMEdGraphNode* EdGraphNode, URigVMNode* ModelNode, class UToolMenu* Menu) const override;
+	UE_API virtual void GetPinContextMenuActions(TScriptInterface<IRigVMClientHost> RigVMClientHost, const UEdGraphPin* EdGraphPin, URigVMPin* ModelPin, class UToolMenu* Menu) const override;
 	virtual bool AssetsPublicFunctionsAllowed(const FAssetData& InAssetData) const override { return true; }
 
 	/** IHasMenuExtensibility interface */
@@ -101,36 +106,39 @@ public:
 	UE_API virtual void GetContextMenuActions(const URigVMEdGraphSchema* Schema, class UToolMenu* Menu, class UGraphNodeContextMenuContext* Context) const override;
 
 	/** Make sure to create the root graph for a given blueprint */
-	UE_API void CreateRootGraphIfRequired(URigVMBlueprint* InBlueprint) const;
+	UE_API void CreateRootGraphIfRequired(TRigVMAssetInterfacePtr InBlueprint) const;
 
 	UE_API virtual FConnectionDrawingPolicy* CreateConnectionDrawingPolicy(int32 InBackLayerID, int32 InFrontLayerID, float InZoomFactor, const FSlateRect& InClippingRect, FSlateWindowElementList& InDrawElements, UEdGraph* InGraphObj) const override;
 
 protected:
 
 	UE_API bool IsRigVMEditorModuleBase() const;
+	UE_API virtual UClass* GetRigVMAssetClass() const { return URigVMAssetInterface::StaticClass(); }
+	UE_API virtual UClass* GetRigVMEdGraphSchemaClass() const { return URigVMEdGraphSchema::StaticClass(); }
+	UE_API virtual UScriptStruct* GetRigVMExecuteContextStruct() const { return FRigVMExecuteContext::StaticStruct(); }
 
 	/** Specific section callbacks for the context menu */
-	UE_API virtual void GetNodeWorkflowContextMenuActions(IRigVMClientHost* RigVMClientHost, const URigVMEdGraphNode* EdGraphNode, URigVMNode* ModelNode, UToolMenu* Menu) const;
-	UE_API virtual void GetNodeEventsContextMenuActions(IRigVMClientHost* RigVMClientHost, const URigVMEdGraphNode* EdGraphNode, URigVMNode* ModelNode, UToolMenu* Menu) const;
-	UE_API virtual void GetNodeDefaultValueContextMenuActions(IRigVMClientHost* RigVMClientHost, const URigVMEdGraphNode* EdGraphNode, URigVMNode* ModelNode, UToolMenu* Menu) const;
-	UE_API virtual void GetNodeConversionContextMenuActions(IRigVMClientHost* RigVMClientHost, const URigVMEdGraphNode* EdGraphNode, URigVMNode* ModelNode, UToolMenu* Menu) const;
-	UE_API virtual void GetNodeDebugContextMenuActions(IRigVMClientHost* RigVMClientHost, const URigVMEdGraphNode* EdGraphNode, URigVMNode* ModelNode, UToolMenu* Menu) const;
-	UE_API virtual void GetNodeVariablesContextMenuActions(IRigVMClientHost* RigVMClientHost, const URigVMEdGraphNode* EdGraphNode, URigVMNode* ModelNode, UToolMenu* Menu) const;
-	UE_API virtual void GetNodeTemplatesContextMenuActions(IRigVMClientHost* RigVMClientHost, const URigVMEdGraphNode* EdGraphNode, URigVMNode* ModelNode, UToolMenu* Menu) const;
-	UE_API virtual void GetNodeOrganizationContextMenuActions(IRigVMClientHost* RigVMClientHost, const URigVMEdGraphNode* EdGraphNode, URigVMNode* ModelNode, UToolMenu* Menu) const;
-	UE_API virtual void GetNodeVariantContextMenuActions(IRigVMClientHost* RigVMClientHost, const URigVMEdGraphNode* EdGraphNode, URigVMNode* ModelNode, UToolMenu* Menu) const;
-	UE_API virtual void GetNodeTestContextMenuActions(IRigVMClientHost* RigVMClientHost, const URigVMEdGraphNode* EdGraphNode, URigVMNode* ModelNode, UToolMenu* Menu) const;
-	UE_API virtual void GetNodeDisplayContextMenuActions(IRigVMClientHost* RigVMClientHost, const URigVMEdGraphNode* EdGraphNode, URigVMNode* ModelNode, UToolMenu* Menu) const;
-	UE_API virtual void GetExposedPinContextMenuActions(IRigVMClientHost* RigVMClientHost, const UEdGraphPin* EdGraphPin, URigVMPin* ModelPin, UToolMenu* Menu) const;
-	UE_API virtual void GetPinWorkflowContextMenuActions(IRigVMClientHost* RigVMClientHost, const UEdGraphPin* EdGraphPin, URigVMPin* ModelPin, UToolMenu* Menu) const;
-	UE_API virtual void GetPinDebugContextMenuActions(IRigVMClientHost* RigVMClientHost, const UEdGraphPin* EdGraphPin, URigVMPin* ModelPin, UToolMenu* Menu) const;
-	UE_API virtual void GetPinArrayContextMenuActions(IRigVMClientHost* RigVMClientHost, const UEdGraphPin* EdGraphPin, URigVMPin* ModelPin, UToolMenu* Menu) const;
-	UE_API virtual void GetPinAggregateContextMenuActions(IRigVMClientHost* RigVMClientHost, const UEdGraphPin* EdGraphPin, URigVMPin* ModelPin, UToolMenu* Menu) const;
-	UE_API virtual void GetPinTemplateContextMenuActions(IRigVMClientHost* RigVMClientHost, const UEdGraphPin* EdGraphPin, URigVMPin* ModelPin, UToolMenu* Menu) const;
-	UE_API virtual void GetPinConversionContextMenuActions(IRigVMClientHost* RigVMClientHost, const UEdGraphPin* EdGraphPin, URigVMPin* ModelPin, UToolMenu* Menu) const;
-	UE_API virtual void GetPinVariableContextMenuActions(IRigVMClientHost* RigVMClientHost, const UEdGraphPin* EdGraphPin, URigVMPin* ModelPin, UToolMenu* Menu) const;
-	UE_API virtual void GetPinResetDefaultContextMenuActions(IRigVMClientHost* RigVMClientHost, const UEdGraphPin* EdGraphPin, URigVMPin* ModelPin, UToolMenu* Menu) const;
-	UE_API virtual void GetPinInjectedNodesContextMenuActions(IRigVMClientHost* RigVMClientHost, const UEdGraphPin* EdGraphPin, URigVMPin* ModelPin, UToolMenu* Menu) const;
+	UE_API virtual void GetNodeWorkflowContextMenuActions(TScriptInterface<IRigVMClientHost> RigVMClientHost, const URigVMEdGraphNode* EdGraphNode, URigVMNode* ModelNode, UToolMenu* Menu) const;
+	UE_API virtual void GetNodeEventsContextMenuActions(TScriptInterface<IRigVMClientHost> RigVMClientHost, const URigVMEdGraphNode* EdGraphNode, URigVMNode* ModelNode, UToolMenu* Menu) const;
+	UE_API virtual void GetNodeDefaultValueContextMenuActions(TScriptInterface<IRigVMClientHost> RigVMClientHost, const URigVMEdGraphNode* EdGraphNode, URigVMNode* ModelNode, UToolMenu* Menu) const;
+	UE_API virtual void GetNodeConversionContextMenuActions(TScriptInterface<IRigVMClientHost> RigVMClientHost, const URigVMEdGraphNode* EdGraphNode, URigVMNode* ModelNode, UToolMenu* Menu) const;
+	UE_API virtual void GetNodeDebugContextMenuActions(TScriptInterface<IRigVMClientHost> RigVMClientHost, const URigVMEdGraphNode* EdGraphNode, URigVMNode* ModelNode, UToolMenu* Menu) const;
+	UE_API virtual void GetNodeVariablesContextMenuActions(TScriptInterface<IRigVMClientHost> RigVMClientHost, const URigVMEdGraphNode* EdGraphNode, URigVMNode* ModelNode, UToolMenu* Menu) const;
+	UE_API virtual void GetNodeTemplatesContextMenuActions(TScriptInterface<IRigVMClientHost> RigVMClientHost, const URigVMEdGraphNode* EdGraphNode, URigVMNode* ModelNode, UToolMenu* Menu) const;
+	UE_API virtual void GetNodeOrganizationContextMenuActions(TScriptInterface<IRigVMClientHost> RigVMClientHost, const URigVMEdGraphNode* EdGraphNode, URigVMNode* ModelNode, UToolMenu* Menu) const;
+	UE_API virtual void GetNodeVariantContextMenuActions(TScriptInterface<IRigVMClientHost> RigVMClientHost, const URigVMEdGraphNode* EdGraphNode, URigVMNode* ModelNode, UToolMenu* Menu) const;
+	UE_API virtual void GetNodeTestContextMenuActions(TScriptInterface<IRigVMClientHost> RigVMClientHost, const URigVMEdGraphNode* EdGraphNode, URigVMNode* ModelNode, UToolMenu* Menu) const;
+	UE_API virtual void GetNodeDisplayContextMenuActions(TScriptInterface<IRigVMClientHost> RigVMClientHost, const URigVMEdGraphNode* EdGraphNode, URigVMNode* ModelNode, UToolMenu* Menu) const;
+	UE_API virtual void GetExposedPinContextMenuActions(TScriptInterface<IRigVMClientHost> RigVMClientHost, const UEdGraphPin* EdGraphPin, URigVMPin* ModelPin, UToolMenu* Menu) const;
+	UE_API virtual void GetPinWorkflowContextMenuActions(TScriptInterface<IRigVMClientHost> RigVMClientHost, const UEdGraphPin* EdGraphPin, URigVMPin* ModelPin, UToolMenu* Menu) const;
+	UE_API virtual void GetPinDebugContextMenuActions(TScriptInterface<IRigVMClientHost> RigVMClientHost, const UEdGraphPin* EdGraphPin, URigVMPin* ModelPin, UToolMenu* Menu) const;
+	UE_API virtual void GetPinArrayContextMenuActions(TScriptInterface<IRigVMClientHost> RigVMClientHost, const UEdGraphPin* EdGraphPin, URigVMPin* ModelPin, UToolMenu* Menu) const;
+	UE_API virtual void GetPinAggregateContextMenuActions(TScriptInterface<IRigVMClientHost> RigVMClientHost, const UEdGraphPin* EdGraphPin, URigVMPin* ModelPin, UToolMenu* Menu) const;
+	UE_API virtual void GetPinTemplateContextMenuActions(TScriptInterface<IRigVMClientHost> RigVMClientHost, const UEdGraphPin* EdGraphPin, URigVMPin* ModelPin, UToolMenu* Menu) const;
+	UE_API virtual void GetPinConversionContextMenuActions(TScriptInterface<IRigVMClientHost> RigVMClientHost, const UEdGraphPin* EdGraphPin, URigVMPin* ModelPin, UToolMenu* Menu) const;
+	UE_API virtual void GetPinVariableContextMenuActions(TScriptInterface<IRigVMClientHost> RigVMClientHost, const UEdGraphPin* EdGraphPin, URigVMPin* ModelPin, UToolMenu* Menu) const;
+	UE_API virtual void GetPinResetDefaultContextMenuActions(TScriptInterface<IRigVMClientHost> RigVMClientHost, const UEdGraphPin* EdGraphPin, URigVMPin* ModelPin, UToolMenu* Menu) const;
+	UE_API virtual void GetPinInjectedNodesContextMenuActions(TScriptInterface<IRigVMClientHost> RigVMClientHost, const UEdGraphPin* EdGraphPin, URigVMPin* ModelPin, UToolMenu* Menu) const;
 
 	/** Extensibility managers */
 	TSharedPtr<FExtensibilityManager> MenuExtensibilityManager;
diff --git a/Engine/Plugins/Runtime/RigVM/Source/RigVMEditor/Public/Widgets/RigVMTreeToolkitNode.h b/Engine/Plugins/Runtime/RigVM/Source/RigVMEditor/Public/Widgets/RigVMTreeToolkitNode.h
--- a/Engine/Plugins/Runtime/RigVM/Source/RigVMEditor/Public/Widgets/RigVMTreeToolkitNode.h
+++ b/Engine/Plugins/Runtime/RigVM/Source/RigVMEditor/Public/Widgets/RigVMTreeToolkitNode.h
@@ -2,7 +2,7 @@
 
 #pragma once
 
-#include "RigVMBlueprint.h"
+#include "RigVMBlueprintLegacy.h"
 #include "RigVMTreeToolkitDefines.h"
 #include "RigVMCore/RigVMVariant.h"
 #include "Styling/SlateTypes.h"
@@ -30,7 +30,9 @@ public:
 	}
 
 	virtual FAssetData GetAssetData() const;
-	virtual URigVMBlueprint* GetBlueprint() const;
+	UE_DEPRECATED(5.7, "Please use IRigVMAssetInterface* GetAsset() const")
+	virtual URigVMBlueprint* GetBlueprint() const { return Cast<URigVMBlueprint>(GetAsset()->GetObject()); }
+	virtual IRigVMAssetInterface* GetAsset() const;
 
 	virtual FText GetLabel() const;
 
diff --git a/Engine/Plugins/Runtime/RigVM/Source/RigVMEditor/Public/Widgets/RigVMTreeToolkitTask.h b/Engine/Plugins/Runtime/RigVM/Source/RigVMEditor/Public/Widgets/RigVMTreeToolkitTask.h
--- a/Engine/Plugins/Runtime/RigVM/Source/RigVMEditor/Public/Widgets/RigVMTreeToolkitTask.h
+++ b/Engine/Plugins/Runtime/RigVM/Source/RigVMEditor/Public/Widgets/RigVMTreeToolkitTask.h
@@ -23,11 +23,17 @@ public:
 		return false;
 	};
 
+	UE_DEPRECATED(5.7, "Please use TRigVMAssetInterfacePtr GetRigVMAsset(const TSharedRef<FRigVMTreePhase>& InPhase) const")
 	virtual URigVMBlueprint* GetBlueprint(const TSharedRef<FRigVMTreePhase>& InPhase) const
 	{
 		return nullptr;
 	}
 
+	virtual TRigVMAssetInterfacePtr GetRigVMAsset(const TSharedRef<FRigVMTreePhase>& InPhase) const
+	{
+		return nullptr;
+	}
+
 	virtual bool RequiresRefresh() const
 	{
 		return false;
diff --git a/Engine/Plugins/Runtime/RigVM/Source/RigVMEditor/Public/Widgets/SRigVMEditorGraphExplorer.h b/Engine/Plugins/Runtime/RigVM/Source/RigVMEditor/Public/Widgets/SRigVMEditorGraphExplorer.h
--- a/Engine/Plugins/Runtime/RigVM/Source/RigVMEditor/Public/Widgets/SRigVMEditorGraphExplorer.h
+++ b/Engine/Plugins/Runtime/RigVM/Source/RigVMEditor/Public/Widgets/SRigVMEditorGraphExplorer.h
@@ -4,6 +4,8 @@
 
 #include "CoreMinimal.h"
 #include "GraphEditorDragDropAction.h"
+#include "RigVMAsset.h"
+#include "RigVMBlueprintLegacy.h"
 #include "DragAndDrop/GraphNodeDragDropOp.h"
 #include "Widgets/DeclarativeSyntaxSupport.h"
 #include "Widgets/SCompoundWidget.h"
@@ -16,7 +18,7 @@ class SRigVMEditorGraphExplorerTreeView;
 //
 class URigVMGraph;
 class URigVMEdGraph;
-class URigVMBlueprint;
+class IRigVMAssetInterface;
 class SGraphActionMenu;
 struct FCreateWidgetForActionData;
 class FRigVMEditorGraphExplorerTreeElement;
@@ -30,14 +32,22 @@ class FRigVMGraphExplorerDragDropOp : public FGraphEditorDragDropAction
 public:
 	DRAG_DROP_OPERATOR_TYPE(FRigVMExplorerDragDropOp, FGraphEditorDragDropAction)
 
+	UE_DEPRECATED(5.7, "Please use New(const FRigVMExplorerElementKey& InElement, TRigVMAssetInterfacePtr InBlueprint);")
 	static TSharedRef<FRigVMGraphExplorerDragDropOp> New(const FRigVMExplorerElementKey& InElement, TObjectPtr<URigVMBlueprint> InBlueprint);
+	static TSharedRef<FRigVMGraphExplorerDragDropOp> New(const FRigVMExplorerElementKey& InElement, TRigVMAssetInterfacePtr InBlueprint);
 
 	const FRigVMExplorerElementKey& GetElement() const
 	{
 		return Element;
 	}
 
+	UE_DEPRECATED(5.7, "Please use TRigVMAssetInterfacePtr GetRigVMAsset() const")
 	TObjectPtr<URigVMBlueprint> GetBlueprint() const
+	{
+		return Cast<URigVMBlueprint>(GetRigVMAsset().GetObject());
+	}
+
+	TRigVMAssetInterfacePtr GetRigVMAsset() const
 	{
 		return SourceBlueprint;
 	}
@@ -50,7 +60,7 @@ public:
 private:
 
 	FRigVMExplorerElementKey Element;
-	TObjectPtr<URigVMBlueprint> SourceBlueprint;
+	TRigVMAssetInterfacePtr SourceBlueprint;
 };
 
 class FRigVMEditorGraphExplorerCommands : public TCommands<FRigVMEditorGraphExplorerCommands>
diff --git a/Engine/Plugins/Runtime/RigVM/Source/RigVMEditor/Public/Widgets/SRigVMEditorSelectedDebugObjectWidget.h b/Engine/Plugins/Runtime/RigVM/Source/RigVMEditor/Public/Widgets/SRigVMEditorSelectedDebugObjectWidget.h
--- a/Engine/Plugins/Runtime/RigVM/Source/RigVMEditor/Public/Widgets/SRigVMEditorSelectedDebugObjectWidget.h
+++ b/Engine/Plugins/Runtime/RigVM/Source/RigVMEditor/Public/Widgets/SRigVMEditorSelectedDebugObjectWidget.h
@@ -62,7 +62,7 @@ public:
 	UE_API void AddDebugObject(UObject* TestObject, const FString& TestObjectName = FString());
 
 private:
-	UE_API UBlueprint* GetBlueprintObj() const;
+	UE_API TRigVMAssetInterfacePtr GetRigVMAsset() const;
 
 	/** Creates a list of all debug objects **/
 	UE_API void GenerateDebugObjectInstances(bool bRestoreSelection);
diff --git a/Engine/Plugins/Runtime/RigVM/Source/RigVMEditor/Public/Widgets/SRigVMExecutionStackView.h b/Engine/Plugins/Runtime/RigVM/Source/RigVMEditor/Public/Widgets/SRigVMExecutionStackView.h
--- a/Engine/Plugins/Runtime/RigVM/Source/RigVMEditor/Public/Widgets/SRigVMExecutionStackView.h
+++ b/Engine/Plugins/Runtime/RigVM/Source/RigVMEditor/Public/Widgets/SRigVMExecutionStackView.h
@@ -5,10 +5,10 @@
 #include "CoreMinimal.h"
 #include "EdGraph/RigVMEdGraph.h"
 #include "RigVMHost.h"
-#include "RigVMBlueprint.h"
+#include "RigVMBlueprintLegacy.h"
 #include "Widgets/Views/STreeView.h"
 #include "RigVMCore/RigVM.h"
-#include "RigVMBlueprint.h"
+#include "RigVMBlueprintLegacy.h"
 
 #define UE_API RIGVMEDITOR_API
 
@@ -34,7 +34,7 @@ class FRigStackEntry : public TSharedFromThis<FRigStackEntry>
 public:
 	UE_API FRigStackEntry(int32 InEntryIndex, ERigStackEntry::Type InEntryType, int32 InInstructionIndex, ERigVMOpCode InOpCode, const FString& InLabel, const FRigVMASTProxy& InProxy);
 
-	UE_API TSharedRef<ITableRow> MakeTreeRowWidget(const TSharedRef<STableViewBase>& InOwnerTable, TSharedRef<FRigStackEntry> InEntry, TSharedRef<FUICommandList> InCommandList, TSharedPtr<SRigVMExecutionStackView> InStackView, TWeakObjectPtr<URigVMBlueprint> InBlueprint);
+	UE_API TSharedRef<ITableRow> MakeTreeRowWidget(const TSharedRef<STableViewBase>& InOwnerTable, TSharedRef<FRigStackEntry> InEntry, TSharedRef<FUICommandList> InCommandList, TSharedPtr<SRigVMExecutionStackView> InStackView, TWeakObjectPtr<UObject> InBlueprint);
 
 	int32 EntryIndex;
 	ERigStackEntry::Type EntryType;
@@ -51,11 +51,11 @@ class SRigStackItem : public STableRow<TSharedPtr<FRigStackEntry>>
 	SLATE_BEGIN_ARGS(SRigStackItem) {}
 	SLATE_END_ARGS()
 
-	void Construct(const FArguments& InArgs, const TSharedRef<STableViewBase>& OwnerTable, TSharedRef<FRigStackEntry> InStackEntry, TSharedRef<FUICommandList> InCommandList, TWeakObjectPtr<URigVMBlueprint> InBlueprint);
+	void Construct(const FArguments& InArgs, const TSharedRef<STableViewBase>& OwnerTable, TSharedRef<FRigStackEntry> InStackEntry, TSharedRef<FUICommandList> InCommandList, TWeakObjectPtr<UObject> InBlueprint);
 
 private:
 	TWeakPtr<FRigStackEntry> WeakStackEntry;
-	TWeakObjectPtr<URigVMBlueprint> WeakBlueprint;
+	TWeakObjectPtr<UObject> WeakBlueprint;
 	TWeakPtr<FUICommandList> WeakCommandList;
 
 	mutable double MicroSeconds;
@@ -95,7 +95,7 @@ private:
 	UE_API void BindCommands();
 
 	/** Make a row widget for the table */
-	UE_API TSharedRef<ITableRow> MakeTableRowWidget(TSharedPtr<FRigStackEntry> InItem, const TSharedRef<STableViewBase>& OwnerTable, TWeakObjectPtr<URigVMBlueprint> InBlueprint);
+	UE_API TSharedRef<ITableRow> MakeTableRowWidget(TSharedPtr<FRigStackEntry> InItem, const TSharedRef<STableViewBase>& OwnerTable, TWeakObjectPtr<UObject> InBlueprint);
 
 	/** Get children for the tree */
 	UE_API void HandleGetChildrenForTree(TSharedPtr<FRigStackEntry> InItem, TArray<TSharedPtr<FRigStackEntry>>& OutChildren);
@@ -140,7 +140,7 @@ private:
 	TSharedPtr<FUICommandList> CommandList;
 
 	TWeakPtr<IRigVMEditor> WeakEditor;
-	TWeakObjectPtr<URigVMBlueprint> RigVMBlueprint;
+	TWeakObjectPtr<UObject> WeakRigVMBlueprint;
 
 	TArray<TSharedPtr<FRigStackEntry>> Operators;
 
diff --git a/Engine/Plugins/Runtime/RigVM/Source/RigVMEditor/Public/Widgets/SRigVMGraphFunctionBulkEditWidget.h b/Engine/Plugins/Runtime/RigVM/Source/RigVMEditor/Public/Widgets/SRigVMGraphFunctionBulkEditWidget.h
--- a/Engine/Plugins/Runtime/RigVM/Source/RigVMEditor/Public/Widgets/SRigVMGraphFunctionBulkEditWidget.h
+++ b/Engine/Plugins/Runtime/RigVM/Source/RigVMEditor/Public/Widgets/SRigVMGraphFunctionBulkEditWidget.h
@@ -3,7 +3,7 @@
 #pragma once
 
 #include "CoreMinimal.h"
-#include "RigVMBlueprint.h"
+#include "RigVMBlueprintLegacy.h"
 #include "RigVMModel/Nodes/RigVMLibraryNode.h"
 #include "Dialogs/Dialogs.h"
 #include "IAssetTypeActions.h"
@@ -18,7 +18,7 @@ public:
 	SLATE_BEGIN_ARGS(SRigVMGraphFunctionBulkEditWidget) {}
 	SLATE_END_ARGS()
 
-	UE_API void Construct(const FArguments& InArgs, URigVMBlueprint* InBlueprint, URigVMController* InController, URigVMLibraryNode* InFunction, ERigVMControllerBulkEditType InEditType);
+	UE_API void Construct(const FArguments& InArgs, TRigVMAssetInterfacePtr InBlueprint, URigVMController* InController, URigVMLibraryNode* InFunction, ERigVMControllerBulkEditType InEditType);
 
 private:
 
@@ -31,7 +31,7 @@ private:
 	/** Show a progress bar and load all of the assets */
 	UE_API void LoadAffectedAssets();
 
-	URigVMBlueprint* Blueprint;
+	TRigVMAssetInterfacePtr Blueprint;
 	URigVMController* Controller;
 	URigVMLibraryNode* Function;
 	ERigVMControllerBulkEditType EditType;
@@ -49,7 +49,7 @@ public:
 	{
 	}
 
-	SLATE_ARGUMENT(URigVMBlueprint*, Blueprint)
+	SLATE_ARGUMENT(TRigVMAssetInterfacePtr, Blueprint)
 	SLATE_ARGUMENT(URigVMController*, Controller)
 	SLATE_ARGUMENT(URigVMLibraryNode*, Function)
 	SLATE_ARGUMENT(ERigVMControllerBulkEditType, EditType)
diff --git a/Engine/Plugins/Runtime/RigVM/Source/RigVMEditor/Public/Widgets/SRigVMGraphNode.h b/Engine/Plugins/Runtime/RigVM/Source/RigVMEditor/Public/Widgets/SRigVMGraphNode.h
--- a/Engine/Plugins/Runtime/RigVM/Source/RigVMEditor/Public/Widgets/SRigVMGraphNode.h
+++ b/Engine/Plugins/Runtime/RigVM/Source/RigVMEditor/Public/Widgets/SRigVMGraphNode.h
@@ -9,7 +9,7 @@
 #include "Widgets/Text/STextBlock.h"
 #include "Widgets/Views/STableRow.h"
 #include "RigVMModel/RigVMPin.h"
-#include "RigVMBlueprint.h"
+#include "RigVMBlueprintLegacy.h"
 #include "Overrides/SOverrideStatusWidget.h"
 
 #define UE_API RIGVMEDITOR_API
@@ -129,7 +129,7 @@ protected:
 
 	UE_API virtual TOptional<FSlateColor> GetHighlightColor(const SGraphPin* InGraphPin) const override;
 	UE_API FSlateColor GetVariableLabelTextColor(TWeakObjectPtr<URigVMFunctionReferenceNode> FunctionReferenceNode, FName InVariableName) const;
-	UE_API FText GetVariableLabelTooltipText(TWeakObjectPtr<URigVMBlueprint> InBlueprint, FName InVariableName) const;
+	UE_API FText GetVariableLabelTooltipText(TWeakInterfacePtr<IRigVMAssetInterface> InBlueprint, FName InVariableName) const;
 
 	UE_API FReply HandleAddArrayElement(FString InModelPinPath);
 
@@ -173,7 +173,7 @@ protected:
 	/** Cache the node title so we can invalidate it */
 	TSharedPtr<SNodeTitle> NodeTitle;
 
-	TWeakObjectPtr<URigVMBlueprint> Blueprint;
+	TWeakInterfacePtr<IRigVMAssetInterface> Blueprint;
 
 	FVector2D LastHighDetailSize;
 
diff --git a/Engine/Plugins/Runtime/RigVM/Source/RigVMEditor/Public/Widgets/SRigVMGraphPinUserDataPath.h b/Engine/Plugins/Runtime/RigVM/Source/RigVMEditor/Public/Widgets/SRigVMGraphPinUserDataPath.h
--- a/Engine/Plugins/Runtime/RigVM/Source/RigVMEditor/Public/Widgets/SRigVMGraphPinUserDataPath.h
+++ b/Engine/Plugins/Runtime/RigVM/Source/RigVMEditor/Public/Widgets/SRigVMGraphPinUserDataPath.h
@@ -9,7 +9,7 @@
 #include "Widgets/Input/SButton.h"
 #include "SGraphPin.h"
 #include "RigVMModel/RigVMPin.h"
-#include "RigVMBlueprint.h"
+#include "RigVMBlueprintLegacy.h"
 #include "RigVMCore/RigVMAssetUserData.h"
 #include "IPropertyAccessEditor.h"
 
@@ -42,7 +42,7 @@ protected:
 	UE_API void FillUserDataPathMenu( FMenuBuilder& InMenuBuilder, FString InParentPath );
 	UE_API void HandleSetUserDataPath(FString InUserDataPath);
 
-	UE_API URigVMBlueprint* GetBlueprint() const;
+	UE_API TRigVMAssetInterfacePtr GetRigVMAsset() const;
 	UE_API FString GetUserDataNameSpace(URigVMPin* ModelPin) const;
 	UE_API FString GetUserDataNameSpace() const;
 	UE_API const UNameSpacedUserData* GetUserDataObject() const;
diff --git a/Engine/Plugins/Runtime/RigVM/Source/RigVMEditor/Public/Widgets/SRigVMGraphPinVariableBinding.h b/Engine/Plugins/Runtime/RigVM/Source/RigVMEditor/Public/Widgets/SRigVMGraphPinVariableBinding.h
--- a/Engine/Plugins/Runtime/RigVM/Source/RigVMEditor/Public/Widgets/SRigVMGraphPinVariableBinding.h
+++ b/Engine/Plugins/Runtime/RigVM/Source/RigVMEditor/Public/Widgets/SRigVMGraphPinVariableBinding.h
@@ -9,7 +9,7 @@
 #include "RigVMModel/RigVMPin.h"
 #include "RigVMModel/Nodes/RigVMFunctionReferenceNode.h"
 #include "IPropertyAccessEditor.h"
-#include "RigVMBlueprint.h"
+#include "RigVMBlueprintLegacy.h"
 
 #define UE_API RIGVMEDITOR_API
 
@@ -21,14 +21,16 @@ public:
 	: _ModelPins()
     , _FunctionReferenceNode(nullptr)
     , _InnerVariableName(NAME_None)
-    , _Blueprint(nullptr)
+    , _Asset(nullptr)
+	, _Blueprint(nullptr)
 	, _CanRemoveBinding(true)
 	{}
 
 		SLATE_ARGUMENT(TArray<URigVMPin*>, ModelPins)
 		SLATE_ARGUMENT(URigVMFunctionReferenceNode*, FunctionReferenceNode)
 		SLATE_ARGUMENT(FName, InnerVariableName)
-		SLATE_ARGUMENT(URigVMBlueprint*, Blueprint)
+		SLATE_ARGUMENT(TRigVMAssetInterfacePtr, Asset)
+		SLATE_ARGUMENT_DEPRECATED(URigVMBlueprint*, Blueprint, 5.7, "Please use Asset instead")
 		SLATE_ARGUMENT(bool, CanRemoveBinding)
 
 	SLATE_END_ARGS()
@@ -52,7 +54,7 @@ protected:
 	TArray<URigVMPin*> ModelPins;
 	URigVMFunctionReferenceNode* FunctionReferenceNode;
 	FName InnerVariableName;
-	URigVMBlueprint* Blueprint;
+	TRigVMAssetInterfacePtr Blueprint;
 	FPropertyBindingWidgetArgs BindingArgs;
 	bool bCanRemoveBinding;
 };
@@ -64,7 +66,8 @@ public:
 	SLATE_BEGIN_ARGS(SRigVMGraphPinVariableBinding){}
 
 		SLATE_ARGUMENT(TArray<URigVMPin*>, ModelPins)
-		SLATE_ARGUMENT(URigVMBlueprint*, Blueprint)
+		SLATE_ARGUMENT_DEPRECATED(URigVMBlueprint*, Blueprint, 5.7, "Please use Asset instead")
+		SLATE_ARGUMENT(TRigVMAssetInterfacePtr, Asset)
 
 	SLATE_END_ARGS()
 
@@ -77,7 +80,7 @@ protected:
 	//~ End SGraphPin Interface
 
 	TArray<URigVMPin*> ModelPins;
-	URigVMBlueprint* Blueprint;
+	TRigVMAssetInterfacePtr Blueprint;
 };
 
 #undef UE_API
diff --git a/Engine/Plugins/Runtime/RigVM/Source/RigVMEditor/Public/Widgets/SRigVMLogWidget.h b/Engine/Plugins/Runtime/RigVM/Source/RigVMEditor/Public/Widgets/SRigVMLogWidget.h
--- a/Engine/Plugins/Runtime/RigVM/Source/RigVMEditor/Public/Widgets/SRigVMLogWidget.h
+++ b/Engine/Plugins/Runtime/RigVM/Source/RigVMEditor/Public/Widgets/SRigVMLogWidget.h
@@ -5,7 +5,7 @@
 #include "Widgets/Layout/SBox.h"
 #include "Model/MessageLogListingModel.h"
 #include "Presentation/MessageLogListingViewModel.h"
-#include "RigVMBlueprint.h"
+#include "RigVMBlueprintLegacy.h"
 
 class IMessageLogListing;
 
@@ -40,8 +40,8 @@ public:
 	void UnbindLog(const FName& InLogName);
 	void BindLog(TSharedPtr<IMessageLogListing> InListing);
 	void UnbindLog(TSharedPtr<IMessageLogListing> InListing);
-	void BindLog(const URigVMBlueprint* InRigVMBlueprint);
-	void UnbindLog(const URigVMBlueprint* InRigVMBlueprint);
+	void BindLog(const TRigVMAssetInterfacePtr InRigVMBlueprint);
+	void UnbindLog(const TRigVMAssetInterfacePtr InRigVMBlueprint);
 
 	TSharedRef<FMessageLogListingModel> GetListing() const
 	{
diff --git a/Engine/Plugins/Runtime/RigVM/Source/RigVMEditor/Public/Widgets/SRigVMSwapFunctionsWidget.h b/Engine/Plugins/Runtime/RigVM/Source/RigVMEditor/Public/Widgets/SRigVMSwapFunctionsWidget.h
--- a/Engine/Plugins/Runtime/RigVM/Source/RigVMEditor/Public/Widgets/SRigVMSwapFunctionsWidget.h
+++ b/Engine/Plugins/Runtime/RigVM/Source/RigVMEditor/Public/Widgets/SRigVMSwapFunctionsWidget.h
@@ -272,7 +272,7 @@ public:
 private:
 
 	URigVMFunctionReferenceNode* GetReferenceNode(const TSharedRef<FRigVMTreePhase>& InPhase) const;
-	virtual URigVMBlueprint* GetBlueprint(const TSharedRef<FRigVMTreePhase>& InPhase) const override;
+	virtual TRigVMAssetInterfacePtr GetRigVMAsset(const TSharedRef<FRigVMTreePhase>& InPhase) const override;
 	
 	FString ObjectPath;
 	FRigVMGraphFunctionIdentifier Identifier;
