RigVM: Fix RigVMFindReferences nonunity compilation error
#rnx

diff --git a/Engine/Plugins/Runtime/RigVM/Source/RigVMEditor/Public/Editor/RigVMFindReferences.h b/Engine/Plugins/Runtime/RigVM/Source/RigVMEditor/Public/Editor/RigVMFindReferences.h
--- a/Engine/Plugins/Runtime/RigVM/Source/RigVMEditor/Public/Editor/RigVMFindReferences.h
+++ b/Engine/Plugins/Runtime/RigVM/Source/RigVMEditor/Public/Editor/RigVMFindReferences.h
@@ -13,6 +13,7 @@
 #define UE_API RIGVMEDITOR_API
 
 class IRigVMAssetInterface;
+class URigVMBlueprint;
 class SWidget;
 class SDockTab;
 class UBlueprint;
