RigVM module unloading optimizations

- Optimize FRigVMRegistry_NoLock::OnModulesUnloadedInternal by parallelizing function and type lookups
- Optimize FRigVMTemplate::HandleTypeRemoval by using a type index to permutation index lookup cache

#tests RigVM test suite, game client test on console
#rnx

diff --git a/Engine/Plugins/Runtime/RigVM/Source/RigVM/Private/RigVMCore/RigVMRegistry.cpp b/Engine/Plugins/Runtime/RigVM/Source/RigVM/Private/RigVMCore/RigVMRegistry.cpp
--- a/Engine/Plugins/Runtime/RigVM/Source/RigVM/Private/RigVMCore/RigVMRegistry.cpp
+++ b/Engine/Plugins/Runtime/RigVM/Source/RigVM/Private/RigVMCore/RigVMRegistry.cpp
@@ -7,6 +7,7 @@
 #include "RigVMTypeUtils.h"
 #include "RigVMStringUtils.h"
 #include "RigVMModule.h"
+#include "Async/ParallelFor.h"
 #include "Animation/AttributeTypes.h"
 #include "AssetRegistry/AssetRegistryModule.h"
 #include "Engine/UserDefinedEnum.h"
@@ -926,7 +927,10 @@ bool FRigVMRegistry_NoLock::OnModulesUnloaded_NoLock(TConstArrayView<FName> Modu
 
 bool FRigVMRegistry_NoLock::OnModulesUnloadedInternal(const TArray<FString>& ModulePrefixes)
 {
+	TRACE_CPUPROFILER_EVENT_SCOPE(FRigVMRegistry_NoLock::OnModulesUnloadedInternal);
+
 	bool bRegistryChanged = false;
+	FCriticalSection ParallelForDataLock;
 
 	// remove the dispatch factories defined in the plugin
 	TArray<UScriptStruct*> FactoriesToRemove;
@@ -983,46 +987,68 @@ bool FRigVMRegistry_NoLock::OnModulesUnloadedInternal(const TArray<FString>& Mod
 		}
 	}
 
-	// remove the functions defined in the plugin
-	for (int32 Index = 0; Index < Functions.Num(); Index++)
+	// remove the functions defined in the plugin - parellize the search as it can last for a significant time
 	{
-		if (!Functions[Index].IsValid())
-		{
-			continue;
-		}
+		TRACE_CPUPROFILER_EVENT_SCOPE(FRigVMRegistry_NoLock::OnModulesUnloadedInternal::Functions);
 
-		if (Functions[Index].Struct == nullptr)
+		TArray<int32> FunctionIndicesToRemove;
+		ParallelFor(Functions.Num(), [this, &ModulePrefixes, &ParallelForDataLock, &FunctionIndicesToRemove](int32 Index)
 		{
-			continue;
-		}
+			if (!Functions[Index].IsValid())
+			{
+				return;
+			}
 
-		if (!IsWithinPlugin(ModulePrefixes, Functions[Index].Struct))
-		{
-			continue;
-		}
+			if (Functions[Index].Struct == nullptr)
+			{
+				return;
+			}
+
+			if (!IsWithinPlugin(ModulePrefixes, Functions[Index].Struct))
+			{
+				return;
+			}
+
+			FScopeLock Lock(&ParallelForDataLock);
+			FunctionIndicesToRemove.Add(Index);
+		});
 
-		if (RemoveFunction_NoLock(Index))
+		for (int32 Index : FunctionIndicesToRemove)
 		{
-			bRegistryChanged = true;
+			if (RemoveFunction_NoLock(Index))
+			{
+				bRegistryChanged = true;
+			}
 		}
-	}	
+	}
 
-	// remove the types defined in the plugin
-	for (int32 Index = 0; Index < Types.Num(); Index++)
+	// remove the types defined in the plugin - parellize the search as it can last for a very significant time
 	{
-		if (Types[Index].IsUnknownType())
-		{
-			continue;
-		}
+		TRACE_CPUPROFILER_EVENT_SCOPE(FRigVMRegistry_NoLock::OnModulesUnloadedInternal::Types);
 
-		if (!IsWithinPlugin(ModulePrefixes, Types[Index].Type.CPPTypeObject))
+		TArray<int32> TypeIndicesToRemove;
+		ParallelFor(Types.Num(), [this, &ModulePrefixes, &ParallelForDataLock, &TypeIndicesToRemove](int32 Index)
 		{
-			continue;
-		}
-		
-		if (RemoveType_NoLock(Index))
+			if (Types[Index].IsUnknownType())
+			{
+				return;
+			}
+			
+			if (!IsWithinPlugin(ModulePrefixes, Types[Index].Type.CPPTypeObject))
+			{
+				return;
+			}
+
+			FScopeLock Lock(&ParallelForDataLock);
+			TypeIndicesToRemove.Add(Index);
+		});
+
+		for (int32 Index : TypeIndicesToRemove)
 		{
-			bRegistryChanged = true;
+			if (RemoveType_NoLock(Index))
+			{
+				bRegistryChanged = true;
+			}
 		}
 	}
 
diff --git a/Engine/Plugins/Runtime/RigVM/Source/RigVM/Private/RigVMCore/RigVMTemplate.cpp b/Engine/Plugins/Runtime/RigVM/Source/RigVM/Private/RigVMCore/RigVMTemplate.cpp
--- a/Engine/Plugins/Runtime/RigVM/Source/RigVM/Private/RigVMCore/RigVMTemplate.cpp
+++ b/Engine/Plugins/Runtime/RigVM/Source/RigVM/Private/RigVMCore/RigVMTemplate.cpp
@@ -976,6 +976,8 @@ FRigVMTemplate::FRigVMTemplate(UScriptStruct* InStruct, const FString& InTemplat
 		if (InFunctionIndex != INDEX_NONE)
 		{
 			Permutations.Add(InFunctionIndex);
+			UpdateTypesIndicesToPermutationIndices(Permutations.Num() - 1, InRegistry);
+
 			for (const FRigVMTemplateArgument& Argument : Arguments)
 			{
 				check(Argument.TypeIndices.Num() == 1);
@@ -1585,6 +1587,7 @@ bool FRigVMTemplate::Merge(const FRigVMTemplate& InOther, const FRigVMRegistryHa
 
 	Permutations.Add(InOther.Permutations[0]);
 
+	UpdateTypesIndicesToPermutationIndices(Permutations.Num() - 1, InRegistry);
 	UpdateTypesHashToPermutation(Permutations.Num()-1, InRegistry);
 	
 	return true;
@@ -1799,6 +1802,7 @@ const FRigVMFunction* FRigVMTemplate::GetOrCreatePermutation_NoLock(int32 InInde
 			}
 
 			Permutations[InIndex] = FunctionIndex;
+			UpdateTypesIndicesToPermutationIndices(InIndex, InRegistry);
 
 			if (FunctionIndex  == InRegistry->Functions.Num())
 			{
@@ -2466,6 +2470,7 @@ bool FRigVMTemplate::UpdateArgumentTypes_Impl(
 				}
 
 				Permutations.Add(INDEX_NONE);
+				UpdateTypesIndicesToPermutationIndices(Permutations.Num() - 1, InRegistry);
 				TypesHashToPermutation.Add(TypeHash, Permutations.Num()-1);
 			}
 			else
@@ -2587,10 +2592,14 @@ bool FRigVMTemplate::HandlePropagatedArgumentType(const TRigVMTypeIndex InTypeIn
 			for(int32 PermutationIndex = 0; PermutationIndex < Permutations.Num(); PermutationIndex++)
 			{
 				Permutations[PermutationIndex] = INDEX_NONE;
+				UpdateTypesIndicesToPermutationIndices(PermutationIndex, InRegistry);
 			}
 			for(int32 FunctionIndex : ValidFunctionIndices)
 			{
-				Permutations[OwnerRegistry->Functions[FunctionIndex].PermutationIndex] = FunctionIndex;
+				const int32 PermutationIndex = OwnerRegistry->Functions[FunctionIndex].PermutationIndex;
+				Permutations[PermutationIndex] = FunctionIndex;
+
+				UpdateTypesIndicesToPermutationIndices(PermutationIndex, InRegistry);
 			}
 		}
 	}
@@ -2600,26 +2609,35 @@ bool FRigVMTemplate::HandlePropagatedArgumentType(const TRigVMTypeIndex InTypeIn
 
 void FRigVMTemplate::HandleTypeRemoval(TRigVMTypeIndex InTypeIndex, FRigVMRegistryHandle& InRegistry)
 {
+	TRACE_CPUPROFILER_EVENT_SCOPE(FRigVMTemplate::HandleTypeRemoval);
+
 	check(InRegistry.GetRegistry() == OwnerRegistry);
 
 	InvalidateHash();
 
-	const int32 CurrentNumPermutations = NumPermutations_NoLock(InRegistry);
-	for (int32 PermutationIndex = 0; PermutationIndex < CurrentNumPermutations; PermutationIndex++)
+	const TArray<int32>* PermutationIndices = TypesIndicesToPermutationIndices.Find(InTypeIndex);
+	if (PermutationIndices)
 	{
-		const FRigVMTemplateTypeMap TypeMap = GetTypesForPermutation_NoLock(PermutationIndex, InRegistry);
-
-		for (const TPair<FName, TRigVMTypeIndex>& Pair : TypeMap)
+		for (int32 PermutationIndex : *PermutationIndices)
 		{
-			if (Pair.Value == InTypeIndex)
+			TRACE_CPUPROFILER_EVENT_SCOPE(TryRemoveAtPermutationIndex);
+
+			const FRigVMTemplateTypeMap TypeMap = GetTypesForPermutation_NoLock(PermutationIndex, InRegistry);
+
+			for (const TPair<FName, TRigVMTypeIndex>& Pair : TypeMap)
 			{
-				if (const FRigVMFunction* Permutation = GetPermutation_NoLock(PermutationIndex, InRegistry))
+				TRACE_CPUPROFILER_EVENT_SCOPE(TryRemoveAtPermutationIndex::Pair);
+
+				if (Pair.Value == InTypeIndex)
 				{
-					const int32 FunctionIndex = Permutation->Index;
-					InRegistry->RemoveFunction_NoLock(FunctionIndex);
-					Permutations[PermutationIndex] = INDEX_NONE;
+					if (const FRigVMFunction* Permutation = GetPermutation_NoLock(PermutationIndex, InRegistry))
+					{
+						const int32 FunctionIndex = Permutation->Index;
+						InRegistry->RemoveFunction_NoLock(FunctionIndex);
+						Permutations[PermutationIndex] = INDEX_NONE;
+					}
+					break;
 				}
-				break;
 			}
 		}
 	}
@@ -2707,6 +2725,24 @@ void FRigVMTemplate::UpdateTypesHashToPermutation(const int32 InPermutation, con
 	TypesHashToPermutation.Add(TypeHash, InPermutation);
 }
 
+
+void FRigVMTemplate::UpdateTypesIndicesToPermutationIndices(const int32 InPermutation, const FRigVMRegistryHandle& InRegistry)
+{
+	check(InRegistry.GetRegistry() == OwnerRegistry);
+
+	TRACE_CPUPROFILER_EVENT_SCOPE(FRigVMTemplate::UpdateTypesIndicesToPermutationIndices);
+
+	if (Permutations.IsValidIndex(InPermutation) && Permutations[InPermutation] != INDEX_NONE)
+	{
+		const FRigVMTemplateTypeMap TypeMap = GetTypesForPermutation_NoLock(InPermutation, InRegistry);
+		for (const TPair<FName, TRigVMTypeIndex>& Pair : TypeMap)
+		{
+			TArray<int32>& PermutationIndices = TypesIndicesToPermutationIndices.FindOrAdd(Pair.Value);
+			PermutationIndices.Add(InPermutation);
+		}
+	}
+}
+
 TArray<FRigVMTemplateArgumentInfo> FRigVMTemplate::GetFlattenedArgumentInfos_NoLock(FRigVMRegistryHandle& InRegistry) const
 {
 	check(InRegistry.GetRegistry() == OwnerRegistry);
diff --git a/Engine/Plugins/Runtime/RigVM/Source/RigVM/Public/RigVMCore/RigVMTemplate.h b/Engine/Plugins/Runtime/RigVM/Source/RigVM/Public/RigVMCore/RigVMTemplate.h
--- a/Engine/Plugins/Runtime/RigVM/Source/RigVM/Public/RigVMCore/RigVMTemplate.h
+++ b/Engine/Plugins/Runtime/RigVM/Source/RigVM/Public/RigVMCore/RigVMTemplate.h
@@ -599,6 +599,8 @@ public:
 
 	RIGVM_API void UpdateTypesHashToPermutation(const int32 InPermutation, const FRigVMRegistryHandle& InRegistry);
 
+	RIGVM_API void UpdateTypesIndicesToPermutationIndices(const int32 InPermutation, const FRigVMRegistryHandle& InRegistry);
+
 	RIGVM_API friend uint32 GetTypeHash(const FRigVMTemplate& InTemplate);
 	RIGVM_API friend uint32 GetTypeHash_NoLock(const FRigVMTemplate& InTemplate, const FRigVMRegistryHandle& InRegistry);
 
@@ -641,6 +643,7 @@ private:
 	mutable TArray<FRigVMExecuteArgument> ExecuteArguments;
 	TArray<int32> Permutations;
 	TMap<uint32, int32> TypesHashToPermutation;
+	TMap<TRigVMTypeIndex, TArray<int32>> TypesIndicesToPermutationIndices;
 	mutable uint32 Hash;
 	FRigVMRegistry_NoLock* OwnerRegistry;
 
