RigVM: Fix compiler to inline instruction to operand map correctly

#rnx

diff --git a/Engine/Plugins/Runtime/RigVM/Source/RigVM/Private/RigVMCore/RigVMByteCode.cpp b/Engine/Plugins/Runtime/RigVM/Source/RigVM/Private/RigVMCore/RigVMByteCode.cpp
--- a/Engine/Plugins/Runtime/RigVM/Source/RigVM/Private/RigVMCore/RigVMByteCode.cpp
+++ b/Engine/Plugins/Runtime/RigVM/Source/RigVM/Private/RigVMCore/RigVMByteCode.cpp
@@ -237,7 +237,8 @@ FRigVMInstructionArray::FRigVMInstructionArray(const FRigVMByteCode& InByteCode,
 			}
 		}
 
-		Instructions.Add(FRigVMInstruction(OpCode, ByteIndex, OperandAlignment));
+		const int32 InstructionIndex = Instructions.Num(); 
+		Instructions.Add(FRigVMInstruction(OpCode, InstructionIndex, ByteIndex, OperandAlignment));
 		ByteIndex += InByteCode.GetOpNumBytesAt(ByteIndex, true);
 	}
 }
@@ -309,10 +310,9 @@ void FRigVMByteCode::Save(FArchive& Ar)
 
 	for (int32 InstructionIndex = 0; InstructionIndex < InstructionCount; InstructionIndex++)
 	{
-		FRigVMInstruction Instruction;
+		const FRigVMInstruction& Instruction = Instructions[InstructionIndex];
 		ERigVMOpCode OpCode = ERigVMOpCode::Invalid;
-
-		Instruction = Instructions[InstructionIndex];
+		
 		OpCode = Instruction.OpCode;
 		Ar << OpCode;
 
@@ -1126,9 +1126,10 @@ int32 FRigVMByteCode::AddCopyOp(const FRigVMOperand& InSource, const FRigVMOpera
 	const int32 ByteIndex = AddOp(Op);
 
 #if WITH_EDITORONLY_DATA
+	const FRigVMInstruction Instruction(Op.OpCode, NumInstructions - 1, ByteIndex);
 	const FRigVMOperandArray InputOperands(&InSource, 1);
 	const FRigVMOperandArray OutputOperands(&InTarget, 1);
-	SetOperandsForInstruction(GetNumInstructions() - 1, InputOperands, OutputOperands);
+	SetOperandsForInstruction(Instruction, InputOperands, OutputOperands);
 #endif
 	
 	return ByteIndex;
@@ -1159,12 +1160,13 @@ int32 FRigVMByteCode::AddEqualsOp(const FRigVMOperand& InA, const FRigVMOperand&
 	const int32 ByteIndex = AddOp(Op);
 	
 #if WITH_EDITORONLY_DATA
+	const FRigVMInstruction Instruction(Op.OpCode, NumInstructions - 1, ByteIndex);
 	TArray<FRigVMOperand> Inputs;
 	Inputs.Add(InA);
 	Inputs.Add(InB);
 	const FRigVMOperandArray InputOperands(&Inputs[0], Inputs.Num());
 	const FRigVMOperandArray OutputOperands(&InResult, 1);
-	SetOperandsForInstruction(GetNumInstructions() - 1, InputOperands, OutputOperands);
+	SetOperandsForInstruction(Instruction, InputOperands, OutputOperands);
 #endif
 
 	return ByteIndex;
@@ -1176,12 +1178,13 @@ int32 FRigVMByteCode::AddNotEqualsOp(const FRigVMOperand& InA, const FRigVMOpera
 	const int32 ByteIndex = AddOp(Op);
 
 #if WITH_EDITORONLY_DATA
+	const FRigVMInstruction Instruction(Op.OpCode, NumInstructions - 1, ByteIndex);
 	TArray<FRigVMOperand> Inputs;
 	Inputs.Add(InA);
 	Inputs.Add(InB);
 	const FRigVMOperandArray InputOperands(&Inputs[0], Inputs.Num());
 	const FRigVMOperandArray OutputOperands(&InResult, 1);
-	SetOperandsForInstruction(GetNumInstructions() - 1, InputOperands, OutputOperands);
+	SetOperandsForInstruction(Instruction, InputOperands, OutputOperands);
 #endif
 
 	return ByteIndex;
@@ -2134,15 +2137,18 @@ uint32 FRigVMByteCode::GetCallstackHash(const TArrayView<TWeakObjectPtr<UObject>
 	return Hash;
 }
 
-void FRigVMByteCode::SetOperandsForInstruction(int32 InInstructionIndex, const FRigVMOperandArray& InputOperands,
+void FRigVMByteCode::SetOperandsForInstruction(const FRigVMInstruction& InInstruction, const FRigVMOperandArray& InputOperands,
 	const FRigVMOperandArray& OutputOperands)
 {
-	if (InputOperandsPerInstruction.Num() <= InInstructionIndex)
+	const FRigVMOperandArray OperandsForInstruction = GetOperandsForOp(InInstruction);
+	const int32 InstructionIndex = InInstruction.Index;
+	
+	if (InputOperandsPerInstruction.Num() <= InstructionIndex)
 	{
-		InputOperandsPerInstruction.AddZeroed(1 + InInstructionIndex - InputOperandsPerInstruction.Num());
+		InputOperandsPerInstruction.AddZeroed(1 + InstructionIndex - InputOperandsPerInstruction.Num());
 	}
-	InputOperandsPerInstruction[InInstructionIndex].Reset();
-	InputOperandsPerInstruction[InInstructionIndex].Reserve(InputOperands.Num());
+	InputOperandsPerInstruction[InstructionIndex].Reset();
+	InputOperandsPerInstruction[InstructionIndex].Reserve(InputOperands.Num());
 
 	for(int32 OperandIndex = 0; OperandIndex < InputOperands.Num(); OperandIndex++)
 	{
@@ -2151,15 +2157,16 @@ void FRigVMByteCode::SetOperandsForInstruction(int32 InInstructionIndex, const F
 		{
 			continue;;
 		}
-		InputOperandsPerInstruction[InInstructionIndex].Add(InputOperands[OperandIndex]);
+		check(OperandsForInstruction.Contains(InputOperands[OperandIndex]));
+		InputOperandsPerInstruction[InstructionIndex].Add(InputOperands[OperandIndex]);
 	}
 
-	if (OutputOperandsPerInstruction.Num() <= InInstructionIndex)
+	if (OutputOperandsPerInstruction.Num() <= InstructionIndex)
 	{
-		OutputOperandsPerInstruction.AddZeroed(1 + InInstructionIndex - OutputOperandsPerInstruction.Num());
+		OutputOperandsPerInstruction.AddZeroed(1 + InstructionIndex - OutputOperandsPerInstruction.Num());
 	}
-	OutputOperandsPerInstruction[InInstructionIndex].Reset();
-	OutputOperandsPerInstruction[InInstructionIndex].Reserve(OutputOperands.Num());
+	OutputOperandsPerInstruction[InstructionIndex].Reset();
+	OutputOperandsPerInstruction[InstructionIndex].Reserve(OutputOperands.Num());
 
 	for(int32 OperandIndex = 0; OperandIndex < OutputOperands.Num(); OperandIndex++)
 	{
@@ -2168,7 +2175,8 @@ void FRigVMByteCode::SetOperandsForInstruction(int32 InInstructionIndex, const F
 		{
 			continue;;
 		}
-		OutputOperandsPerInstruction[InInstructionIndex].Add(OutputOperands[OperandIndex]);
+		check(OperandsForInstruction.Contains(OutputOperands[OperandIndex]));
+		OutputOperandsPerInstruction[InstructionIndex].Add(OutputOperands[OperandIndex]);
 	}
 }
 
diff --git a/Engine/Plugins/Runtime/RigVM/Source/RigVM/Private/RigVMCore/RigVMGraphFunctionDefinition.cpp b/Engine/Plugins/Runtime/RigVM/Source/RigVM/Private/RigVMCore/RigVMGraphFunctionDefinition.cpp
--- a/Engine/Plugins/Runtime/RigVM/Source/RigVM/Private/RigVMCore/RigVMGraphFunctionDefinition.cpp
+++ b/Engine/Plugins/Runtime/RigVM/Source/RigVM/Private/RigVMCore/RigVMGraphFunctionDefinition.cpp
@@ -776,7 +776,25 @@ bool FRigVMGraphFunctionData::PatchSharedArgumentOperandsIfRequired()
 			}
 		}
 
-		// step 5: add copy operations at the end of the byte code
+#if WITH_EDITOR
+		// step 5: update the input and output operand maps
+		for (TArray<FRigVMOperand>& InputOperands : CompilationData.ByteCode.InputOperandsPerInstruction)
+		{
+			for (FRigVMOperand& Operand : InputOperands)
+			{
+				UpdateOperand(Operand);
+			}
+		}
+		for (TArray<FRigVMOperand>& OutputOperands : CompilationData.ByteCode.OutputOperandsPerInstruction)
+		{
+			for (FRigVMOperand& Operand : OutputOperands)
+			{
+				UpdateOperand(Operand);
+			}
+		}
+#endif
+
+		// step 6: add copy operations at the end of the byte code
 		CompilationData.ByteCode.AddCopyOp(SourceOperand, TargetOperand);
 	}
 	
diff --git a/Engine/Plugins/Runtime/RigVM/Source/RigVM/Public/RigVMCore/RigVMByteCode.h b/Engine/Plugins/Runtime/RigVM/Source/RigVM/Public/RigVMCore/RigVMByteCode.h
--- a/Engine/Plugins/Runtime/RigVM/Source/RigVM/Public/RigVMCore/RigVMByteCode.h
+++ b/Engine/Plugins/Runtime/RigVM/Source/RigVM/Public/RigVMCore/RigVMByteCode.h
@@ -1038,13 +1038,17 @@ struct FRigVMInstruction
 {
 	GENERATED_USTRUCT_BODY()
 
-	FRigVMInstruction(ERigVMOpCode InOpCode = ERigVMOpCode::Invalid, int32 InByteCodeIndex = INDEX_NONE, uint8 InOperandAlignment = 0)
-		: ByteCodeIndex(InByteCodeIndex)
+	FRigVMInstruction(ERigVMOpCode InOpCode = ERigVMOpCode::Invalid, int32 InIndex = INDEX_NONE, int32 InByteCodeIndex = INDEX_NONE, uint8 InOperandAlignment = 0)
+		: Index(InIndex)
+		, ByteCodeIndex(InByteCodeIndex)
 		, OpCode(InOpCode)
 		, OperandAlignment(InOperandAlignment)
 	{
 	}
 
+	UPROPERTY()
+	int32 Index;
+
 	UPROPERTY()
 	int32 ByteCodeIndex;
 
@@ -1434,7 +1438,7 @@ public:
 		return FRigVMOperandArray();
 	}
 
-	UE_API void SetOperandsForInstruction(int32 InInstructionIndex, const FRigVMOperandArray& InputOperands, const FRigVMOperandArray& OutputOperands);
+	UE_API void SetOperandsForInstruction(const FRigVMInstruction& InInstruction, const FRigVMOperandArray& InputOperands, const FRigVMOperandArray& OutputOperands);
 
 #endif
 
@@ -1558,6 +1562,7 @@ private:
 	friend class FRigVMInvokeEntryTest;
 	friend class URigVM;
 	friend struct FRigVMCodeGenerator;
+	friend struct FRigVMGraphFunctionData;
 };
 
 #undef UE_API
diff --git a/Engine/Plugins/Runtime/RigVM/Source/RigVMDeveloper/Private/RigVMCompiler/RigVMCompiler.cpp b/Engine/Plugins/Runtime/RigVM/Source/RigVMDeveloper/Private/RigVMCompiler/RigVMCompiler.cpp
--- a/Engine/Plugins/Runtime/RigVM/Source/RigVMDeveloper/Private/RigVMCompiler/RigVMCompiler.cpp
+++ b/Engine/Plugins/Runtime/RigVM/Source/RigVMDeveloper/Private/RigVMCompiler/RigVMCompiler.cpp
@@ -2485,8 +2485,9 @@ bool URigVMCompiler::TraverseCallExtern(const FRigVMCallExternExprAST* InExpr, F
 		// setup the instruction
 		const int32 FunctionIndex = WorkData.VM->AddRigVMFunction(Function->GetName());
 		check(FunctionIndex != INDEX_NONE);
-		WorkData.VM->GetByteCode().AddExecuteOp(FunctionIndex, Operands, StartPredicateIndex, PredicateCount);
+		const int32 CallExternByteIndex = WorkData.VM->GetByteCode().AddExecuteOp(FunctionIndex, Operands, StartPredicateIndex, PredicateCount);
 		CallExternInstructionIndex = WorkData.VM->GetByteCode().GetNumInstructions() - 1;
+		const FRigVMInstruction CallExternInstruction(ERigVMOpCode::Execute, CallExternInstructionIndex, CallExternByteIndex);
 
 		// setup the branch infos for this call extern instruction
 		if(const TArray<FRigVMBranchInfo>* BranchInfosPtr = WorkData.BranchInfos.Find(Node))
@@ -2555,7 +2556,7 @@ bool URigVMCompiler::TraverseCallExtern(const FRigVMCallExternExprAST* InExpr, F
 		}
 
 		WorkData.VM->GetByteCode().SetOperandsForInstruction(
-			CallExternInstructionIndex,
+			CallExternInstruction,
 			FRigVMOperandArray(InputsOperands.GetData(), InputsOperands.Num()),
 			FRigVMOperandArray(OutputOperands.GetData(), OutputOperands.Num()));
 
@@ -2991,14 +2992,22 @@ bool URigVMCompiler::TraverseInlineFunction(const FRigVMInlineFunctionExprAST* I
 			{
 				return false;
 			}
+
+			const int32 FunctionInstructionIndex = i - InstructionIndexStart;
+			
 			const FRigVMInstruction& Instruction = Instructions[i];
 			const FRigVMOperandArray OperandArray = ByteCode.GetOperandsForOp(Instruction);
+			const FRigVMOperandArray FunctionOperandArray = FunctionByteCode.GetOperandsForOp(FunctionInstructions[FunctionInstructionIndex]);
+			check(FunctionOperandArray.Num() == OperandArray.Num());
+			
 			int32 OperandsIndex = ByteCode.GetFirstOperandByteIndex(Instruction);
 			for (int32 j=0; j<OperandArray.Num(); ++j)
 			{
-				FRigVMOperand* Operand = (FRigVMOperand*)(ByteCode.ByteCode.GetData() + OperandsIndex + j*sizeof(FRigVMOperand));
-				FRigVMOperand FunctionOperand = *Operand;
-				ERigVMMemoryType OriginalMemoryType = Operand->GetMemoryType();
+				FRigVMOperand* Operand = const_cast<FRigVMOperand*>(&OperandArray[j]);
+
+				const FRigVMOperand OriginalOperand = *Operand;
+				check(OriginalOperand == FunctionOperandArray[j]);
+				
 				FRigVMOperand InterfaceOperand;
 
 				// Remap the variable: find the operand index of the outer variable
@@ -3044,7 +3053,7 @@ bool URigVMCompiler::TraverseInlineFunction(const FRigVMInlineFunctionExprAST* I
 					else
 					{
 						FRigVMFunctionCompilationPropertyPath DescriptionFromFunction;
-						switch (OriginalMemoryType)
+						switch (OriginalOperand.MemoryType)
 						{
 							case ERigVMMemoryType::Work:
 							{
@@ -3088,11 +3097,11 @@ bool URigVMCompiler::TraverseInlineFunction(const FRigVMInlineFunctionExprAST* I
 				}
 
 				// Store correspondence between Function operand and New operands
-				check(!FunctionOperandToNewOperand.Contains(FunctionOperand) || FunctionOperandToNewOperand.FindChecked(FunctionOperand) == *Operand);
-				FunctionOperandToNewOperand.Add(FunctionOperand, *Operand);
+				check(!FunctionOperandToNewOperand.Contains(OriginalOperand) || FunctionOperandToNewOperand.FindChecked(OriginalOperand) == *Operand);
+				FunctionOperandToNewOperand.Add(OriginalOperand, *Operand);
 
 				// Also find the function debug operands and store correspondence to new debug operands
-				if (const TArray<FRigVMOperand>* FunctionDebugRegisters = FunctionCompilationData->OperandToDebugRegisters.Find(FunctionOperand))
+				if (const TArray<FRigVMOperand>* FunctionDebugRegisters = FunctionCompilationData->OperandToDebugRegisters.Find(OriginalOperand))
 				{
 					for (const FRigVMOperand& FunctionDebugRegister : *FunctionDebugRegisters)
 					{
@@ -3125,19 +3134,37 @@ bool URigVMCompiler::TraverseInlineFunction(const FRigVMInlineFunctionExprAST* I
 
 			if (WorkData.Settings.SetupNodeInstructionIndex)
 			{
-				if (const TArray<TWeakObjectPtr<UObject>>* Callstack = FunctionByteCode.GetCallstackForInstruction(i-InstructionIndexStart))
+				if (const TArray<TWeakObjectPtr<UObject>>* Callstack = FunctionByteCode.GetCallstackForInstruction(FunctionInstructionIndex))
 				{
 					if (Callstack->Num() > 1)
 					{
 						FRigVMCallstack InstructionCallstack = FuncRefCallstack;
 						InstructionCallstack.Stack.Append(&(*Callstack)[1], Callstack->Num()-1);
-						WorkData.VM->GetByteCode().SetSubject(i, InstructionCallstack.GetCallPath(), InstructionCallstack.GetStack());
+						ByteCode.SetSubject(i, InstructionCallstack.GetCallPath(), InstructionCallstack.GetStack());
 					}
 				}
 
 				// also store the instruction for the subject so that profiling can determine the cost of functions.
-				WorkData.VM->GetByteCode().AddInstructionForSubject(FunctionReferenceNode, i);
+				ByteCode.AddInstructionForSubject(FunctionReferenceNode, i);
+			}
+
+			// remap the input and output operands per instruction
+			TArray<FRigVMOperand> InputOperands, OutputOperands;
+			if (FunctionByteCode.InputOperandsPerInstruction.IsValidIndex(FunctionInstructionIndex))
+			{
+				for (const FRigVMOperand& Operand : FunctionByteCode.InputOperandsPerInstruction[FunctionInstructionIndex])
+				{
+					InputOperands.Add(FunctionOperandToNewOperand.FindChecked(Operand));
+				}
 			}
+			if (FunctionByteCode.OutputOperandsPerInstruction.IsValidIndex(FunctionInstructionIndex))
+			{
+				for (const FRigVMOperand& Operand : FunctionByteCode.OutputOperandsPerInstruction[FunctionInstructionIndex])
+				{
+					OutputOperands.Add(FunctionOperandToNewOperand.FindChecked(Operand));
+				}
+			}
+			ByteCode.SetOperandsForInstruction(Instruction, InputOperands, OutputOperands);
 		}
 
 		// Add all pin paths to operand from the function
@@ -3188,39 +3215,6 @@ bool URigVMCompiler::TraverseInlineFunction(const FRigVMInlineFunctionExprAST* I
 				WorkData.VM->OperandToDebugRegisters.Add(*NewOperand, NewValues);
 			}
 		}
-
-#if WITH_EDITORONLY_DATA
-		TArray<FRigVMOperand> InputsOperands, OutputOperands;
-
-		int32 ArgumentIndex = 0;
-		for(const URigVMPin* InputPin : Node->GetPins())
-		{
-			if(InputPin->IsExecuteContext())
-			{
-				continue;
-			}
-
-			const FRigVMOperand& Operand = InterfaceOperands[ArgumentIndex++];
-
-			if(InputPin->GetDirection() == ERigVMPinDirection::Output || InputPin->GetDirection() == ERigVMPinDirection::IO)
-			{
-				OutputOperands.Add(Operand);
-			}
-
-			if(InputPin->GetDirection() != ERigVMPinDirection::Input && InputPin->GetDirection() != ERigVMPinDirection::IO)
-			{
-				continue;
-			}
-
-			InputsOperands.Add(Operand);
-		}
-
-		WorkData.VM->GetByteCode().SetOperandsForInstruction(
-			InstructionIndexStart,
-			FRigVMOperandArray(InputsOperands.GetData(), InputsOperands.Num()),
-			FRigVMOperandArray(OutputOperands.GetData(), OutputOperands.Num()));
-
-#endif
 	}
 
 	return true;
