RigVM: fixup for 44779120

#rnx

diff --git a/Engine/Plugins/Runtime/RigVM/Source/RigVM/Private/RigVMCore/RigVMByteCode.cpp b/Engine/Plugins/Runtime/RigVM/Source/RigVM/Private/RigVMCore/RigVMByteCode.cpp
--- a/Engine/Plugins/Runtime/RigVM/Source/RigVM/Private/RigVMCore/RigVMByteCode.cpp
+++ b/Engine/Plugins/Runtime/RigVM/Source/RigVM/Private/RigVMCore/RigVMByteCode.cpp
@@ -2142,6 +2142,7 @@ void FRigVMByteCode::SetOperandsForInstruction(const FRigVMInstruction& InInstru
 {
 	const FRigVMOperandArray OperandsForInstruction = GetOperandsForOp(InInstruction);
 	const int32 InstructionIndex = InInstruction.Index;
+	check(InInstruction.Index != INDEX_NONE && InInstruction.Index < NumInstructions);
 	
 	if (InputOperandsPerInstruction.Num() <= InstructionIndex)
 	{
