RigVM / ControlRig: Remove debug breakpoint code

#rnx

diff --git a/Engine/Plugins/Animation/ControlRig/Source/ControlRig/Private/ControlRig.cpp b/Engine/Plugins/Animation/ControlRig/Source/ControlRig/Private/ControlRig.cpp
--- a/Engine/Plugins/Animation/ControlRig/Source/ControlRig/Private/ControlRig.cpp
+++ b/Engine/Plugins/Animation/ControlRig/Source/ControlRig/Private/ControlRig.cpp
@@ -940,13 +940,7 @@ bool UControlRig::Execute(const FName& InEventName)
 		if (UControlRig* CDO = GetClass()->GetDefaultObject<UControlRig>())
 		{
 			// Copy the breakpoints. This will not override the state of the breakpoints
-			DebugInfo.SetBreakpoints(CDO->DebugInfo.GetBreakpoints());
-
-			// If there are any breakpoints, create the Snapshot VM if it hasn't been created yet
-			if (DebugInfo.GetBreakpoints().Num() > 0)
-			{
-				GetSnapshotVM();
-			}
+			DebugInfo.CopyFrom(CDO->DebugInfo);
 		}
 
 		ExtendedExecuteContext.SetDebugInfo(&DebugInfo);
@@ -954,7 +948,6 @@ bool UControlRig::Execute(const FName& InEventName)
 	else
 	{
 		ExtendedExecuteContext.SetDebugInfo(nullptr);
-		GetSnapshotContext().Reset();
 	}
 	
 	if (IsProfilingEnabled())
@@ -1746,31 +1739,6 @@ bool UControlRig::Execute_Internal(const FName& InEventName)
 		}
 #endif
 		
-		const bool bUseDebuggingSnapshots = !VM->IsNativized();
-		
-#if WITH_EDITOR
-		if(bUseDebuggingSnapshots)
-		{
-			if(URigVM* SnapShotVM = GetSnapshotVM(false)) // don't create it for normal runs
-			{
-				const bool bIsEventFirstInQueue = !LocalEventQueueToRun.IsEmpty() && LocalEventQueueToRun[0] == InEventName; 
-				const bool bIsEventLastInQueue = !LocalEventQueueToRun.IsEmpty() && LocalEventQueueToRun.Last() == InEventName;
-
-				if (GetHaltedAtBreakpoint().IsValid())
-				{
-					if(bIsEventFirstInQueue)
-					{
-						CopyVMMemory(GetRigVMExtendedExecuteContext(), GetSnapshotContext());
-					}
-				}
-				else if(bIsEventLastInQueue)
-				{
-					CopyVMMemory(GetSnapshotContext(), GetRigVMExtendedExecuteContext());
-				}
-			}
-		}
-#endif
-
 		URigHierarchy* Hierarchy = GetHierarchy();
 		FRigHierarchyExecuteContextBracket HierarchyContextGuard(Hierarchy, &Context);
 
diff --git a/Engine/Plugins/Runtime/RigVM/Source/RigVM/Private/RigVMCore/RigVM.cpp b/Engine/Plugins/Runtime/RigVM/Source/RigVM/Private/RigVMCore/RigVM.cpp
--- a/Engine/Plugins/Runtime/RigVM/Source/RigVM/Private/RigVMCore/RigVM.cpp
+++ b/Engine/Plugins/Runtime/RigVM/Source/RigVM/Private/RigVMCore/RigVM.cpp
@@ -104,17 +104,10 @@ URigVM::URigVM()
 {
 }
 
-PRAGMA_DISABLE_DEPRECATION_WARNINGS // required until we eliminate ExecutionReachedExit and ExecutionHalted
 URigVM::~URigVM()
 {
 	Reset_Internal();
-
-	ExecutionReachedExit().Clear();
-#if WITH_EDITOR
-	ExecutionHalted().Clear();
-#endif
 }
-PRAGMA_ENABLE_DEPRECATION_WARNINGS // required until we eliminate ExecutionReachedExit and ExecutionHalted
 
 void URigVM::Serialize(FArchive& Ar)
 {
@@ -1051,12 +1044,7 @@ bool URigVM::CanExecuteEntry(const FRigVMExtendedExecuteContext& Context, const
 
 bool URigVM::ResumeExecution(FRigVMExtendedExecuteContext& Context, const FName& InEntryName)
 {
-	if (FRigVMDebugInfo* RigVMDebugInfo = Context.GetRigVMDebugInfo())
-	{
-		RigVMDebugInfo->ResumeExecution();
-	}
-
-	return ExecuteVM(Context, InEntryName) != ERigVMExecuteResult::Failed;
+	return false;
 }
 
 #endif 
@@ -1452,198 +1440,6 @@ void URigVM::RebuildByteCodeOnLoad()
 	}
 }
 
-#if WITH_EDITOR
-bool URigVM::ShouldHaltAtInstruction(FRigVMExtendedExecuteContext& Context, const FName& InEventName, const int32 InstructionIndex)
-{
-	FRigVMDebugInfo* RigVMDebugInfo = Context.GetRigVMDebugInfo();
-
-	if (RigVMDebugInfo == nullptr)
-	{
-		return false;
-	}
-
-	if(RigVMDebugInfo->IsEmpty())
-	{
-		return false;
-	}
-	
-	FRigVMByteCode& ByteCode = GetByteCode();
-	FRigVMExecuteContext& ContextPublicData = Context.GetPublicData<>();
-
-	TArray<FRigVMBreakpoint> BreakpointsAtInstruction = RigVMDebugInfo->FindBreakpointsAtInstruction(InstructionIndex);
-	for (FRigVMBreakpoint Breakpoint : BreakpointsAtInstruction)
-	{
-		if (RigVMDebugInfo->IsActive(Breakpoint))
-		{
-			switch (RigVMDebugInfo->GetCurrentBreakpointAction())
-			{
-				case ERigVMBreakpointAction::None:
-				{
-					// Halted at breakpoint. Check if this is a new breakpoint different from the previous halt.
-					if (RigVMDebugInfo->GetHaltedAtBreakpoint() != Breakpoint ||
-						RigVMDebugInfo->GetHaltedAtBreakpointHit() != RigVMDebugInfo->GetBreakpointHits(Breakpoint))
-					{
-						RigVMDebugInfo->SetHaltedAtBreakpoint(Breakpoint);
-						RigVMDebugInfo->SetHaltedAtBreakpointHit(RigVMDebugInfo->GetBreakpointHits(Breakpoint));
-						RigVMDebugInfo->SetCurrentActiveBreakpoint(Breakpoint);
-						
-						// We want to keep the callstack up to the node that produced the halt
-						const TArray<TWeakObjectPtr<UObject>>* FullCallstack = ByteCode.GetCallstackForInstruction(ContextPublicData.InstructionIndex);
-						if (FullCallstack)
-						{
-							RigVMDebugInfo->SetCurrentActiveBreakpointCallstack(TArray<TWeakObjectPtr<UObject>>(FullCallstack->GetData(), FullCallstack->Find((TWeakObjectPtr<UObject>)Breakpoint.Subject)+1));
-						}
-						UObject* BreakpointNode = Breakpoint.Subject.IsValid() ? Breakpoint.Subject.Get() : nullptr;
-						RigVMDebugInfo->ExecutionHalted().Broadcast(ContextPublicData.InstructionIndex, BreakpointNode, InEventName);
-					}
-					return true;
-				}
-				case ERigVMBreakpointAction::Resume:
-				{
-					RigVMDebugInfo->SetCurrentBreakpointAction(ERigVMBreakpointAction::None);
-
-					if (RigVMDebugInfo->IsTemporaryBreakpoint(Breakpoint))
-					{
-						RigVMDebugInfo->RemoveBreakpoint(Breakpoint);
-					}
-					else
-					{
-						RigVMDebugInfo->IncrementBreakpointActivationOnHit(Breakpoint);
-						RigVMDebugInfo->HitBreakpoint(Breakpoint);
-					}
-					break;
-				}
-				case ERigVMBreakpointAction::StepOver:
-				case ERigVMBreakpointAction::StepInto:
-				case ERigVMBreakpointAction::StepOut:
-				{
-					// If we are stepping, check if we were halted at the current instruction, and remember it 
-					if (!RigVMDebugInfo->GetCurrentActiveBreakpoint())
-					{
-						RigVMDebugInfo->SetCurrentActiveBreakpoint(Breakpoint);
-						const TArray<TWeakObjectPtr<UObject>>* FullCallstack = ByteCode.GetCallstackForInstruction(ContextPublicData.InstructionIndex);
-						
-						// We want to keep the callstack up to the node that produced the halt
-						if (FullCallstack)
-						{
-							RigVMDebugInfo->SetCurrentActiveBreakpointCallstack(TArray<TWeakObjectPtr<UObject>>(FullCallstack->GetData(), FullCallstack->Find((TWeakObjectPtr<UObject>)RigVMDebugInfo->GetCurrentActiveBreakpoint().Subject)+1));
-						}
-					}
-					
-					break;	
-				}
-				default:
-				{
-					ensure(false);
-					break;
-				}
-			}
-		}
-		else
-		{
-			RigVMDebugInfo->HitBreakpoint(Breakpoint);
-		}
-	}
-
-	// If we are stepping, and the last active breakpoint was set, check if this is the new temporary breakpoint
-	if (RigVMDebugInfo->GetCurrentBreakpointAction() != ERigVMBreakpointAction::None && RigVMDebugInfo->GetCurrentActiveBreakpoint())
-	{
-		const TArray<TWeakObjectPtr<UObject>>* CurrentCallstack = ByteCode.GetCallstackForInstruction(ContextPublicData.InstructionIndex);
-		if (CurrentCallstack && !CurrentCallstack->IsEmpty())
-		{
-			UObject* NewBreakpointNode = nullptr;
-
-			// Find the first difference in the callstack
-			int32 DifferenceIndex = INDEX_NONE;
-			TArray<TWeakObjectPtr<UObject>>& PreviousCallstack = RigVMDebugInfo->GetCurrentActiveBreakpointCallstack();
-			for (int32 i=0; i<PreviousCallstack.Num(); ++i)
-			{
-				if (CurrentCallstack->Num() == i)
-				{
-					DifferenceIndex = i-1;
-					break;
-				}
-				if (PreviousCallstack[i] != CurrentCallstack->operator[](i))
-				{
-					DifferenceIndex = i;
-					break;
-				}
-			}
-
-			if (RigVMDebugInfo->GetCurrentBreakpointAction() == ERigVMBreakpointAction::StepOver)
-			{
-				if (DifferenceIndex != INDEX_NONE)
-				{
-					if ((*CurrentCallstack)[DifferenceIndex].IsValid())
-					{
-						NewBreakpointNode = (*CurrentCallstack)[DifferenceIndex].Get();
-					}
-				}
-			}
-			else if (RigVMDebugInfo->GetCurrentBreakpointAction() == ERigVMBreakpointAction::StepInto)
-			{
-				if (DifferenceIndex == INDEX_NONE)
-				{
-					if (!CurrentCallstack->IsEmpty() && !PreviousCallstack.IsEmpty() && CurrentCallstack->Last() != PreviousCallstack.Last())
-					{
-						const int32 MinIndex = FMath::Min(PreviousCallstack.Num(), CurrentCallstack->Num()-1);
-						if ((*CurrentCallstack)[MinIndex].IsValid())
-						{
-							NewBreakpointNode = (*CurrentCallstack)[MinIndex].Get();
-						}
-					}
-				}
-				else
-				{
-					if ((*CurrentCallstack)[DifferenceIndex].IsValid())
-					{
-						NewBreakpointNode = (*CurrentCallstack)[DifferenceIndex].Get();
-					}
-				}
-			}
-			else if (RigVMDebugInfo->GetCurrentBreakpointAction() == ERigVMBreakpointAction::StepOut)
-			{
-				if (DifferenceIndex != INDEX_NONE && DifferenceIndex <= PreviousCallstack.Num() - 2)
-                {
-					if ((*CurrentCallstack)[DifferenceIndex].IsValid())
-					{
-						NewBreakpointNode = (*CurrentCallstack)[DifferenceIndex].Get();
-					}
-                }
-			}
-			
-			if (NewBreakpointNode)
-			{
-				// Remove or hit previous breakpoint
-				if (RigVMDebugInfo->IsTemporaryBreakpoint(RigVMDebugInfo->GetCurrentActiveBreakpoint()))
-				{
-					RigVMDebugInfo->RemoveBreakpoint(RigVMDebugInfo->GetCurrentActiveBreakpoint());
-				}
-				else
-				{
-					RigVMDebugInfo->IncrementBreakpointActivationOnHit(RigVMDebugInfo->GetCurrentActiveBreakpoint());
-					RigVMDebugInfo->HitBreakpoint(RigVMDebugInfo->GetCurrentActiveBreakpoint());
-				}
-
-				// Create new temporary breakpoint
-				const FRigVMBreakpoint& NewBreakpoint = RigVMDebugInfo->AddBreakpoint(ContextPublicData.InstructionIndex, NewBreakpointNode, 0, true);
-				RigVMDebugInfo->SetBreakpointHits(NewBreakpoint, static_cast<uint16>(GetInstructionVisitedCount(Context, ContextPublicData.InstructionIndex)));
-				RigVMDebugInfo->SetBreakpointActivationOnHit(NewBreakpoint, static_cast<uint16>(GetInstructionVisitedCount(Context, ContextPublicData.InstructionIndex)));
-				RigVMDebugInfo->SetCurrentBreakpointAction(ERigVMBreakpointAction::None);
-
-				RigVMDebugInfo->SetHaltedAtBreakpoint(NewBreakpoint);
-				RigVMDebugInfo->SetHaltedAtBreakpointHit(RigVMDebugInfo->GetBreakpointHits(RigVMDebugInfo->GetHaltedAtBreakpoint()));
-				RigVMDebugInfo->ExecutionHalted().Broadcast(ContextPublicData.InstructionIndex, NewBreakpointNode, InEventName);
-		
-				return true;
-			}
-		}
-	}
-
-	return false;
-}
-#endif
-
 bool URigVM::Initialize(FRigVMExtendedExecuteContext& Context)
 {
 	if (Context.ExecutingThreadId != INDEX_NONE)
@@ -1813,16 +1609,6 @@ ERigVMExecuteResult URigVM::ExecuteVM(FRigVMExtendedExecuteContext& Context, con
 
 	FEntryExecuteGuard EntryExecuteGuard(Context.EntriesBeingExecuted, EntryIndexToPush);
 
-#if WITH_EDITOR
-	if (bIsRootEntry)
-	{
-		if (FRigVMDebugInfo* RigVMDebugInfo = Context.GetRigVMDebugInfo())
-		{
-			RigVMDebugInfo->StartExecution();
-		}
-	}
-#endif
-
 	if(bIsRootEntry)
 	{
 		ContextPublicData.NumExecutions++;
@@ -1860,10 +1646,8 @@ ERigVMExecuteResult URigVM::ExecuteVM(FRigVMExtendedExecuteContext& Context, con
 		
 		if (FRigVMDebugInfo* RigVMDebugInfo = Context.GetRigVMDebugInfo())
 		{
-			if (RigVMDebugInfo->GetHaltedAtBreakpoint().IsValid() && Context.CurrentExecuteResult != ERigVMExecuteResult::Halted)
+			if (!RigVMDebugInfo->IsEmpty())
 			{
-				RigVMDebugInfo->SetCurrentActiveBreakpoint(FRigVMBreakpoint());
-				RigVMDebugInfo->GetHaltedAtBreakpoint().Reset();
 				RigVMDebugInfo->ExecutionHalted().Broadcast(INDEX_NONE, nullptr, InEntryName);
 			}
 		}
@@ -1903,24 +1687,10 @@ ERigVMExecuteResult URigVM::ExecuteInstructions(FRigVMExtendedExecuteContext& Co
 	while (Instructions.IsValidIndex(ContextPublicData.InstructionIndex))
 	{
 #if WITH_EDITOR
-		if (ShouldHaltAtInstruction(Context, Context.CurrentEntryName, ContextPublicData.InstructionIndex))
-		{
-#if UE_RIGVM_DEBUG_EXECUTION
-			if (CVarControlRigDebugAllVMExecutions->GetBool() || ContextPublicData.bDebugExecution)
-			{
-				ContextPublicData.Log(EMessageSeverity::Info, ContextPublicData.DebugMemoryString);
-			}
-#endif
-			// we'll recursively exit all invoked
-			// entries here.
-			return Context.CurrentExecuteResult = ERigVMExecuteResult::Halted;
-		}
-
 		if(Context.CurrentExecuteResult == ERigVMExecuteResult::Halted)
 		{
 			return Context.CurrentExecuteResult;
 		}
-
 #endif
 		
 		if(ContextPublicData.InstructionIndex > (uint16)InLastInstruction)
@@ -2183,10 +1953,8 @@ ERigVMExecuteResult URigVM::ExecuteInstructions(FRigVMExtendedExecuteContext& Co
 #if WITH_EDITOR					
 					if (RigVMDebugInfo != nullptr)
 					{
-						if (RigVMDebugInfo->GetHaltedAtBreakpoint().IsValid())
+						if (RigVMDebugInfo && !RigVMDebugInfo->IsEmpty())
 						{
-							RigVMDebugInfo->GetHaltedAtBreakpoint().Reset();
-							RigVMDebugInfo->SetCurrentActiveBreakpoint(FRigVMBreakpoint());
 							RigVMDebugInfo->ExecutionHalted().Broadcast(INDEX_NONE, nullptr, Context.CurrentEntryName);
 						}
 					}
@@ -2234,7 +2002,7 @@ ERigVMExecuteResult URigVM::ExecuteInstructions(FRigVMExtendedExecuteContext& Co
 
 #if WITH_EDITOR
 					// if we are halted at a break point we need to exit here
-					if (ShouldHaltAtInstruction(Context, Context.CurrentEntryName, ContextPublicData.InstructionIndex))
+					if (RigVMDebugInfo && RigVMDebugInfo->ShouldEarlyExitBeforeInstruction(ContextPublicData.InstructionIndex))
 					{
 						return Context.CurrentExecuteResult = ERigVMExecuteResult::Halted;
 					}
@@ -2397,6 +2165,20 @@ ERigVMExecuteResult URigVM::ExecuteInstructions(FRigVMExtendedExecuteContext& Co
 			}
 		}
 
+		if (RigVMDebugInfo && RigVMDebugInfo->ShouldEarlyExitBeforeInstruction(ContextPublicData.InstructionIndex))
+		{
+#if UE_RIGVM_DEBUG_EXECUTION
+			if (CVarControlRigDebugAllVMExecutions->GetBool() || ContextPublicData.bDebugExecution)
+			{
+				ContextPublicData.Log(EMessageSeverity::Info, ContextPublicData.DebugMemoryString);
+			}
+#endif
+			// we'll recursively exit all invoked
+			// entries here.
+			RigVMDebugInfo->ExecutionHalted().Broadcast(ContextPublicData.InstructionIndex, nullptr, Context.CurrentEntryName);
+			return Context.CurrentExecuteResult = ERigVMExecuteResult::Halted;
+		}
+
 #if UE_RIGVM_DEBUG_EXECUTION
 		if (CVarControlRigDebugAllVMExecutions->GetBool() || ContextPublicData.bDebugExecution)
 		{
@@ -2531,6 +2313,8 @@ FRigVMRegistry_NoLock* URigVM::GetLocalizedRegistry() const
 
 bool URigVM::CreateLocalizedRegistryIfRequired(bool bForce)
 {
+	UE::TScopeLock CreateLocalizedRegistryScopeLock(CreateLocalizedRegistryMutex);
+
 	if(LocalizedRegistry.IsValid())
 	{
 		return true;
@@ -2542,9 +2326,8 @@ bool URigVM::CreateLocalizedRegistryIfRequired(bool bForce)
 	}
 		
 	// create a new registry
-	TSharedPtr<FRigVMRegistry_NoLock> TemporaryRegistryPtr = FRigVMRegistry_NoLock::CreateLocalizedRegistry();
-	FRigVMRegistry_NoLock& TemporaryRegistry = *TemporaryRegistryPtr.Get();
-	FRigVMRegistryHandle TemporaryHandle(&TemporaryRegistry);
+	TSharedRef<FRigVMRegistry_NoLock> TemporaryRegistry = FRigVMRegistry_NoLock::CreateLocalizedRegistry();
+	FRigVMRegistryHandle TemporaryHandle(&TemporaryRegistry.Get());
 
 	RefreshInstructionsIfRequired();
 	if(ExternalPropertyPaths.IsEmpty())
@@ -2630,7 +2413,7 @@ bool URigVM::CreateLocalizedRegistryIfRequired(bool bForce)
 			for(const TRigVMTypeIndex& ArgumentTypeIndex :  Function->GetArgumentTypeIndices())
 			{
 				const FRigVMTemplateArgumentType& Type = Function->OwnerRegistry->GetType_NoLock(ArgumentTypeIndex);
-				TemporaryRegistry.FindOrAddType_NoLock(Type);
+				TemporaryRegistry->FindOrAddType_NoLock(Type);
 			}
 		}
 	}
@@ -2682,14 +2465,14 @@ bool URigVM::CreateLocalizedRegistryIfRequired(bool bForce)
 						return false;
 					}
 
-					TypeIndices[TypeIndex] = TemporaryRegistry.FindOrAddType_NoLock(Type);
+					TypeIndices[TypeIndex] = TemporaryRegistry->FindOrAddType_NoLock(Type);
 					if (TypeIndices[TypeIndex] == INDEX_NONE)
 					{
 						UE_LOG(LogRigVM, Error, TEXT("Cannot register type '%s' in dispatch '%s' for VM '%s'\nCannot localize registry."), *Type.CPPType.ToString(), *Factory->GetFactoryName().ToString(), *GetPathName());
 						return false;
 					}
 
-					const FRigVMTemplateArgumentType TemporaryType = TemporaryRegistry.GetType_NoLock(TypeIndices[TypeIndex]);
+					const FRigVMTemplateArgumentType TemporaryType = TemporaryRegistry->GetType_NoLock(TypeIndices[TypeIndex]);
 					if(Type != TemporaryType)
 					{
 						UE_LOG(LogRigVM, Error, TEXT("Registered type '%s' in dispatch '%s' does not match ('%s') for VM '%s'\nCannot localize registry."), *Type.CPPType.ToString(), *Factory->GetFactoryName().ToString(), *TemporaryType.CPPType.ToString(), *GetPathName());
@@ -2701,7 +2484,7 @@ bool URigVM::CreateLocalizedRegistryIfRequired(bool bForce)
 		}
 
 		// register the factory
-		FRigVMDispatchFactory* TemporaryFactory = const_cast<FRigVMDispatchFactory*>(TemporaryRegistry.RegisterFactory_NoLock(FactoryStruct, TemporaryArgumentInfos));
+		FRigVMDispatchFactory* TemporaryFactory = const_cast<FRigVMDispatchFactory*>(TemporaryRegistry->RegisterFactory_NoLock(FactoryStruct, TemporaryArgumentInfos));
 
 		// update all argument name caches,
 		// force create all function pointers
@@ -2733,14 +2516,14 @@ bool URigVM::CreateLocalizedRegistryIfRequired(bool bForce)
 		if(Function->Factory)
 		{
 			// pulling on the function makes sure the permutation exists
-			if (!TemporaryRegistry.FindFunction_NoLock(*Function->Name))
+			if (!TemporaryRegistry->FindFunction_NoLock(*Function->Name))
 			{
-				for (int32 FactoryIndex = 0; FactoryIndex < LocalizedRegistry->GetFactories_NoLock().Num(); FactoryIndex++)
+				for (int32 FactoryIndex = 0; FactoryIndex < TemporaryRegistry->GetFactories_NoLock().Num(); FactoryIndex++)
 				{
-					const FRigVMDispatchFactory* AvailableFactory = LocalizedRegistry->GetFactories_NoLock()[FactoryIndex];
+					const FRigVMDispatchFactory* AvailableFactory = TemporaryRegistry->GetFactories_NoLock()[FactoryIndex];
 					UE_LOG(LogRigVM, Error, TEXT("Available factory [%03d] '%s' for VM '%s'"), FactoryIndex, *AvailableFactory->GetFactoryName().ToString(), *GetPathName());
 				}
-				for (const FRigVMFunction& AvailableFunction : LocalizedRegistry->GetFunctions_NoLock())
+				for (const FRigVMFunction& AvailableFunction : TemporaryRegistry->GetFunctions_NoLock())
 				{
 					UE_LOG(LogRigVM, Error, TEXT("Available function [%03d] '%s' for VM '%s'"), AvailableFunction.Index, *AvailableFunction.Name, *GetPathName());
 				}
@@ -2752,11 +2535,12 @@ bool URigVM::CreateLocalizedRegistryIfRequired(bool bForce)
 		else
 		{
 			// registering the function will also register it's potential template
-			TemporaryRegistry.Register_NoLock(*Function->Name, Function->FunctionPtr, Function->Struct, Function->Arguments);
+			TemporaryRegistry->Register_NoLock(*Function->Name, Function->FunctionPtr, Function->Struct, Function->Arguments);
 		}
 	}
 
-	Swap(LocalizedRegistry, TemporaryRegistryPtr);
+	// finally relate this VM to the localized registry
+	LocalizedRegistry = TemporaryRegistry.ToSharedPtr();
 
 	// now that we've registered all functions - let's re-resolve the functions
 	FunctionsStorage.Reset();
diff --git a/Engine/Plugins/Runtime/RigVM/Source/RigVM/Private/RigVMCore/RigVMDebugInfo.cpp b/Engine/Plugins/Runtime/RigVM/Source/RigVM/Private/RigVMCore/RigVMDebugInfo.cpp
--- a/Engine/Plugins/Runtime/RigVM/Source/RigVM/Private/RigVMCore/RigVMDebugInfo.cpp
+++ b/Engine/Plugins/Runtime/RigVM/Source/RigVM/Private/RigVMCore/RigVMDebugInfo.cpp
@@ -3,242 +3,3 @@
 #include "RigVMCore/RigVMDebugInfo.h"
 
 #include UE_INLINE_GENERATED_CPP_BY_NAME(RigVMDebugInfo)
-
-void FRigVMDebugInfo::ResetState()
-{
-	BreakpointHits.Empty();
-	BreakpointActivationOnHit.Empty();
-	TemporaryBreakpoint.Reset();
-	CurrentActiveBreakpoint.Invalidate();
-	CurrentActiveBreakpointCallstack.Empty();
-}
-
-void FRigVMDebugInfo::StartExecution()
-{
-	BreakpointHits.Empty();
-	CurrentActiveBreakpoint.Invalidate();
-	CurrentActiveBreakpointCallstack.Empty();
-}
-
-const FRigVMBreakpoint& FRigVMDebugInfo::FindBreakpoint(const int32 InInstructionIndex, const UObject* InSubject) const
-{
-	TArray<FRigVMBreakpoint> BreakpointAtInstruction = FindBreakpointsAtInstruction(InInstructionIndex);
-	for (const FRigVMBreakpoint& BP : BreakpointAtInstruction)
-	{
-		if (BP.Subject == InSubject)
-		{
-			return BP;
-		}
-	}
-	
-	static const FRigVMBreakpoint EmptyBreakpoint;
-	return EmptyBreakpoint;
-}
-
-TArray<FRigVMBreakpoint> FRigVMDebugInfo::FindBreakpointsAtInstruction(const int32 InInstructionIndex) const
-{
-	TArray<FRigVMBreakpoint> Result;
-	for (const FRigVMBreakpoint& BP : Breakpoints)
-	{
-		if (BP.InstructionIndex == InInstructionIndex)
-		{
-			Result.Add(BP);
-		}
-	}
-
-	if (TemporaryBreakpoint.IsValid() && TemporaryBreakpoint.InstructionIndex == InInstructionIndex)
-	{
-		Result.Add(TemporaryBreakpoint);
-	}
-	
-	return Result;
-}
-
-const FRigVMBreakpoint& FRigVMDebugInfo::FindBreakpoint(const FGuid& InGuid) const
-{
-	if(InGuid.IsValid())
-	{
-		if(TemporaryBreakpoint.IsValid())
-		{
-			if(TemporaryBreakpoint.Guid == InGuid)
-			{
-				return TemporaryBreakpoint;
-			}
-		}
-
-		for(const FRigVMBreakpoint& Breakpoint : Breakpoints)
-		{
-			if(Breakpoint.Guid == InGuid)
-			{
-				return Breakpoint;
-			}
-		}
-	}
-	
-	static const FRigVMBreakpoint EmptyBreakpoint;
-	return EmptyBreakpoint;
-}
-
-const FRigVMBreakpoint& FRigVMDebugInfo::AddBreakpoint(const int32 InstructionIndex, TWeakObjectPtr<UObject> InNode, const uint16 InDepth,
-                                                            const bool bIsTemporary)
-{
-	static const FRigVMBreakpoint EmptyBreakpoint;
-	
-	for (const FRigVMBreakpoint& BP : FindBreakpointsAtInstruction(InstructionIndex))
-	{
-		if (BP.Subject == InNode)
-		{
-			return EmptyBreakpoint;
-		}
-	}
-	
-	if (bIsTemporary)
-	{
-		TemporaryBreakpoint = FRigVMBreakpoint(InstructionIndex, InNode, 0);
-			
-		// Do not override the state if it already exists
-		BreakpointActivationOnHit.FindOrAdd(TemporaryBreakpoint.Guid, 0);
-		BreakpointHits.FindOrAdd(TemporaryBreakpoint.Guid, 0);
-		return TemporaryBreakpoint;
-	}
-
-	// Breakpoints are sorted by instruction index and callstack depth
-	int32 Index = 0;
-	for(; Index<Breakpoints.Num(); ++Index)
-	{
-		if (InstructionIndex < Breakpoints[Index].InstructionIndex ||
-			(InstructionIndex == Breakpoints[Index].InstructionIndex && InDepth < Breakpoints[Index].Depth))
-		{
-			break;
-		}
-	}
-	const int32 InsertedIndex = Breakpoints.Insert(FRigVMBreakpoint(InstructionIndex, InNode, InDepth), Index);
-	const FRigVMBreakpoint& NewBP = Breakpoints[InsertedIndex];
-	
-	// Do not override the state if it already exists
-	BreakpointActivationOnHit.FindOrAdd(NewBP.Guid, 0);
-	BreakpointHits.FindOrAdd(NewBP.Guid, 0);
-	return NewBP;
-}
-
-bool FRigVMDebugInfo::RemoveBreakpoint(const FRigVMBreakpoint& InBreakpoint)
-{
-	bool Found = false;
-	int32 Index = Breakpoints.Find(InBreakpoint);
-	if (Index != INDEX_NONE)
-	{
-		BreakpointHits.Remove(InBreakpoint.Guid);
-		BreakpointActivationOnHit.Remove(InBreakpoint.Guid);
-		Breakpoints.RemoveAt(Index);
-		Found = true;
-	}
-	if (TemporaryBreakpoint.IsValid() && InBreakpoint == TemporaryBreakpoint)
-	{
-		BreakpointHits.Remove(TemporaryBreakpoint.Guid);
-		BreakpointActivationOnHit.Remove(TemporaryBreakpoint.Guid);
-		TemporaryBreakpoint.Reset();
-		Found = true;
-	}
-	return Found;
-}
-
-bool FRigVMDebugInfo::IsActive(const FRigVMBreakpoint& InBreakpoint) const
-{
-	if(!InBreakpoint.IsValid())
-	{
-		if(TemporaryBreakpoint.IsValid())
-		{
-			return IsActive(TemporaryBreakpoint);
-		}
-	}
-	if (InBreakpoint.IsValid())
-	{
-		if (InBreakpoint.bIsActive)
-		{
-			uint16 Hits = 0;
-			uint16 OnHit = 0;
-			if (BreakpointHits.Contains(InBreakpoint.Guid))
-			{
-				Hits = BreakpointHits.FindChecked(InBreakpoint.Guid);
-			}
-			if (BreakpointActivationOnHit.Contains(InBreakpoint.Guid))
-			{
-				OnHit = BreakpointActivationOnHit.FindChecked(InBreakpoint.Guid);
-			}
-			return Hits == OnHit;
-		}				
-	}
-	return false;		
-}
-
-void FRigVMDebugInfo::SetBreakpointHits(const FRigVMBreakpoint& InBreakpoint, const uint16 InBreakpointHits)
-{
-	if (BreakpointHits.Contains(InBreakpoint.Guid))
-	{
-		BreakpointHits[InBreakpoint.Guid] = InBreakpointHits;
-	}
-	else
-	{
-		BreakpointHits.Add(InBreakpoint.Guid, InBreakpointHits);
-	}
-}
-
-void FRigVMDebugInfo::HitBreakpoint(const FRigVMBreakpoint& InBreakpoint)
-{
-	if (BreakpointHits.Contains(InBreakpoint.Guid))
-	{
-		BreakpointHits[InBreakpoint.Guid]++;
-	}
-	else
-	{
-		BreakpointHits.Add(InBreakpoint.Guid, 1);
-	}
-}
-
-void FRigVMDebugInfo::SetBreakpointActivationOnHit(const FRigVMBreakpoint& InBreakpoint, const uint16 InActivationOnHit)
-{
-	if (BreakpointActivationOnHit.Contains(InBreakpoint.Guid))
-	{
-		BreakpointActivationOnHit[InBreakpoint.Guid] = InActivationOnHit;
-	}
-	else
-	{
-		BreakpointActivationOnHit.Add(InBreakpoint.Guid, InActivationOnHit);
-	}
-}
-
-void FRigVMDebugInfo::IncrementBreakpointActivationOnHit(const FRigVMBreakpoint& InBreakpoint)
-{
-	if (BreakpointActivationOnHit.Contains(InBreakpoint.Guid))
-	{
-		BreakpointActivationOnHit[InBreakpoint.Guid]++;
-	}
-	else
-	{
-		BreakpointActivationOnHit.Add(InBreakpoint.Guid, 1);
-	}
-}
-
-uint16 FRigVMDebugInfo::GetBreakpointHits(const FRigVMBreakpoint& InBreakpoint) const
-{
-	if (BreakpointActivationOnHit.Contains(InBreakpoint.Guid))
-	{
-		return BreakpointActivationOnHit[InBreakpoint.Guid];
-	}
-	return 0;
-}
-
-bool FRigVMDebugInfo::ResumeExecution()
-{
-	HaltedAtBreakpoint.Reset();
-	HaltedAtBreakpointHit = INDEX_NONE;
-
-	if (const FRigVMBreakpoint& CurrentBreakpoint = GetCurrentActiveBreakpoint())
-	{
-		IncrementBreakpointActivationOnHit(CurrentBreakpoint);
-		SetCurrentActiveBreakpoint(FRigVMBreakpoint());
-		return true;
-	}
-
-	return false;
-}
diff --git a/Engine/Plugins/Runtime/RigVM/Source/RigVM/Private/RigVMHost.cpp b/Engine/Plugins/Runtime/RigVM/Source/RigVM/Private/RigVMHost.cpp
--- a/Engine/Plugins/Runtime/RigVM/Source/RigVM/Private/RigVMHost.cpp
+++ b/Engine/Plugins/Runtime/RigVM/Source/RigVM/Private/RigVMHost.cpp
@@ -44,9 +44,6 @@ URigVMHost::URigVMHost(const FObjectInitializer& ObjectInitializer)
 	, bRequiresInitExecution(false)
 	, InitBracket(0)
 	, ExecuteBracket(0)
-#if WITH_EDITORONLY_DATA
-	, VMSnapshotBeforeExecution(nullptr)
-#endif
 #if UE_RIGVM_PROFILE_EXECUTE_UNITS_NUM
 	, ProfilingRunsLeft(0)
 	, AccumulatedCycles(0)
@@ -181,16 +178,6 @@ void URigVMHost::PostLoad()
 			}
 		}
 	}
-
-#if WITH_EDITORONLY_DATA
-	if (VMSnapshotBeforeExecution)
-	{
-		// Some VMSnapshots might have been created without the Transient flag.
-		// Assets from that version require the snapshot to be flagged as below.
-		VMSnapshotBeforeExecution->SetFlags(VMSnapshotBeforeExecution->GetFlags() | RF_Transient);
-	}
-#endif
-
 }
 
 void URigVMHost::PreSave(FObjectPreSaveContext SaveContext)
@@ -210,13 +197,6 @@ void URigVMHost::BeginDestroy()
 	DebugInfo.Reset();
 #endif
 
-#if WITH_EDITORONLY_DATA
-	if (VMSnapshotBeforeExecution)
-	{
-		VMSnapshotBeforeExecution = nullptr;
-	}
-#endif
-
 	TRACE_OBJECT_LIFETIME_END(this);
 }
 
@@ -360,42 +340,6 @@ void URigVMHost::Evaluate_AnyThread()
 	for (const FName& EventName : EventQueueToRun)
 	{
 		Execute(EventName);
-
-#if WITH_EDITOR
-		if (VM)
-		{
-			const FRigVMBreakpoint& Breakpoint = GetHaltedAtBreakpoint(); 
-			if (Breakpoint.IsValid())
-			{
-				// make sure that the instruction index for the breakpoint is part
-				// of the current entry.
-				const FRigVMByteCode& ByteCode = VM->GetByteCode();
-				const int32 CurrentEntryIndex = ByteCode.FindEntryIndex(EventName);
-				if(CurrentEntryIndex != INDEX_NONE)
-				{
-					const FRigVMByteCodeEntry& CurrentEntry = ByteCode.GetEntry(CurrentEntryIndex);
-					if(Breakpoint.InstructionIndex >= CurrentEntry.InstructionIndex)
-					{
-						int32 LastInstructionIndexInEntry;
-						if(CurrentEntryIndex == ByteCode.NumEntries() - 1)
-						{
-							LastInstructionIndexInEntry = ByteCode.GetNumInstructions();
-						}
-						else
-						{
-							const int32 NextEntryIndex = CurrentEntryIndex + 1;
-							LastInstructionIndexInEntry = ByteCode.GetEntry(NextEntryIndex).InstructionIndex - 1;
-						}
-
-						if(Breakpoint.InstructionIndex <= LastInstructionIndexInEntry)
-						{
-							break;
-						}
-					}
-				}
-			}
-		}
-#endif
 	}
 }
 
@@ -544,14 +488,7 @@ bool URigVMHost::Execute(const FName& InEventName)
 	{
 		if (CDO)
 		{
-			// Copy the breakpoints. This will not override the state of the breakpoints
-			DebugInfo.SetBreakpoints(CDO->DebugInfo.GetBreakpoints());
-
-			// If there are any breakpoints, create the Snapshot VM if it hasn't been created yet
-			if (DebugInfo.GetBreakpoints().Num() > 0)
-			{
-				GetSnapshotVM();
-			}
+			DebugInfo.CopyFrom(CDO->DebugInfo);
 		}
 
 		ExtendedExecuteContext.SetDebugInfo(&DebugInfo);
@@ -713,31 +650,6 @@ bool URigVMHost::Execute_Internal(const FName& InEventName)
 	}
 #endif
 	
-	const bool bUseDebuggingSnapshots = !VM->IsNativized();
-	
-#if WITH_EDITOR
-	if(bUseDebuggingSnapshots)
-	{
-		if(URigVM* SnapShotVM = GetSnapshotVM(false)) // don't create it for normal runs
-		{
-			const bool bIsEventFirstInQueue = !EventQueueToRun.IsEmpty() && EventQueueToRun[0] == InEventName; 
-			const bool bIsEventLastInQueue = !EventQueueToRun.IsEmpty() && EventQueueToRun.Last() == InEventName;
-
-			if (GetHaltedAtBreakpoint().IsValid())
-			{
-				if(bIsEventFirstInQueue)
-				{
-					CopyVMMemory(GetRigVMExtendedExecuteContext(), GetSnapshotContext());
-				}
-			}
-			else if(bIsEventLastInQueue)
-			{
-				CopyVMMemory(GetSnapshotContext(), GetRigVMExtendedExecuteContext());
-			}
-		}
-	}
-#endif
-
 	const bool bSuccess = VM->ExecuteVM(ExtendedExecuteContext, InEventName) != ERigVMExecuteResult::Failed;
 
 #if UE_RIGVM_PROFILE_EXECUTE_UNITS_NUM
@@ -1189,18 +1101,6 @@ TArray<const UObject*> URigVMHost::GetUserDefinedDependencies(const TArray<const
 
 void URigVMHost::HandleExecutionReachedExit(const FName& InEventName)
 {
-#if WITH_EDITOR
-	if (EventQueueToRun.IsEmpty() || EventQueueToRun.Last() == InEventName)
-	{
-		if(URigVM* SnapShotVM = GetSnapshotVM(false))
-		{
-			CopyVMMemory(GetSnapshotContext(), GetRigVMExtendedExecuteContext());
-		}
-		DebugInfo.ResetState();
-		SetBreakpointAction(ERigVMBreakpointAction::None);
-	}
-#endif
-	
 	if (bAccumulateTime)
 	{
 		AbsoluteTime += DeltaTime;
@@ -1302,11 +1202,6 @@ void URigVMHost::InstantiateVMFromCDO()
 				ExtendedExecuteContext.WorkMemoryStorage = CDO->VM->GetDefaultWorkMemory();
 				ExtendedExecuteContext.DebugMemoryStorage = CDO->VM->GetDefaultDebugMemory();
 				ExtendedExecuteContext.VMHash = CDO->VM->GetVMHash();
-
-#if WITH_EDITOR
-				// Fix AutoCompile while stopped in a breakpoint
-				CopyVMMemory(GetSnapshotContext(), ExtendedExecuteContext);
-#endif // WITH_EDITOR
 			}
 		}
 		else if (VM)
@@ -1475,34 +1370,6 @@ void URigVMHost::LogOnce(EMessageSeverity::Type InSeverity, int32 InInstructionI
 	LoggedMessages.Add(InMessage, true);
 }
 
-void URigVMHost::AddBreakpoint(int32 InstructionIndex, UObject* InSubject, uint16 InDepth)
-{
-	DebugInfo.AddBreakpoint(InstructionIndex, InSubject, InDepth);
-}
-
-bool URigVMHost::ExecuteBreakpointAction(const ERigVMBreakpointAction BreakpointAction)
-{
-	if (GetHaltedAtBreakpoint().IsValid())
-	{
-		SetBreakpointAction(BreakpointAction);
-		return true;
-	}
-	return false;
-}
-
-URigVM* URigVMHost::GetSnapshotVM(bool bCreateIfNeeded)
-{
-	if ((VMSnapshotBeforeExecution == nullptr) && bCreateIfNeeded)
-	{
-		VMSnapshotBeforeExecution = NewObject<URigVM>(GetTransientPackage(), NAME_None, RF_Transient);
-	}
-	return VMSnapshotBeforeExecution;
-}
-
-FRigVMExtendedExecuteContext& URigVMHost::GetSnapshotContext()
-{
-	return SnapshotContext;
-}
 #endif
 
 #undef LOCTEXT_NAMESPACE
diff --git a/Engine/Plugins/Runtime/RigVM/Source/RigVM/Public/RigVMCore/RigVM.h b/Engine/Plugins/Runtime/RigVM/Source/RigVM/Public/RigVMCore/RigVM.h
--- a/Engine/Plugins/Runtime/RigVM/Source/RigVM/Public/RigVMCore/RigVM.h
+++ b/Engine/Plugins/Runtime/RigVM/Source/RigVM/Public/RigVMCore/RigVM.h
@@ -157,13 +157,6 @@ class URigVM : public UObject
 	GENERATED_BODY()
 
 public:
-
-	/** Bindable event for external objects to be notified when the VM reaches an Exit Operation */
-	DECLARE_EVENT_OneParam(URigVM, FExecutionReachedExitEvent, const FName&);
-#if WITH_EDITOR
-	DECLARE_EVENT_ThreeParams(URigVM, FExecutionHaltedEvent, int32, UObject*, const FName&);
-#endif
-
 	RIGVM_API URigVM();
 	RIGVM_API virtual ~URigVM();
 
@@ -528,17 +521,9 @@ public:
 		return TArray<int32>();
 	}
 
+	UE_DEPRECATED(5.7, "ResumeExecution is no longer supported.")
 	RIGVM_API bool ResumeExecution(FRigVMExtendedExecuteContext& Context, const FName& InEntryName = NAME_None);
 
-	UE_DEPRECATED(5.3, "Please, use ResumeExecution with Context param")
-	bool ResumeExecution() { return false; }
-	UE_DEPRECATED(5.3, "Please, use ResumeExecution with Context param")
-	bool ResumeExecution(TArrayView<URigVMMemoryStorage*> Memory, const FName& InEntryName = NAME_None) { return false; }
-	UE_DEPRECATED(5.4, "Please, use ResumeExecution with Context param and entry name")
-	bool ResumeExecution(FRigVMExtendedExecuteContext& Context, TArrayView<URigVMMemoryStorage*> Memory, const FName& InEntryName = NAME_None) { return false; }
-	UE_DEPRECATED(5.4, "Please, use ResumeExecution in the Host RigVMDebugInfo")
-	bool ResumeExecution(FRigVMExtendedExecuteContext& Context) { return false; }
-
 #endif // WITH_EDITOR
 
 	// Returns the parameters of the VM
@@ -923,30 +908,9 @@ public:
 	RIGVM_API FString GetOperandLabel(FRigVMExtendedExecuteContext& Context, const FRigVMOperand & InOperand, TFunction<FString(const FString& RegisterName, const FString& RegisterOffsetName)> FormatFunction = nullptr);
 #endif
 
-	UE_DEPRECATED(5.3, "Please use ExecutionReachedExit in the ExtendedExecuteContext.")
-	FExecutionReachedExitEvent& ExecutionReachedExit() 
-	{
-		PRAGMA_DISABLE_DEPRECATION_WARNINGS
-		return OnExecutionReachedExit;
-		PRAGMA_ENABLE_DEPRECATION_WARNINGS
-	}
 #if WITH_EDITOR
-	UE_DEPRECATED(5.3, "Please use ExecutionHalted in the ExtendedExecuteContext.")
-	FExecutionHaltedEvent& ExecutionHalted()
-	{
-		PRAGMA_DISABLE_DEPRECATION_WARNINGS
-		return OnExecutionHalted;
-		PRAGMA_ENABLE_DEPRECATION_WARNINGS
-	}
-
 	UE_DEPRECATED(5.3, "Please use SetDebugInfo in the ExtendedExecuteContext.")
 	void SetDebugInfo(FRigVMDebugInfo* InDebugInfo) { }
-
-	UE_DEPRECATED(5.3, "Please use GetHaltedAtBreakpoint in the RigVMHost.")
-	const FRigVMBreakpoint& GetHaltedAtBreakpoint() const { static FRigVMBreakpoint Dummy; return Dummy; }
-
-	UE_DEPRECATED(5.3, "Please use SetBreakpointAction in the RigVMHost.")
-	void SetBreakpointAction(const ERigVMBreakpointAction& Action) { }
 #endif
 
 	UE_DEPRECATED(5.3, "Please use GetNumExecutions() in the ExtendedExecuteContext.")
@@ -1035,10 +999,6 @@ public:
 
 private:
 
-#if WITH_EDITOR
-	RIGVM_API bool ShouldHaltAtInstruction(FRigVMExtendedExecuteContext& Context, const FName& InEventName, const int32 InstructionIndex);
-#endif
-
 	UPROPERTY()
 	TArray<FName> FunctionNamesStorage;
 	TArray<FName>* FunctionNamesPtr;
@@ -1169,14 +1129,6 @@ private:
 
 	mutable TArray<FName> EntryNames;
 
-	UE_DEPRECATED(5.3, "Please use OnExecutionReachedExit in the ExtendedExecuteContext.")
-	FExecutionReachedExitEvent OnExecutionReachedExit;
-
-#if WITH_EDITOR
-	UE_DEPRECATED(5.3, "Please use OnExecutionHalted in the ExtendedExecuteContext.")
-	FExecutionHaltedEvent OnExecutionHalted;
-#endif
-
 protected:
 	RIGVM_API const URigVMHost* GetHostCDO() const;
 
@@ -1189,6 +1141,7 @@ protected:
 	RIGVM_API void StopProfiling(FRigVMExtendedExecuteContext& Context);
 
 	FTransactionallySafeCriticalSection ResolveFunctionsMutex;
+	FTransactionallySafeCriticalSection CreateLocalizedRegistryMutex;
 
 	friend class URigVMCompiler;
 	friend struct FRigVMCompilerWorkData;
diff --git a/Engine/Plugins/Runtime/RigVM/Source/RigVM/Public/RigVMCore/RigVMDebugInfo.h b/Engine/Plugins/Runtime/RigVM/Source/RigVM/Public/RigVMCore/RigVMDebugInfo.h
--- a/Engine/Plugins/Runtime/RigVM/Source/RigVM/Public/RigVMCore/RigVMDebugInfo.h
+++ b/Engine/Plugins/Runtime/RigVM/Source/RigVM/Public/RigVMCore/RigVMDebugInfo.h
@@ -16,87 +16,6 @@
 class UObject;
 class URigVMNode;
 
-// Possible actions when the VM is halted at a breakpoint
-UENUM()
-enum class ERigVMBreakpointAction : uint8
-{
-	None,
-	Resume,
-	StepOver,
-	StepInto,
-	StepOut,
-	Max UMETA(Hidden),
-};
-
-USTRUCT()
-struct FRigVMBreakpoint
-{
-	GENERATED_BODY()
-
-	FRigVMBreakpoint()
-	: bIsActive(true)
-	, InstructionIndex(INDEX_NONE)
-	, Subject(nullptr)
-	, Depth(0)
-	{
-		Guid.Invalidate();
-	}
-	
-	FRigVMBreakpoint(const int32 InInstructionIndex, TWeakObjectPtr<UObject> InNode, const uint16 InDepth)
-	: bIsActive(true)
-	, Guid(FGuid::NewGuid())
-	, InstructionIndex(InInstructionIndex)
-	, Subject(InNode)
-	, Depth(InDepth)
-	{
-
-	}
-
-	bool IsValid() const
-	{
-		return InstructionIndex != INDEX_NONE && Guid.IsValid();
-	}
-
-	void Reset()
-	{
-		bIsActive = true;
-		InstructionIndex = INDEX_NONE;
-		Guid.Invalidate();
-		Subject = nullptr;
-		Depth = 0;
-	}
-
-	operator bool() const
-	{
-		return IsValid();
-	}
-
-	// Whether or not the breakpoint is active
-	bool bIsActive;
-
-	// guid used to identify the breakpoint
-	FGuid Guid;
-	
-	// Instruction where this breakpoint is set
-	int32 InstructionIndex;
-
-	// Node related to the breakpoint
-	TWeakObjectPtr<UObject> Subject; 
-
-	// The same instruction might be related to multiple breakpoints at different callstack depths
-	uint16 Depth; 
-
-	bool operator==(const FRigVMBreakpoint& Other) const
-	{
-		return Guid == Other.Guid;
-	}
-
-	bool operator!=(const FRigVMBreakpoint& Other) const
-	{
-		return Guid != Other.Guid;
-	}
-};
-
 USTRUCT()
 struct FRigVMDebugInfo
 {
@@ -106,99 +25,40 @@ struct FRigVMDebugInfo
 	{
 	}
 
-	UE_API void ResetState();
-
-	UE_API void StartExecution();
-
 	void Reset()
 	{
-		Breakpoints.Reset();
-
-		// Do not remove state
-		OnExecutionHalted.Clear();
-
-		HaltedAtBreakpoint = FRigVMBreakpoint();
-		HaltedAtBreakpointHit = INDEX_NONE;
-		CurrentBreakpointAction = ERigVMBreakpointAction::None;
+		EarlyExitInstruction.Reset();
 	}
 
 	bool IsEmpty() const
 	{
-		return GetBreakpoints().IsEmpty() && !TemporaryBreakpoint.IsValid();
+		return !EarlyExitInstruction.IsSet();
 	}
 
-	UE_API const FRigVMBreakpoint& FindBreakpoint(const int32 InInstructionIndex, const UObject* InSubject) const;
-	UE_API TArray<FRigVMBreakpoint> FindBreakpointsAtInstruction(const int32 InInstructionIndex) const;
-	UE_API const FRigVMBreakpoint& FindBreakpoint(const FGuid& InGuid) const;
-	
-	UE_API const FRigVMBreakpoint& AddBreakpoint(const int32 InstructionIndex, TWeakObjectPtr<UObject> InNode, const uint16 InDepth, const bool bIsTemporary = false);
-
-	UE_API bool RemoveBreakpoint(const FRigVMBreakpoint& Breakpoint);
-
-	const TArray<FRigVMBreakpoint>& GetBreakpoints() const { return Breakpoints; }
-
-	void SetBreakpoints(const TArray<FRigVMBreakpoint>& InBreakpoints)
+	int32 GetInstructionForEarlyExit() const
 	{
-		Breakpoints = InBreakpoints;
+		return EarlyExitInstruction.Get(INDEX_NONE);
+	}
 
-		if (CurrentActiveBreakpoint.IsValid() && !FindBreakpoint(CurrentActiveBreakpoint).IsValid())
-		{
-			CurrentActiveBreakpoint.Invalidate();
-			CurrentActiveBreakpointCallstack.Reset();
-		}
+	void SetInstructionForEarlyExit(const int32 InInstruction)
+	{
+		EarlyExitInstruction = InInstruction;
 	}
 
-	bool IsTemporaryBreakpoint(FRigVMBreakpoint Breakpoint) const
+	bool ShouldEarlyExitBeforeInstruction(const int32 InInstruction) const
 	{
-		if (Breakpoint.IsValid())
+		if (EarlyExitInstruction.IsSet())
 		{
-			return Breakpoint == TemporaryBreakpoint;
+			return InInstruction > EarlyExitInstruction.GetValue();
 		}
 		return false;
 	}
 
-	UE_API bool IsActive(const FRigVMBreakpoint& InBreakpoint) const;
-
-	UE_API void SetBreakpointHits(const FRigVMBreakpoint& InBreakpoint, const uint16 InBreakpointHits);
-
-	UE_API void HitBreakpoint(const FRigVMBreakpoint& InBreakpoint);
-
-	UE_API void SetBreakpointActivationOnHit(const FRigVMBreakpoint& InBreakpoint, const uint16 InActivationOnHit);
-
-	UE_API void IncrementBreakpointActivationOnHit(const FRigVMBreakpoint& InBreakpoint);
-
-	UE_API uint16 GetBreakpointHits(const FRigVMBreakpoint& InBreakpoint) const;
-
-	const FRigVMBreakpoint& GetCurrentActiveBreakpoint() const { return FindBreakpoint(CurrentActiveBreakpoint); }
-
-	void SetCurrentActiveBreakpoint(const FRigVMBreakpoint& InBreakpoint)
+	void CopyFrom(const FRigVMDebugInfo& Other)
 	{
-		if(InBreakpoint.IsValid())
-		{
-			CurrentActiveBreakpoint = InBreakpoint.Guid;
-		}
-		else
-		{
-			CurrentActiveBreakpoint.Invalidate();
-		}
+		EarlyExitInstruction = Other.EarlyExitInstruction;
 	}
 
-	inline TArray<TWeakObjectPtr<UObject>>& GetCurrentActiveBreakpointCallstack() { return CurrentActiveBreakpointCallstack; }
-
-	inline void SetCurrentActiveBreakpointCallstack(TArray<TWeakObjectPtr<UObject>> Callstack) { CurrentActiveBreakpointCallstack = Callstack; }
-
-	inline FRigVMBreakpoint& GetHaltedAtBreakpoint() { return HaltedAtBreakpoint; }
-	inline const FRigVMBreakpoint& GetHaltedAtBreakpoint() const { return HaltedAtBreakpoint; }
-	inline void SetHaltedAtBreakpoint(const FRigVMBreakpoint& InHaltedAtBreakpoint) { HaltedAtBreakpoint = InHaltedAtBreakpoint; }
-
-	inline int32 GetHaltedAtBreakpointHit() const { return HaltedAtBreakpointHit; }
-	inline void SetHaltedAtBreakpointHit(int32 InHaltedAtBreakpointHit) { HaltedAtBreakpointHit = InHaltedAtBreakpointHit; }
-
-	inline ERigVMBreakpointAction GetCurrentBreakpointAction() const { return CurrentBreakpointAction; }
-	inline void SetCurrentBreakpointAction(ERigVMBreakpointAction InCurrentBreakpointAction) { CurrentBreakpointAction = InCurrentBreakpointAction; }
-
-	UE_API bool ResumeExecution();
-
 	DECLARE_EVENT_ThreeParams(URigVM, FExecutionHaltedEvent, int32, UObject*, const FName&);
 	inline FExecutionHaltedEvent& ExecutionHalted()
 	{
@@ -206,21 +66,8 @@ struct FRigVMDebugInfo
 	}
 
 private:
-	TArray<FRigVMBreakpoint> Breakpoints;
-	FRigVMBreakpoint TemporaryBreakpoint;
-	TMap<FGuid, uint16> BreakpointActivationOnHit; // After how many instruction executions, this breakpoint becomes active
-	TMap<FGuid, uint16> BreakpointHits; // How many times this instruction has been executed
-
-	FGuid CurrentActiveBreakpoint;
-	TArray<TWeakObjectPtr<UObject>> CurrentActiveBreakpointCallstack;
-
-	// --- ---
-
-	FRigVMBreakpoint HaltedAtBreakpoint;
-	int32 HaltedAtBreakpointHit = INDEX_NONE;
-	ERigVMBreakpointAction CurrentBreakpointAction = ERigVMBreakpointAction::None;
-
 	FExecutionHaltedEvent OnExecutionHalted;
+	TOptional<int32> EarlyExitInstruction;
 };
 
 #undef UE_API
diff --git a/Engine/Plugins/Runtime/RigVM/Source/RigVM/Public/RigVMHost.h b/Engine/Plugins/Runtime/RigVM/Source/RigVM/Public/RigVMHost.h
--- a/Engine/Plugins/Runtime/RigVM/Source/RigVM/Public/RigVMHost.h
+++ b/Engine/Plugins/Runtime/RigVM/Source/RigVM/Public/RigVMHost.h
@@ -490,11 +490,6 @@ private:
 
 #if WITH_EDITORONLY_DATA
 
-	UPROPERTY(transient)
-	TObjectPtr<URigVM> VMSnapshotBeforeExecution;
-	UPROPERTY(transient)
-	FRigVMExtendedExecuteContext SnapshotContext;
-
 	/** The current execution mode */
 	UPROPERTY(transient)
 	bool bIsInDebugMode;
@@ -520,33 +515,12 @@ public:
 		return IsInDebugMode() || VMRuntimeSettings.bEnableProfiling;
 	}
 	
-	/** Adds a breakpoint in the VM at the InstructionIndex for the Node / Subject */
-	UE_API void AddBreakpoint(int32 InstructionIndex, UObject* InSubject, uint16 InDepth);
-
-	/** If the VM is halted at a breakpoint, it sets a breakpoint action so that
-	 *  it is applied on the next VM execution */
-	UE_API bool ExecuteBreakpointAction(const ERigVMBreakpointAction BreakpointAction);
-
-	const FRigVMBreakpoint& GetHaltedAtBreakpoint() const
-	{
-		return GetDebugInfo().GetHaltedAtBreakpoint();
-	}
-
-	void SetBreakpointAction(const ERigVMBreakpointAction& Action)
-	{
-		GetDebugInfo().SetCurrentBreakpointAction(Action);
-	}
-	
 	FRigVMDebugInfo& GetDebugInfo() { return DebugInfo; }
 	const FRigVMDebugInfo& GetDebugInfo() const { return DebugInfo; }
 
 	FRigVMProfilingInfo& GetProfilingInfo() { return ProfilingInfo; }
 	const FRigVMProfilingInfo& GetProfilingInfo() const { return ProfilingInfo; }
 
-	/** Creates the snapshot VM if required and returns it */
-	UE_API URigVM* GetSnapshotVM(bool bCreateIfNeeded = true);
-	UE_API FRigVMExtendedExecuteContext& GetSnapshotContext();
-
 #endif	
 
 private:
diff --git a/Engine/Plugins/Runtime/RigVM/Source/RigVMDeveloper/Private/RigVMAsset.cpp b/Engine/Plugins/Runtime/RigVM/Source/RigVMDeveloper/Private/RigVMAsset.cpp
--- a/Engine/Plugins/Runtime/RigVM/Source/RigVMDeveloper/Private/RigVMAsset.cpp
+++ b/Engine/Plugins/Runtime/RigVM/Source/RigVMDeveloper/Private/RigVMAsset.cpp
@@ -925,6 +925,68 @@ const FLazyName& IRigVMAssetInterface::GetPanelPinFactoryName() const
 	return RigVMPanelPinFactoryName;
 }
 
+bool IRigVMAssetInterface::SetEarlyExitInstruction(int32 InInstructionIndex)
+{
+	if (InInstructionIndex == INDEX_NONE)
+	{
+		return ResetEarlyExitInstruction();
+	}
+
+	const FRigVMByteCode* ByteCode = GetController()->GetCurrentByteCode();
+	if (ByteCode == nullptr)
+	{
+		return false;
+	}
+
+	if (InInstructionIndex < 0 || InInstructionIndex >= ByteCode->GetNumInstructions())
+	{
+		return false;
+	}
+
+	GetDebugInfo().SetInstructionForEarlyExit(InInstructionIndex);
+
+	if (URigVMHost* DebuggedHost = GetDebuggedRigVMHost())
+	{
+		DebuggedHost->GetDebugInfo().CopyFrom(GetDebugInfo());
+	}
+	
+	return true;
+}
+
+bool IRigVMAssetInterface::SetEarlyExitInstruction(URigVMNode* InNodeToExitEarlyAfter)
+{
+	if (InNodeToExitEarlyAfter == nullptr)
+	{
+		return ResetEarlyExitInstruction();
+	}
+
+	const FRigVMByteCode* ByteCode = GetController()->GetCurrentByteCode();
+	if (ByteCode == nullptr)
+	{
+		return false;
+	}
+
+	const TArray<int32>& InstructionIndices = ByteCode->GetAllInstructionIndicesForSubject(InNodeToExitEarlyAfter);
+	if (InstructionIndices.IsEmpty())
+	{
+		return false;
+	}
+
+	return SetEarlyExitInstruction(InstructionIndices.Last());
+}
+
+bool IRigVMAssetInterface::ResetEarlyExitInstruction()
+{
+	GetDebugInfo().Reset();
+
+	if (URigVMHost* DebuggedHost = GetDebuggedRigVMHost())
+	{
+		DebuggedHost->GetDebugInfo().CopyFrom(GetDebugInfo());
+	}
+
+	return true;
+}
+
 IRigVMEditorModule* IRigVMAssetInterface::GetEditorModule() const
 {
 	return &IRigVMEditorModule::Get();
@@ -1487,7 +1549,7 @@ void IRigVMAssetInterface::RecompileVM()
 		URigVMCompiler* Compiler = URigVMCompiler::StaticClass()->GetDefaultObject<URigVMCompiler>();
 		GetVMCompileSettings().SetExecuteContextStruct(RigVMClient->GetDefaultExecuteContextStruct());
 
-		const FRigVMCompileSettings Settings = (bCompileInDebugMode) ? FRigVMCompileSettings::Fast(GetVMCompileSettings().GetExecuteContextStruct()) : GetVMCompileSettings();
+		const FRigVMCompileSettings Settings = GetVMCompileSettings();
 		Compiler->Compile(Settings, RigVMClient->GetAllModels(false, false), GetOrCreateController(), VM, *Context, GetExternalVariables(false), &PinToOperandMap);
 		
 
@@ -1513,10 +1575,6 @@ void IRigVMAssetInterface::RecompileVM()
 
 		bVMRecompilationRequired = false;
 		VMCompiledEvent.Broadcast(GetObject(), VM, *Context);
-
-#if WITH_EDITOR
-		RefreshBreakpoints();
-#endif
 	}
 }
 
@@ -1714,233 +1772,6 @@ TArray<IRigVMGraphFunctionHost*> IRigVMAssetInterface::GetReferencedFunctionHost
 
 #if WITH_EDITOR
 
-void IRigVMAssetInterface::SetDebugMode(const bool bValue)
-{
-	bCompileInDebugMode = bValue;
-}
-
-void IRigVMAssetInterface::ClearBreakpoints()
-{
-	for(URigVMNode* Node : GetRigVMBreakpointNodes())
-	{
-		Node->SetHasBreakpoint(false);		
-	}
-	
-	GetRigVMBreakpointNodes().Empty();
-	RefreshBreakpoints();
-}
-
-bool IRigVMAssetInterface::AddBreakpoint(const FString& InBreakpointNodePath)
-{
-	URigVMLibraryNode* FunctionNode = nullptr;
-	
-	// Find the node in the graph
-	FRigVMClient* RigVMClient = GetRigVMClientHost()->GetRigVMClient();
-	for(const URigVMGraph* Model : *RigVMClient)
-	{
-		URigVMNode* BreakpointNode = Model->FindNode(InBreakpointNodePath);
-		if (BreakpointNode == nullptr)
-		{
-			// If we cannot find the node, it might be because it is inside a function
-			FString FunctionName = InBreakpointNodePath, Right;
-			URigVMNode::SplitNodePathAtStart(InBreakpointNodePath, FunctionName, Right);
-
-			// Look inside the local function library
-			if (URigVMLibraryNode* LibraryNode = GetLocalFunctionLibrary()->FindFunction(FName(FunctionName)))
-			{
-				BreakpointNode = LibraryNode->GetContainedGraph()->FindNode(Right);
-				FunctionNode = LibraryNode;
-			}
-		}
-
-		if(BreakpointNode)
-		{
-			return AddBreakpoint(BreakpointNode, FunctionNode);
-		}
-	}
-	return false;
-}
-
-bool IRigVMAssetInterface::AddBreakpoint(URigVMNode* InBreakpointNode, URigVMLibraryNode* LibraryNode)
-{
-	if (InBreakpointNode == nullptr)
-	{
-		return false;
-	}
-
-	bool bSuccess = true;
-	if (LibraryNode)
-	{
-		// If the breakpoint node is inside a library node, find all references to the library node
-		TArray<TSoftObjectPtr<URigVMFunctionReferenceNode>> References = LibraryNode->GetLibrary()->GetReferencesForFunction(LibraryNode->GetFName());
-		for (TSoftObjectPtr<URigVMFunctionReferenceNode> Reference : References)
-		{
-			if (!Reference.IsValid())
-			{
-				continue;
-			}
-
-			IRigVMAssetInterface* ReferenceBlueprint = Reference->GetImplementingOuter<IRigVMAssetInterface>();
-
-			// If the reference is not inside another function, add a breakpoint in the blueprint containing the
-			// reference, without a function specified
-			bool bIsInsideFunction = Reference->GetRootGraph()->IsA<URigVMFunctionLibrary>();
-			if(!bIsInsideFunction)
-			{
-				bSuccess &= ReferenceBlueprint->AddBreakpoint(InBreakpointNode);
-			}
-			else
-			{
-				// Otherwise, we need to add breakpoints to all the blueprints that reference this
-				// function (when the blueprint graph is flattened)
-				
-				// Get all the functions containing this reference
-				URigVMNode* Node = Reference.Get();
-				while (Node->GetGraph() != ReferenceBlueprint->GetLocalFunctionLibrary())
-				{
-					if (URigVMLibraryNode* ParentLibraryNode = Cast<URigVMLibraryNode>(Node->GetGraph()->GetOuter()))
-					{
-						// Recursively add breakpoints to the reference blueprint, specifying the parent function
-						bSuccess &= ReferenceBlueprint->AddBreakpoint(InBreakpointNode, ParentLibraryNode);
-					}
-
-					Node = Cast<URigVMNode>(Node->GetGraph()->GetOuter());
-				}
-			}
-		}
-	}
-	else
-	{
-		if (!GetRigVMBreakpointNodes().Contains(InBreakpointNode))
-		{
-			// Add the breakpoint to the VM
-			bSuccess = AddBreakpointToHost(InBreakpointNode);
-			BreakpointAddedEvent.Broadcast();
-		}
-	}
-
-	return bSuccess;
-}
-
-bool IRigVMAssetInterface::AddBreakpointToHost(URigVMNode* InBreakpointNode)
-{
-	const FRigVMByteCode* ByteCode = GetController()->GetCurrentByteCode();
-	FRigVMClient* RigVMClient = GetRigVMClientHost()->GetRigVMClient();
-	TSet<FString> AddedCallpaths;
-
-	if (ByteCode)
-	{
-		FRigVMInstructionArray Instructions = ByteCode->GetInstructions();	
-		
-		// For each instruction, see if the node is in the callpath
-		// Only add one breakpoint for each callpath related to this node (i.e. if a node produces multiple
-		// instructions, only add a breakpoint to the first instruction)
-		for (int32 i = 0; i< Instructions.Num(); ++i)
-		{
-			for(URigVMGraph* Model : *RigVMClient)
-			{
-				const FRigVMASTProxy Proxy = FRigVMASTProxy::MakeFromCallPath(ByteCode->GetCallPathForInstruction(i), Model);
-				if (Proxy.GetCallstack().Contains(InBreakpointNode))
-				{
-					// Find the callpath related to the breakpoint node
-					FRigVMASTProxy BreakpointProxy = Proxy;
-					while(BreakpointProxy.GetSubject() != InBreakpointNode)
-					{
-						BreakpointProxy = BreakpointProxy.GetParent();
-					}
-					const FString& BreakpointCallPath = BreakpointProxy.GetCallstack().GetCallPath();
-
-					// Only add this callpath breakpoint once
-					if (!AddedCallpaths.Contains(BreakpointCallPath))
-					{
-						AddedCallpaths.Add(BreakpointCallPath);
-						GetDebugInfo().AddBreakpoint(i, InBreakpointNode, static_cast<uint16>(BreakpointProxy.GetCallstack().Num()));
-					}
-				}
-			}
-		}
-	}
-
-	if (AddedCallpaths.Num() > 0)
-	{
-		GetRigVMBreakpointNodes().AddUnique(InBreakpointNode);
-		return true;
-	}
-	
-	return false;
-}
-
-bool IRigVMAssetInterface::RemoveBreakpoint(const FString& InBreakpointNodePath)
-{
-	// Find the node in the graph
-	URigVMNode* BreakpointNode = nullptr;
-	FRigVMClient* RigVMClient = GetRigVMClientHost()->GetRigVMClient();
-	
-	for(URigVMGraph* Model : *RigVMClient)
-	{
-		BreakpointNode = Model->FindNode(InBreakpointNodePath);
-		if (BreakpointNode == nullptr)
-		{
-			// If we cannot find the node, it might be because it is inside a function
-			FString FunctionName = InBreakpointNodePath, Right;
-			URigVMNode::SplitNodePathAtStart(InBreakpointNodePath, FunctionName, Right);
-
-			// Look inside the local function library
-			if (URigVMLibraryNode* LibraryNode = GetLocalFunctionLibrary()->FindFunction(FName(FunctionName)))
-			{
-				BreakpointNode = LibraryNode->GetContainedGraph()->FindNode(Right);
-			}
-		}
-		if(BreakpointNode)
-		{
-			break;
-		}
-	}
-
-	if(BreakpointNode)
-	{
-		bool bSuccess = RemoveBreakpoint(BreakpointNode);
-
-		// Remove the breakpoint from all the loaded dependent blueprints
-		TArray<IRigVMAssetInterface*> DependentBlueprints = GetDependentResolvedAssets(true, true);
-		DependentBlueprints.Remove(this);
-		for (IRigVMAssetInterface* Dependent : DependentBlueprints)
-		{
-			bSuccess &= Dependent->RemoveBreakpoint(BreakpointNode);
-		}
-		return bSuccess;
-	}
-
-	return false;
-}
-
-bool IRigVMAssetInterface::RemoveBreakpoint(URigVMNode* InBreakpointNode)
-{
-	if (GetRigVMBreakpointNodes().Contains(InBreakpointNode))
-	{
-		GetRigVMBreakpointNodes().Remove(InBreakpointNode);
-
-		// Multiple breakpoint nodes might set a breakpoint to the same instruction. When we remove
-		// one of the breakpoint nodes, we do not want to remove the instruction breakpoint if there
-		// is another breakpoint node addressing it. For that reason, we just recompute all the
-		// breakpoint instructions.
-		// Refreshing breakpoints in the control rig will keep the state it had before.
-		RefreshBreakpoints();
-		return true;
-	}
-
-	return false;
-}
-
-void IRigVMAssetInterface::RefreshBreakpoints()
-{
-	GetDebugInfo().Reset();
-	for (URigVMNode* Node : GetRigVMBreakpointNodes())
-	{
-		AddBreakpointToHost(Node);
-	}
-}
-
 TArray<FRigVMReferenceNodeData> IRigVMAssetInterface::GetReferenceNodeData() const
 {
 	TArray<FRigVMReferenceNodeData> Data;
@@ -3132,14 +2963,7 @@ void IRigVMAssetInterface::HandleModifiedEvent(ERigVMGraphNotifType InNotifType,
 			case ERigVMGraphNotifType::NodeAdded:
 			case ERigVMGraphNotifType::NodeRemoved:
 			{
-				bool bAdded = InNotifType == ERigVMGraphNotifType::NodeAdded;
-				if (!bAdded)
-				{
-					if (URigVMNode* RigVMNode = Cast<URigVMNode>(InSubject))
-					{
-						RemoveBreakpoint(RigVMNode);
-					}
-				}
+				const bool bAdded = InNotifType == ERigVMGraphNotifType::NodeAdded;
 
 				if (URigVMCollapseNode* CollapseNode = Cast<URigVMCollapseNode>(InSubject))
 				{
diff --git a/Engine/Plugins/Runtime/RigVM/Source/RigVMDeveloper/Private/RigVMModel/RigVMNode.cpp b/Engine/Plugins/Runtime/RigVM/Source/RigVMDeveloper/Private/RigVMModel/RigVMNode.cpp
--- a/Engine/Plugins/Runtime/RigVM/Source/RigVMDeveloper/Private/RigVMModel/RigVMNode.cpp
+++ b/Engine/Plugins/Runtime/RigVM/Source/RigVMDeveloper/Private/RigVMModel/RigVMNode.cpp
@@ -27,8 +27,7 @@ URigVMNode::URigVMNode()
 , Size(FVector2D::ZeroVector)
 , NodeColor(FLinearColor::White)
 , NodeColorType(ERigVMNodeColorType::FromMetadata)
-, bHasBreakpoint(false)
-, bHaltedAtThisNode(false)
+, bHasEarlyExitMarker(false)
 , NodeVersion(0)
 #if WITH_EDITOR
 , ProfilingHash(0)
diff --git a/Engine/Plugins/Runtime/RigVM/Source/RigVMDeveloper/Public/RigVMAsset.h b/Engine/Plugins/Runtime/RigVM/Source/RigVMDeveloper/Public/RigVMAsset.h
--- a/Engine/Plugins/Runtime/RigVM/Source/RigVMDeveloper/Public/RigVMAsset.h
+++ b/Engine/Plugins/Runtime/RigVM/Source/RigVMDeveloper/Public/RigVMAsset.h
@@ -270,22 +270,16 @@ public:
 
 	TMap<FString, FRigVMOperand> PinToOperandMap;
 
-	/** Removes all the breakpoints from the blueprint and the VM */
-	void ClearBreakpoints();
+	/**
+	 * Sets the instruction index to exit early on 
+	 */
+	bool SetEarlyExitInstruction(int32 InInstructionIndex);
+	bool SetEarlyExitInstruction(URigVMNode* InNodeToExitEarlyAfter);
 
-	/** Adds a breakpoint to all loaded blueprints which use the node indicated by InBreakpointNodePath 
-	  * If the node is inside a public function, it will add a breakpoint to all blueprints calling this function. */
-	bool AddBreakpoint(const FString& InBreakpointNodePath);
-
-	/** Adds a breakpoint to all loaded blueprints which use the InBreakpointNode. 
-	  * If LibraryNode is not null, it indicates that the library uses the InBreapointNode, and the function will add
-	  * breakpoints to any other loaded blueprint that references this library. */
-	bool AddBreakpoint(URigVMNode* InBreakpointNode, URigVMLibraryNode* LibraryNode = nullptr);
-
-	/** Removes the given breakpoint from all the loaded blueprints that use this node, and recomputes all breakpoints
-	  * in the VM  */
-	bool RemoveBreakpoint(const FString& InBreakpointNodePath);
-	bool RemoveBreakpoint(URigVMNode* InBreakpointNode);
+	/**
+	 * Resets / removes the instruction index to exit early on 
+	 */
+	bool ResetEarlyExitInstruction();
 
 	virtual FRigVMController_RequestJumpToHyperlinkDelegate& OnRequestJumpToHyperlink()  { return RequestJumpToHyperlink; };
 
@@ -596,12 +590,10 @@ protected:
 
 	FOnRigVMRefreshEditorEvent RefreshEditorEvent;
 	FOnRigVMVariableDroppedEvent VariableDroppedEvent;
-	FOnRigVMBreakpointAdded BreakpointAddedEvent;
 
 	void SuspendNotifications(bool bSuspendNotifs);
 	
 	virtual FOnRigVMVariableDroppedEvent& OnVariableDropped()  { return VariableDroppedEvent; }
-	virtual FOnRigVMBreakpointAdded& OnBreakpointAdded()  { return BreakpointAddedEvent; }
 
 
 	FOnRigVMCompiledEvent VMCompiledEvent;
@@ -696,29 +688,9 @@ protected:
 	
 	TArray<IRigVMGraphFunctionHost*> GetReferencedFunctionHosts(bool bForceLoad) const;
 
-
-	bool bCompileInDebugMode;
-
-	virtual TArray<TObjectPtr<URigVMNode>>& GetRigVMBreakpointNodes() = 0;
-
 	FOnRigVMRequestInspectObject OnRequestInspectObjectEvent;
 	FOnRigVMRequestInspectMemoryStorage OnRequestInspectMemoryStorageEvent;
 
-	/** Sets the execution mode. In Release mode the rig will ignore all breakpoints. */
-	void SetDebugMode(const bool bValue);
-
-	/** Returns the execution mode */
-	bool IsInDebugMode() const { return bCompileInDebugMode; }
-
-	
-	/** Adds a breakpoint to the first instruction of each callpath related to the InBreakpointNode */
-	bool AddBreakpointToHost(URigVMNode* InBreakpointNode);
-
-	
-
-	/** Recomputes the instruction breakpoints given the node breakpoints in the blueprint */
-	void RefreshBreakpoints();
-
 	/** Shape libraries to load during package load completed */ 
 	TArray<FString> ShapeLibrariesToLoadOnPackageLoaded;
 
diff --git a/Engine/Plugins/Runtime/RigVM/Source/RigVMDeveloper/Public/RigVMBlueprintLegacy.h b/Engine/Plugins/Runtime/RigVM/Source/RigVMDeveloper/Public/RigVMBlueprintLegacy.h
--- a/Engine/Plugins/Runtime/RigVM/Source/RigVMDeveloper/Public/RigVMBlueprintLegacy.h
+++ b/Engine/Plugins/Runtime/RigVM/Source/RigVMDeveloper/Public/RigVMBlueprintLegacy.h
@@ -454,11 +454,6 @@ protected:
 	TObjectPtr<URigVMHost> EditorHost = nullptr;
 	virtual TObjectPtr<URigVMHost>& GetEditorHost() override { return EditorHost; }
 
-	// Declaring these transient properties here instead of the IRigVMAssetInterface class in order to have them as UPROPERTY, to avoid having stale pointers when GC happens
-	UPROPERTY(transient)
-	TArray<TObjectPtr<URigVMNode>> RigVMBreakpointNodes;
-	virtual TArray<TObjectPtr<URigVMNode>>& GetRigVMBreakpointNodes() override { return RigVMBreakpointNodes; }
-	
 	static FName FindHostMemberVariableUniqueName(TSharedPtr<FKismetNameValidator> InNameValidator, const FString& InBaseName);
 	static int32 AddHostMemberVariable(URigVMBlueprint* InBlueprint, const FName& InVarName, FEdGraphPinType InVarType, bool bIsPublic, bool bIsReadOnly, FString InDefaultValue);
 	virtual FName AddHostMemberVariableFromExternal(FRigVMExternalVariable InVariableToCreate, FString InDefaultValue = FString()) override;
diff --git a/Engine/Plugins/Runtime/RigVM/Source/RigVMDeveloper/Public/RigVMModel/RigVMNode.h b/Engine/Plugins/Runtime/RigVM/Source/RigVMDeveloper/Public/RigVMModel/RigVMNode.h
--- a/Engine/Plugins/Runtime/RigVM/Source/RigVMDeveloper/Public/RigVMModel/RigVMNode.h
+++ b/Engine/Plugins/Runtime/RigVM/Source/RigVMDeveloper/Public/RigVMModel/RigVMNode.h
@@ -343,17 +343,27 @@ public:
 	UFUNCTION(BlueprintCallable, Category = RigVMNode)
 	UE_API virtual TArray<FRigVMUserWorkflow> GetSupportedWorkflows(ERigVMUserWorkflowType InType, const UObject* InSubject) const;
 
+	UE_DEPRECATED(5.7, "Please, use HasEarlyExitMarker")
 	UFUNCTION(BlueprintCallable, Category = RigVMNode)
-	bool HasBreakpoint() const { return bHasBreakpoint; }
+	bool HasBreakpoint() const { return false; }
 
+	UE_DEPRECATED(5.7, "Please, use SetHasEarlyExitMarker")
 	UFUNCTION(BlueprintCallable, Category = RigVMNode)
-	void SetHasBreakpoint(const bool bValue) { bHasBreakpoint = bValue; }
+	void SetHasBreakpoint(const bool bValue) {}
 
 	UFUNCTION(BlueprintCallable, Category = RigVMNode)
-	bool ExecutionIsHaltedAtThisNode() const { return bHaltedAtThisNode; }
+	bool HasEarlyExitMarker() const { return bHasEarlyExitMarker; }
 
 	UFUNCTION(BlueprintCallable, Category = RigVMNode)
-	void SetExecutionIsHaltedAtThisNode(const bool bValue) { bHaltedAtThisNode = bValue; }
+	void SetHasEarlyExitMarker(const bool bValue) { bHasEarlyExitMarker = bValue; }
+
+	UE_DEPRECATED(5.7, "Please, use HasEarlyExitMarker")
+	UFUNCTION(BlueprintCallable, Category = RigVMNode)
+	bool ExecutionIsHaltedAtThisNode() const { return HasEarlyExitMarker(); }
+
+	UE_DEPRECATED(5.7, "Please, use SetHasEarlyExitMarker")
+	UFUNCTION(BlueprintCallable, Category = RigVMNode)
+	void SetExecutionIsHaltedAtThisNode(const bool bValue) { SetHasEarlyExitMarker(bValue); }
 
 	UFUNCTION(BlueprintPure, Category = RigVMNode)
 	UE_API virtual bool IsAggregate() const;
@@ -486,10 +496,7 @@ protected:
 	FName PreviousName;
 
 	UPROPERTY(transient)
-	bool bHasBreakpoint;
-
-	UPROPERTY(transient)
-	bool bHaltedAtThisNode;
+	bool bHasEarlyExitMarker;
 
 	UPROPERTY()
 	TArray<FString> TraitRootPinNames;
diff --git a/Engine/Plugins/Runtime/RigVM/Source/RigVMEditor/Private/Editor/RigVMEditor.cpp b/Engine/Plugins/Runtime/RigVM/Source/RigVMEditor/Private/Editor/RigVMEditor.cpp
--- a/Engine/Plugins/Runtime/RigVM/Source/RigVMEditor/Private/Editor/RigVMEditor.cpp
+++ b/Engine/Plugins/Runtime/RigVM/Source/RigVMEditor/Private/Editor/RigVMEditor.cpp
@@ -56,7 +56,6 @@
 FRigVMEditorBase::FRigVMEditorBase()
 	: bAnyErrorsLeft(false)
 	, KnownInstructionLimitWarnings()
-	, HaltedAtNode(nullptr)
 	, LastDebuggedHost()
 	, bSuspendDetailsPanelRefresh(false)
 	, bDetailsPanelRequiresClear(false)
@@ -65,7 +64,6 @@ FRigVMEditorBase::FRigVMEditorBase()
 	, bRigVMEditorInitialized(false)
 	, bIsCompilingThroughUI(false)
 	, WrapperObjects()
-	, ExecutionMode(ERigVMEditorExecutionModeType_Release)
 	, LastEventQueue()
 
 {
@@ -97,7 +95,6 @@ void FRigVMEditorBase::UnbindEditor()
 
 		RigVMBlueprint->OnRefreshEditor().RemoveAll(this);
 		RigVMBlueprint->OnVariableDropped().RemoveAll(this);
-		RigVMBlueprint->OnBreakpointAdded().RemoveAll(this);
 		RigVMBlueprint->OnNodeDoubleClicked().RemoveAll(this);
 		RigVMBlueprint->OnGraphImported().RemoveAll(this);
 		RigVMBlueprint->OnRequestLocalizeFunctionDialog().RemoveAll(this);
@@ -121,9 +118,6 @@ void FRigVMEditorBase::UnbindEditor()
 		}
 
 #if WITH_EDITOR
-		RigVMBlueprint->SetDebugMode(false);
-		RigVMBlueprint->ClearBreakpoints();
-		SetHaltedNode(nullptr);
 		RigVMBlueprint->OnGetFocusedGraph().Unbind();
 #endif
 
@@ -358,7 +352,6 @@ void FRigVMEditorBase::InitRigVMEditor(const EToolkitMode::Type Mode, const TSha
 
 		InRigVMBlueprint->OnRefreshEditor().AddSP(StaticCastSharedRef<FRigVMEditorBase>(SharedRef()), &FRigVMEditorBase::HandleRefreshEditorFromBlueprint);
 		InRigVMBlueprint->OnVariableDropped().AddSP(StaticCastSharedRef<FRigVMEditorBase>(SharedRef()), &FRigVMEditorBase::HandleVariableDroppedFromBlueprint);
-		InRigVMBlueprint->OnBreakpointAdded().AddSP(StaticCastSharedRef<FRigVMEditorBase>(SharedRef()), &FRigVMEditorBase::HandleBreakpointAdded);
 
 		InRigVMBlueprint->OnNodeDoubleClicked().AddSP(StaticCastSharedRef<FRigVMEditorBase>(SharedRef()), &FRigVMEditorBase::OnNodeDoubleClicked);
 		InRigVMBlueprint->OnGraphImported().AddSP(StaticCastSharedRef<FRigVMEditorBase>(SharedRef()), &FRigVMEditorBase::OnGraphImported);
@@ -1423,56 +1416,6 @@ void FRigVMEditorBase::BindCommands()
 		FExecuteAction::CreateSP(ThisRef, &FRigVMEditorBase::ToggleEventQueue),
 		FCanExecuteAction());
 
-	ToolkitCommands->MapAction(
-		EditorCommands.ToggleExecutionMode,
-		FExecuteAction::CreateSP(ThisRef, &FRigVMEditorBase::ToggleExecutionMode),
-		FCanExecuteAction());
-
-	ToolkitCommands->MapAction(
-		EditorCommands.ReleaseMode,
-		FExecuteAction::CreateSP(ThisRef, &FRigVMEditorBase::SetExecutionMode, ERigVMEditorExecutionModeType_Release),
-		FCanExecuteAction());
-
-	ToolkitCommands->MapAction(
-		EditorCommands.DebugMode,
-		FExecuteAction::CreateSP(ThisRef, &FRigVMEditorBase::SetExecutionMode, ERigVMEditorExecutionModeType_Debug),
-		FCanExecuteAction());
-
-	ToolkitCommands->MapAction(
-		EditorCommands.ResumeExecution,
-		FExecuteAction::CreateSP(ThisRef, &FRigVMEditorBase::HandleBreakpointActionRequested, ERigVMBreakpointAction::Resume),			
-		FIsActionChecked::CreateSP(ThisRef, &FRigVMEditorBase::IsHaltedAtBreakpoint),
-		FIsActionChecked(),
-		FIsActionButtonVisible::CreateSP(ThisRef, &FRigVMEditorBase::IsHaltedAtBreakpoint));
-
-	ToolkitCommands->MapAction(
-		EditorCommands.ShowCurrentStatement,
-		FExecuteAction::CreateSP(ThisRef, &FRigVMEditorBase::HandleShowCurrentStatement),
-		FIsActionChecked::CreateSP(ThisRef, &FRigVMEditorBase::IsHaltedAtBreakpoint),
-		FIsActionChecked(),
-		FIsActionButtonVisible::CreateSP(ThisRef, &FRigVMEditorBase::IsHaltedAtBreakpoint));
-
-	ToolkitCommands->MapAction(
-		EditorCommands.StepOver,
-		FExecuteAction::CreateSP(ThisRef, &FRigVMEditorBase::HandleBreakpointActionRequested, ERigVMBreakpointAction::StepOver),
-		FIsActionChecked::CreateSP(ThisRef, &FRigVMEditorBase::IsHaltedAtBreakpoint),
-		FIsActionChecked(),
-		FIsActionButtonVisible::CreateSP(ThisRef, &FRigVMEditorBase::IsHaltedAtBreakpoint));
-
-	ToolkitCommands->MapAction(
-		EditorCommands.StepInto,
-		FExecuteAction::CreateSP(ThisRef, &FRigVMEditorBase::HandleBreakpointActionRequested, ERigVMBreakpointAction::StepInto),
-		FIsActionChecked::CreateSP(ThisRef, &FRigVMEditorBase::IsHaltedAtBreakpoint),
-		FIsActionChecked(),
-		FIsActionButtonVisible::CreateSP(ThisRef, &FRigVMEditorBase::IsHaltedAtBreakpoint));
-
-	ToolkitCommands->MapAction(
-		EditorCommands.StepOut,
-		FExecuteAction::CreateSP(ThisRef, &FRigVMEditorBase::HandleBreakpointActionRequested, ERigVMBreakpointAction::StepOut),
-		FIsActionChecked::CreateSP(ThisRef, &FRigVMEditorBase::IsHaltedAtBreakpoint),
-		FIsActionChecked(),
-		FIsActionButtonVisible::CreateSP(ThisRef, &FRigVMEditorBase::IsHaltedAtBreakpoint));
-
 	ToolkitCommands->MapAction(
 		EditorCommands.FrameSelection,
 		FExecuteAction::CreateSP(ThisRef, &FRigVMEditorBase::FrameSelection),
@@ -1520,13 +1463,6 @@ void FRigVMEditorBase::ToggleEventQueue()
 	SetEventQueue(LastEventQueue);
 }
 
-void FRigVMEditorBase::ToggleExecutionMode()
-{
-	SetExecutionMode((ExecutionMode == ERigVMEditorExecutionModeType_Debug) ?
-		ERigVMEditorExecutionModeType_Release
-		: ERigVMEditorExecutionModeType_Debug);
-}
-
 TSharedRef<SWidget> FRigVMEditorBase::GenerateEventQueueMenuContent()
 {
 	FMenuBuilder MenuBuilder(true, GetToolkitCommands());
@@ -1538,16 +1474,6 @@ void FRigVMEditorBase::GenerateEventQueueMenuContent(FMenuBuilder& MenuBuilder)
 {
 }
 
-TSharedRef<SWidget> FRigVMEditorBase::GenerateExecutionModeMenuContent()
-{
-	FMenuBuilder MenuBuilder(true, GetToolkitCommands());
-	MenuBuilder.BeginSection(TEXT("Events"));
-	MenuBuilder.AddMenuEntry(FRigVMEditorCommands::Get().ReleaseMode, TEXT("Release"), TAttribute<FText>(), TAttribute<FText>(), GetExecutionModeIcon(ERigVMEditorExecutionModeType_Release));
-	MenuBuilder.AddMenuEntry(FRigVMEditorCommands::Get().DebugMode, TEXT("Debug"), TAttribute<FText>(), TAttribute<FText>(), GetExecutionModeIcon(ERigVMEditorExecutionModeType_Debug));
-	MenuBuilder.EndSection();
-	return MenuBuilder.MakeWidget();
-}
-
 FMenuBuilder FRigVMEditorBase::GenerateBulkEditMenu()
 {
 	FMenuBuilder MenuBuilder(true, GetToolkitCommands());
@@ -2038,63 +1964,6 @@ void FRigVMEditorBase::HandleVMExecutedEvent(URigVMHost* InHost, const FName& In
 	UpdateGraphCompilerErrors();
 }
 
-void FRigVMEditorBase::HandleVMExecutionHalted(const int32 InstructionIndex, UObject* InNodeObject, const FName& InEntryName)
-{
-	if (HaltedAtNode == InNodeObject)
-	{
-		return;
-	}
-		
-	if (URigVMNode* InNode = Cast<URigVMNode>(InNodeObject))
-	{
-		SetHaltedNode(InNode);
-		
-		if (TRigVMAssetInterfacePtr RigVMBlueprint = GetRigVMAsset())
-		{
-			if (RigVMBlueprint->GetAllModels().Contains(InNode->GetGraph()))
-			{
-				if(URigVMEdGraph* EdGraph = Cast<URigVMEdGraph>(RigVMBlueprint->GetEdGraph(InNode->GetGraph())))
-				{
-					if(UEdGraphNode* EdNode = EdGraph->FindNodeForModelNodeName(InNode->GetFName()))
-					{
-						FFunctionGraphTask::CreateAndDispatchWhenReady([this, EdNode]()
-						{
-							JumpToHyperlink(EdNode, false);
-						}, TStatId(), nullptr, ENamedThreads::GameThread);
-					}
-				}
-			}
-		}
-	}
-	else 
-	{
-		if (URigVMHost* RigVMHost = GetRigVMHost())
-		{
-			const TArray<FName>& EventQueue = RigVMHost->GetEventQueue();
-			if (!EventQueue.IsEmpty())
-			{
-				if (InEntryName == EventQueue.Last())
-				{
-					SetHaltedNode(nullptr);
-				}
-			}
-		}
-	}
-}
-
-void FRigVMEditorBase::SetHaltedNode(URigVMNode* Node)
-{
-	if (HaltedAtNode)
-	{
-		HaltedAtNode->SetExecutionIsHaltedAtThisNode(false);
-	}
-	HaltedAtNode = Node;
-	if (HaltedAtNode)
-	{
-		HaltedAtNode->SetExecutionIsHaltedAtThisNode(true);
-	}
-}
-
 void FRigVMEditorBase::NotifyPreChange(FProperty* PropertyAboutToChange)
 {
 	NotifyPreChangeImpl(PropertyAboutToChange);
@@ -2679,10 +2548,6 @@ void FRigVMEditorBase::UpdateRigVMHost()
  		}
 		check(IsValid(RigVMHost));
 
-#if WITH_EDITOR
-		RigVMHost->SetIsInDebugMode(ExecutionMode == ERigVMEditorExecutionModeType_Debug);
-#endif
-
 		CacheNameLists();
 
 		// Make sure the object being debugged is the preview instance
@@ -2698,7 +2563,6 @@ void FRigVMEditorBase::UpdateRigVMHost()
 		RigVMHost->OnInitialized_AnyThread().AddSP(StaticCastSharedRef<FRigVMEditorBase>(SharedRef()), &FRigVMEditorBase::HandleVMExecutedEvent);
 		RigVMHost->OnExecuted_AnyThread().AddSP(StaticCastSharedRef<FRigVMEditorBase>(SharedRef()), &FRigVMEditorBase::HandleVMExecutedEvent);
 		RigVMHost->RequestInit();
-		RigVMHost->GetDebugInfo().ExecutionHalted().AddSP(StaticCastSharedRef<FRigVMEditorBase>(SharedRef()), &FRigVMEditorBase::HandleVMExecutionHalted);
 	}
 }
 
@@ -3630,43 +3494,6 @@ void FRigVMEditorBase::FillToolbar(FToolBarBuilder& ToolbarBuilder, bool bEndSec
 		AddSelectedDebugObjectWidget(ToolbarBuilder);
 
 		ToolbarBuilder.AddSeparator();
-		ToolbarBuilder.AddToolBarButton(
-			FRigVMEditorCommands::Get().ToggleExecutionMode,
-			NAME_None, 
-			TAttribute<FText>::Create(TAttribute<FText>::FGetter::CreateSP(StaticCastSharedRef<FRigVMEditorBase>(SharedRef()), &FRigVMEditorBase::GetExecutionModeLabel)),
-			TAttribute<FText>(), 
-			TAttribute<FSlateIcon>::Create(TAttribute<FSlateIcon>::FGetter::CreateSP(StaticCastSharedRef<FRigVMEditorBase>(SharedRef()), &FRigVMEditorBase::GetExecutionModeIcon))
-		);
-		
-		FUIAction DefaultExecutionMode;
-		ToolbarBuilder.AddComboButton(
-			DefaultExecutionMode,
-			FOnGetContent::CreateSP(StaticCastSharedRef<FRigVMEditorBase>(SharedRef()), &FRigVMEditorBase::GenerateExecutionModeMenuContent),
-			LOCTEXT("ExecutionMode_Label", "Execution Modes"),
-			LOCTEXT("ExecutionMode_ToolTip", "Pick between different execution modes for testing the Control Rig"),
-			FSlateIcon(FAppStyle::GetAppStyleSetName(), "LevelEditor.Recompile"),
-			true);
-
-		ToolbarBuilder.BeginStyleOverride(FName("Toolbar.BackplateLeftPlay"));
-		ToolbarBuilder.AddToolBarButton(FRigVMEditorCommands::Get().ResumeExecution,
-			NAME_None, TAttribute<FText>(), TAttribute<FText>(), FSlateIcon(FAppStyle::GetAppStyleSetName(), "PlayWorld.ResumePlaySession"));
-
-		ToolbarBuilder.BeginStyleOverride(FName("Toolbar.BackplateLeft"));
-		ToolbarBuilder.AddToolBarButton(FRigVMEditorCommands::Get().ShowCurrentStatement,
-			NAME_None, TAttribute<FText>(), TAttribute<FText>(), FSlateIcon(FAppStyle::GetAppStyleSetName(), "PlayWorld.ShowCurrentStatement"));
-
-		ToolbarBuilder.BeginStyleOverride(FName("Toolbar.BackplateCenter"));
-		ToolbarBuilder.AddToolBarButton(FRigVMEditorCommands::Get().StepOver,
-			NAME_None, TAttribute<FText>(), TAttribute<FText>(), FSlateIcon(FAppStyle::GetAppStyleSetName(), "PlayWorld.StepOver"));
-		ToolbarBuilder.AddToolBarButton(FRigVMEditorCommands::Get().StepInto,
-			NAME_None, TAttribute<FText>(), TAttribute<FText>(), FSlateIcon(FAppStyle::GetAppStyleSetName(), "PlayWorld.StepInto"));
-
-		ToolbarBuilder.BeginStyleOverride(FName("Toolbar.BackplateRight"));
-		ToolbarBuilder.AddToolBarButton(FRigVMEditorCommands::Get().StepOut,
-			NAME_None, TAttribute<FText>(), TAttribute<FText>(), FSlateIcon(FAppStyle::GetAppStyleSetName(), "PlayWorld.StepOut"));
-
-		ToolbarBuilder.EndStyleOverride();
-
 
 		FUIAction DefaultBulkEditAction;
 		ToolbarBuilder.AddComboButton(
@@ -3824,11 +3651,6 @@ void FRigVMEditorBase::HandleVariableDroppedFromBlueprint(UObject* InSubject, FP
 	);
 }
 
-void FRigVMEditorBase::HandleBreakpointAdded()
-{
-	SetExecutionMode(ERigVMEditorExecutionModeType_Debug);
-}
-
 void FRigVMEditorBase::OnGraphNodeClicked(URigVMEdGraphNode* InNode)
 {
 	if (InNode)
@@ -3885,36 +3707,6 @@ bool FRigVMEditorBase::OnActionMatchesName(FEdGraphSchemaAction* InAction, const
 	return false;
 }
 
-void FRigVMEditorBase::HandleShowCurrentStatement()
-{
-	if(HaltedAtNode)
-	{
-		if(TRigVMAssetInterfacePtr Blueprint = GetRigVMAsset())
-		{
-			if(URigVMEdGraph* EdGraph = Cast<URigVMEdGraph>(Blueprint->GetEdGraph(HaltedAtNode->GetGraph())))
-			{
-				if(UEdGraphNode* EdNode = EdGraph->FindNodeForModelNodeName(HaltedAtNode->GetFName()))
-				{
-					JumpToHyperlink(EdNode, false);
-				}
-			}
-		}
-	}
-}
-
-void FRigVMEditorBase::HandleBreakpointActionRequested(const ERigVMBreakpointAction BreakpointAction)
-{
-	if(URigVMHost* DebuggedHost = Cast<URigVMHost>(GetRigVMAsset()->GetObjectBeingDebugged()))
-	{
-		DebuggedHost->ExecuteBreakpointAction(BreakpointAction);
-	}
-}
-
-bool FRigVMEditorBase::IsHaltedAtBreakpoint() const
-{
-	return HaltedAtNode != nullptr;
-}
-
 void FRigVMEditorBase::FrameSelection()
 {
 	if (SGraphEditor* GraphEd = GetFocusedGraphEditor().Pin().Get())
@@ -4284,8 +4076,6 @@ void FRigVMEditorBase::SetEventQueue(TArray<FName> InEventQueue, bool bCompile)
 
 	LastEventQueue = GetEventQueue();
 
-	SetHaltedNode(nullptr);
-
 	if (URigVMHost* CurrentHost = GetRigVMHost())
 	{
 		if (InEventQueue.Num() > 0)
@@ -4305,55 +4095,6 @@ FSlateIcon FRigVMEditorBase::GetEventQueueIcon() const
 	return GetEventQueueIcon(GetEventQueue());
 }
 
-void FRigVMEditorBase::SetExecutionMode(const ERigVMEditorExecutionModeType InExecutionMode)
-{
-	if (ExecutionMode == InExecutionMode)
-	{
-		return;
-	}
-
-	ExecutionMode = InExecutionMode;
-	GetRigVMAsset()->SetDebugMode(InExecutionMode == ERigVMEditorExecutionModeType_Debug);
-	Compile();
-	
-	if (URigVMHost* CurrentHost = GetRigVMHost())
-	{
-		CurrentHost->SetIsInDebugMode(InExecutionMode == ERigVMEditorExecutionModeType_Debug);
-	}
-
-	SetHaltedNode(nullptr);
-	
-	RefreshDetailView();
-}
-
-int32 FRigVMEditorBase::GetExecutionModeComboValue() const
-{
-	return (int32) ExecutionMode;
-}
-
-FText FRigVMEditorBase::GetExecutionModeLabel() const
-{
-	if (ExecutionMode == ERigVMEditorExecutionModeType_Debug)
-	{
-		return FText::FromString(TEXT("DebugMode"));
-	}
-	return FText::FromString(TEXT("ReleaseMode"));
-}
-
-FSlateIcon FRigVMEditorBase::GetExecutionModeIcon(ERigVMEditorExecutionModeType InExecutionMode)
-{
-	if (InExecutionMode == ERigVMEditorExecutionModeType_Debug)
-	{
-		return FSlateIcon(FAppStyle::GetAppStyleSetName(), "LevelEditor.Tabs.Debug");
-	}
-	return FSlateIcon(FRigVMEditorStyle::Get().GetStyleSetName(), "RigVM.ReleaseMode");
-}
-
-FSlateIcon FRigVMEditorBase::GetExecutionModeIcon() const
-{
-	return GetExecutionModeIcon(ExecutionMode);
-}
-
 void FRigVMEditorBase::GetDebugObjects(TArray<FRigVMCustomDebugObject>& DebugList) const
 {
 	TRigVMAssetInterfacePtr RigVMBlueprint = GetRigVMAsset();
@@ -4502,10 +4243,6 @@ void FRigVMEditorBase::HandleSetObjectBeingDebugged(UObject* InObject)
 	{
 		DebuggedHost->SetLog(&RigVMLog);
 		DebuggedHost->OnExecuted_AnyThread().AddSP(StaticCastSharedRef<FRigVMEditorBase>(SharedRef()), &FRigVMEditorBase::HandleVMExecutedEvent);
-#if WITH_EDITOR
-		DebuggedHost->GetDebugInfo().ExecutionHalted().AddSP(StaticCastSharedRef<FRigVMEditorBase>(SharedRef()), &FRigVMEditorBase::HandleVMExecutionHalted);
-		DebuggedHost->SetIsInDebugMode(ExecutionMode == ERigVMEditorExecutionModeType_Debug);	
-#endif
 	}
 
 	RefreshDetailView();
diff --git a/Engine/Plugins/Runtime/RigVM/Source/RigVMEditor/Private/Editor/RigVMEditorCommands.cpp b/Engine/Plugins/Runtime/RigVM/Source/RigVMEditor/Private/Editor/RigVMEditorCommands.cpp
--- a/Engine/Plugins/Runtime/RigVM/Source/RigVMEditor/Private/Editor/RigVMEditorCommands.cpp
+++ b/Engine/Plugins/Runtime/RigVM/Source/RigVMEditor/Private/Editor/RigVMEditorCommands.cpp
@@ -19,14 +19,6 @@ void FRigVMEditorCommands::RegisterCommands()
 	UI_COMMAND(ExecuteGraph, "Execute", "Execute the rig graph if On.", EUserInterfaceActionType::ToggleButton, FInputChord());
 	UI_COMMAND(AutoCompileGraph, "Auto Compile", "Auto-compile the rig graph if On.", EUserInterfaceActionType::ToggleButton, FInputChord());
 	UI_COMMAND(ToggleEventQueue, "Toggle Event", "Toggle between the current and last running event", EUserInterfaceActionType::Button, FInputChord());
-	UI_COMMAND(ToggleExecutionMode, "Toggle Execution Mode", "Toggle between Release and Debug execution mode", EUserInterfaceActionType::Button, FInputChord());
-	UI_COMMAND(ReleaseMode, "Release Mode", "Compiles and Executes the rig, ignoring debug data.", EUserInterfaceActionType::Button, FInputChord());
-	UI_COMMAND(DebugMode, "Debug Mode", "Compiles and Executes the unoptimized rig, stopping at breakpoints.", EUserInterfaceActionType::Button, FInputChord());
-	UI_COMMAND(ResumeExecution, "Resume", "Resumes execution after being halted at a breakpoint.", EUserInterfaceActionType::Button, FInputChord(EKeys::F5));
-	UI_COMMAND(ShowCurrentStatement, "Show Current Statement", "Focuses on the node currently being debugged.", EUserInterfaceActionType::Button, FInputChord());
-	UI_COMMAND(StepInto, "Step Into", "Steps into the collapsed/function node, when halted at a breakpoint.", EUserInterfaceActionType::Button, FInputChord(EKeys::F11));
-	UI_COMMAND(StepOut, "Step Out", "Steps out of the collapsed/function node, when halted at a breakpoint.", EUserInterfaceActionType::Button, FInputChord(EKeys::F11, EModifierKey::Shift));
-	UI_COMMAND(StepOver, "Step Over", "Steps over the node, when halted at a breakpoint.", EUserInterfaceActionType::Button, FInputChord(EKeys::F10));
 	UI_COMMAND(FrameSelection, "Frame Selection", "Frames the selected nodes in the Graph View.", EUserInterfaceActionType::Button, FInputChord(EKeys::F));
 	UI_COMMAND(SwapFunctionWithinAsset, "Swap Function (Asset)", "Swaps a function for all occurrences within this asset.", EUserInterfaceActionType::Button, FInputChord());
 	UI_COMMAND(SwapFunctionAcrossProject, "Swap Function (Project)", "Swaps a function for all occurrences in the project.", EUserInterfaceActionType::Button, FInputChord());
diff --git a/Engine/Plugins/Runtime/RigVM/Source/RigVMEditor/Private/Editor/RigVMEditorCommands.h b/Engine/Plugins/Runtime/RigVM/Source/RigVMEditor/Private/Editor/RigVMEditorCommands.h
--- a/Engine/Plugins/Runtime/RigVM/Source/RigVMEditor/Private/Editor/RigVMEditorCommands.h
+++ b/Engine/Plugins/Runtime/RigVM/Source/RigVMEditor/Private/Editor/RigVMEditorCommands.h
@@ -41,30 +41,6 @@ public:
 	/** Toggle between this and the last event queue */
 	TSharedPtr< FUICommandInfo > ToggleEventQueue;
 
-	/** Toggle between Release and Debug execution mode */
-	TSharedPtr< FUICommandInfo > ToggleExecutionMode;
-
-	/** Compile and run the optimized rig, ignoring any debug data */
-	TSharedPtr< FUICommandInfo > ReleaseMode;
-
-	/** Compile and run the unoptimized rig, ignoring any debug data */
-	TSharedPtr< FUICommandInfo > DebugMode;
-	
-	/** Resume the execution of the graph when halted at a breakpoint */
-	TSharedPtr< FUICommandInfo > ResumeExecution;
-
-	/** Focuses on the node currently being debugged */
-	TSharedPtr< FUICommandInfo > ShowCurrentStatement;
-
-	/** Steps to the next node in the execution graph (at the same graph level) when halted at a breakpoint */
-	TSharedPtr< FUICommandInfo > StepOver;
-
-	/** Steps into the collapsed/function node, when halted at a breakpoint */
-	TSharedPtr< FUICommandInfo > StepInto;
-
-	/** Steps out of the collapsed/function node, when halted at a breakpoint */
-	TSharedPtr< FUICommandInfo > StepOut;
-
 	/** Frames the selected nodes */
 	TSharedPtr< FUICommandInfo > FrameSelection;
 
diff --git a/Engine/Plugins/Runtime/RigVM/Source/RigVMEditor/Private/RigVMEditorModule.cpp b/Engine/Plugins/Runtime/RigVM/Source/RigVMEditor/Private/RigVMEditorModule.cpp
--- a/Engine/Plugins/Runtime/RigVM/Source/RigVMEditor/Private/RigVMEditorModule.cpp
+++ b/Engine/Plugins/Runtime/RigVM/Source/RigVMEditor/Private/RigVMEditorModule.cpp
@@ -893,7 +893,7 @@ void FRigVMEditorModule::GetNodeDebugContextMenuActions(TScriptInterface<IRigVMC
 	if(TRigVMAssetInterfacePtr RigVMBlueprint = RigVMClientHost.GetObject())
 	{
 		FToolMenuSection& DebugSection = Menu->AddSection("RigVMEditorContextMenuDebug", LOCTEXT("DebugHeader", "Debug"));
-		bool bNoneHasBreakpoint = true;
+		bool bAnyNodeIsEarlyExit = false;
 
 		const URigVMGraph* Model = ModelNode->GetGraph();
 		URigVMController* Controller = RigVMBlueprint->GetController(Model);
@@ -907,48 +907,44 @@ void FRigVMEditorModule::GetNodeDebugContextMenuActions(TScriptInterface<IRigVMC
 			if (URigVMNode* FoundNode = Model->FindNodeByName(SelectedNodeName))
 			{
 				SelectedNodes.Add(FoundNode);
-				if (FoundNode->HasBreakpoint())
+				if (FoundNode->HasEarlyExitMarker())
 				{
-					bNoneHasBreakpoint = false;
+					bAnyNodeIsEarlyExit = true;
 				}
 			}
 		}
 	
-		if (bNoneHasBreakpoint)
+		if (SelectedNodeNames.Num() == 1)
 		{
 			DebugSection.AddMenuEntry(
-			"Add Breakpoint",
-			LOCTEXT("AddBreakpoint", "Add Breakpoint"),
-			LOCTEXT("AddBreakpoint_Tooltip", "Adds a breakpoint to the graph at this node"),
+			"Preview Here",
+			LOCTEXT("PreviewHere", "Preview Here"),
+			LOCTEXT("PreviewHere_Tooltip", "Preview the result of the node here"),
 			FSlateIcon(),
 			FUIAction(FExecuteAction::CreateLambda([Controller, SelectedNodes, RigVMBlueprint]()
 			{
 				for (URigVMNode* SelectedNode : SelectedNodes)
 				{
-					if (RigVMBlueprint->AddBreakpoint(SelectedNode))
-					{
-						SelectedNode->SetHasBreakpoint(true);
-					}
+					// todooooo
+					//if (RigVMBlueprint->AddBreakpoint(SelectedNode))
 				}
 			})));
 		}
-		else
+		else if (bAnyNodeIsEarlyExit)
 		{						
 			DebugSection.AddMenuEntry(
-			"Remove Breakpoint",
-			LOCTEXT("RemoveBreakpoint", "Remove Breakpoint"),
-			LOCTEXT("RemoveBreakpoint_Tooltip", "Removes a breakpoint to the graph at this node"),
+			"Stop Preview",
+			LOCTEXT("StopPreview", "Stop Preview Here"),
+			LOCTEXT("StopPreview_Tooltip", "Stops the preview from here functionality and runs the whole graph"),
 			FSlateIcon(),
 			FUIAction(FExecuteAction::CreateLambda([Controller, SelectedNodes, RigVMBlueprint]()
 			{
 				for (URigVMNode* SelectedNode : SelectedNodes)
 				{
-					if (SelectedNode->HasBreakpoint())
+					if (SelectedNode->HasEarlyExitMarker())
 					{
-						if (RigVMBlueprint->RemoveBreakpoint(SelectedNode))
-						{
-							SelectedNode->SetHasBreakpoint(false);
-						}
+						// todooooo
+						//if (RigVMBlueprint->RemoveBreakpoint(SelectedNode))
 					}
 				}                            
 			})));
diff --git a/Engine/Plugins/Runtime/RigVM/Source/RigVMEditor/Private/Widgets/SRigVMGraphNode.cpp b/Engine/Plugins/Runtime/RigVM/Source/RigVMEditor/Private/Widgets/SRigVMGraphNode.cpp
--- a/Engine/Plugins/Runtime/RigVM/Source/RigVMEditor/Private/Widgets/SRigVMGraphNode.cpp
+++ b/Engine/Plugins/Runtime/RigVM/Source/RigVMEditor/Private/Widgets/SRigVMGraphNode.cpp
@@ -1078,48 +1078,6 @@ FReply SRigVMGraphNode::HandleAddArrayElement(FString InModelPinPath)
 	return FReply::Handled();
 }
 
-/** Populate the brushes array with any overlay brushes to render */
-void SRigVMGraphNode::GetOverlayBrushes(bool bSelected, const FVector2f& WidgetSize, TArray<FOverlayBrushInfo>& Brushes) const
-{
-	URigVMEdGraphNode* RigGraphNode = Cast<URigVMEdGraphNode>(GraphNode);
-
-	if (const URigVMNode* VMNode = RigGraphNode->GetModelNode())
-	{
-		const bool bHasBreakpoint = VMNode->HasBreakpoint();
-		if (bHasBreakpoint)
-		{
-			FOverlayBrushInfo BreakpointOverlayInfo;
-
-			BreakpointOverlayInfo.Brush = FAppStyle::GetBrush(TEXT("Kismet.DebuggerOverlay.Breakpoint.EnabledAndValid"));
-			if (BreakpointOverlayInfo.Brush != NULL)
-			{
-				BreakpointOverlayInfo.OverlayOffset -= BreakpointOverlayInfo.Brush->ImageSize / 2.f;
-			}
-
-			Brushes.Add(BreakpointOverlayInfo);
-		}
-
-		// Paint red arrow pointing at breakpoint node that caused a halt in execution
-		{
-			FOverlayBrushInfo IPOverlayInfo;
-			if (VMNode->ExecutionIsHaltedAtThisNode())
-			{
-				IPOverlayInfo.Brush = FAppStyle::GetBrush( TEXT("Kismet.DebuggerOverlay.InstructionPointerBreakpoint") );
-				if (IPOverlayInfo.Brush != NULL)
-				{
-					float Overlap = 10.f;
-					IPOverlayInfo.OverlayOffset.X = (WidgetSize.X/2.f) - (IPOverlayInfo.Brush->ImageSize.X/2.f);
-					IPOverlayInfo.OverlayOffset.Y = (Overlap - IPOverlayInfo.Brush->ImageSize.Y);
-				}
-
-				IPOverlayInfo.AnimationEnvelope = FVector2f(0.f, 10.f);
-
-				Brushes.Add(IPOverlayInfo);
-			}
-		}
-	}
-}
-
 void SRigVMGraphNode::GetNodeInfoPopups(FNodeInfoContext* Context, TArray<FGraphInformationPopupInfo>& Popups) const
 {
 	FKismetNodeInfoContext* K2Context = (FKismetNodeInfoContext*)Context;
diff --git a/Engine/Plugins/Runtime/RigVM/Source/RigVMEditor/Public/Editor/RigVMEditor.h b/Engine/Plugins/Runtime/RigVM/Source/RigVMEditor/Public/Editor/RigVMEditor.h
--- a/Engine/Plugins/Runtime/RigVM/Source/RigVMEditor/Public/Editor/RigVMEditor.h
+++ b/Engine/Plugins/Runtime/RigVM/Source/RigVMEditor/Public/Editor/RigVMEditor.h
@@ -402,13 +402,6 @@ protected:
 
 	UE_API virtual bool SelectLocalVariable(const UEdGraph* Graph, const FName& VariableName) override;
 
-
-	enum ERigVMEditorExecutionModeType
-	{
-		ERigVMEditorExecutionModeType_Release,
-		ERigVMEditorExecutionModeType_Debug
-	};
-
 	// FBlueprintEditor Interface
 	UE_API void CreateDefaultCommandsImpl();
 	UE_API void OnCreateGraphEditorCommands(TSharedPtr<FUICommandList> GraphEditorCommandsList);
@@ -435,8 +428,6 @@ protected:
 	virtual void HandleVMCompiledEvent(UObject* InCompiledObject, URigVM* InVM) {}
 	UE_API virtual void HandleVMCompiledEvent(UObject* InCompiledObject, URigVM* InVM, FRigVMExtendedExecuteContext& InContext);
 	UE_API virtual void HandleVMExecutedEvent(URigVMHost* InHost, const FName& InEventName);
-	UE_API void HandleVMExecutionHalted(const int32 InstructionIndex, UObject* InNode, const FName& InEntryName);
-	UE_API void SetHaltedNode(URigVMNode* Node);
 	
 	// FNotifyHook Interface
 	UE_API void NotifyPreChange(FProperty* PropertyAboutToChange);
@@ -470,9 +461,7 @@ protected:
 	UE_API bool IsAutoCompileGraphOn() const;
 	bool CanAutoCompileGraph() const { return true; }
 	UE_API void ToggleEventQueue();
-	UE_API void ToggleExecutionMode();
 	UE_API TSharedRef<SWidget> GenerateEventQueueMenuContent();
-	UE_API TSharedRef<SWidget> GenerateExecutionModeMenuContent();
 	UE_API virtual FMenuBuilder GenerateBulkEditMenu();
 	UE_API TSharedRef<SWidget> GenerateBulkEditMenuContent();
 	UE_API virtual void GenerateEventQueueMenuContent(FMenuBuilder& MenuBuilder);
@@ -531,12 +520,6 @@ protected:
 	UE_API virtual FSlateIcon GetEventQueueIcon(const TArray<FName>& InEventQueue) const;
 	UE_API FSlateIcon GetEventQueueIcon() const;
 
-	UE_API void SetExecutionMode(const ERigVMEditorExecutionModeType InExecutionMode);
-	UE_API int32 GetExecutionModeComboValue() const;
-	UE_API FText GetExecutionModeLabel() const;
-	static UE_API FSlateIcon GetExecutionModeIcon(const ERigVMEditorExecutionModeType InExecutionMode);
-	UE_API FSlateIcon GetExecutionModeIcon() const;
-	
 	UE_API virtual void GetDebugObjects(TArray<FRigVMCustomDebugObject>& DebugList) const override;
 	virtual bool OnlyShowCustomDebugObjects() const override { return true; }
 	UE_API void HandleSetObjectBeingDebugged(UObject* InObject);
@@ -547,14 +530,10 @@ protected:
 	
 	UE_API virtual void HandleRefreshEditorFromBlueprint(TRigVMAssetInterfacePtr InBlueprint);
 	UE_API void HandleVariableDroppedFromBlueprint(UObject* InSubject, FProperty* InVariableToDrop, const FVector2D& InDropPosition, const FVector2D& InScreenPosition);
-	UE_API void HandleBreakpointAdded();
 	UE_API void OnGraphNodeClicked(URigVMEdGraphNode* InNode);
 	UE_API void OnNodeDoubleClicked(TRigVMAssetInterfacePtr InBlueprint, URigVMNode* InNode);
 	UE_API void OnGraphImported(UEdGraph* InEdGraph);
 	UE_API bool OnActionMatchesName(FEdGraphSchemaAction* InAction, const FName& InName) const;
-	UE_API void HandleShowCurrentStatement();
-	UE_API void HandleBreakpointActionRequested(const ERigVMBreakpointAction BreakpointAction);
-	UE_API bool IsHaltedAtBreakpoint() const;
 	UE_API void FrameSelection();
 	UE_API void SwapFunctionWithinAsset();
 	UE_API void SwapFunctionAcrossProject();
@@ -592,7 +571,6 @@ protected:
 protected:
 	bool bAnyErrorsLeft;
 	TMap<FString, FString> KnownInstructionLimitWarnings;
-	URigVMNode* HaltedAtNode;
 	FString LastDebuggedHost;
 
 	bool bSuspendDetailsPanelRefresh;
@@ -607,8 +585,6 @@ protected:
 
 	TArray<TStrongObjectPtr<URigVMDetailsViewWrapperObject>> WrapperObjects;
 
-	ERigVMEditorExecutionModeType ExecutionMode;
-
 	/** The log to use for errors resulting from the init phase of the units */
 	FRigVMLog RigVMLog;
 	
diff --git a/Engine/Plugins/Runtime/RigVM/Source/RigVMEditor/Public/Widgets/SRigVMGraphNode.h b/Engine/Plugins/Runtime/RigVM/Source/RigVMEditor/Public/Widgets/SRigVMGraphNode.h
--- a/Engine/Plugins/Runtime/RigVM/Source/RigVMEditor/Public/Widgets/SRigVMGraphNode.h
+++ b/Engine/Plugins/Runtime/RigVM/Source/RigVMEditor/Public/Widgets/SRigVMGraphNode.h
@@ -95,7 +95,6 @@ public:
 	UE_API virtual FReply OnMouseButtonDoubleClick(const FGeometry& InMyGeometry, const FPointerEvent& InMouseEvent) override;
 
 	UE_API virtual TSharedRef<SWidget> CreateNodeContentArea() override;
-	UE_API virtual void GetOverlayBrushes(bool bSelected, const FVector2f& WidgetSize, TArray<FOverlayBrushInfo>& Brushes) const override;
 	UE_API virtual void GetNodeInfoPopups(FNodeInfoContext* Context, TArray<FGraphInformationPopupInfo>& Popups) const override;
 	UE_API virtual TArray<FOverlayWidgetInfo> GetOverlayWidgets(bool bSelected, const FVector2f& WidgetSize) const override;
 	UE_API virtual void OnFocusLost(const FFocusEvent& InFocusEvent) override;
