RigVM: Implement blend based fading of nodes + flashlight functionality.

The new OnTick delegates allows asset editors to tap into the graph editors' tick - for custom code such as node layout or node style.

#rnx

diff --git a/Engine/Plugins/Animation/ControlRig/Source/ControlRigEditor/Private/RigDependencyGraph/RigDependencyConnectionDrawingPolicy.cpp b/Engine/Plugins/Animation/ControlRig/Source/ControlRigEditor/Private/RigDependencyGraph/RigDependencyConnectionDrawingPolicy.cpp
--- a/Engine/Plugins/Animation/ControlRig/Source/ControlRigEditor/Private/RigDependencyGraph/RigDependencyConnectionDrawingPolicy.cpp
+++ b/Engine/Plugins/Animation/ControlRig/Source/ControlRigEditor/Private/RigDependencyGraph/RigDependencyConnectionDrawingPolicy.cpp
@@ -17,9 +17,7 @@ void FRigDependencyConnectionDrawingPolicy::DetermineWiringStyle(UEdGraphPin* Ou
 			Params.WireColor.A = 0.75f;
 		}
 
-		if (OutputNode->IsFadedOut() || InputNode->IsFadedOut())
-		{
-			Params.WireColor.A *= 0.2f;
-		}
+		const float FadedState = FMath::Min(OutputNode->GetFadedOutState(), InputNode->GetFadedOutState());
+		Params.WireColor.A *= FadedState;
 	}
 }
\ No newline at end of file
diff --git a/Engine/Plugins/Animation/ControlRig/Source/ControlRigEditor/Private/RigDependencyGraph/SRigDependencyGraph.cpp b/Engine/Plugins/Animation/ControlRig/Source/ControlRigEditor/Private/RigDependencyGraph/SRigDependencyGraph.cpp
--- a/Engine/Plugins/Animation/ControlRig/Source/ControlRigEditor/Private/RigDependencyGraph/SRigDependencyGraph.cpp
+++ b/Engine/Plugins/Animation/ControlRig/Source/ControlRigEditor/Private/RigDependencyGraph/SRigDependencyGraph.cpp
@@ -61,6 +61,7 @@ void SRigDependencyGraph::Construct(const FArguments& InArgs, const TSharedRef<I
 		.GraphToEdit(GraphObj)
 		.GraphEvents(GraphEvents)
 		.Appearance(AppearanceInfo)
+		.OnTick(this, &SRigDependencyGraph::OnGraphEditorTick)
 		.ShowGraphStateOverlay(false);
 
 	static const FName DefaultForegroundName("DefaultForeground");
@@ -726,6 +727,8 @@ void SRigDependencyGraph::HandleSelectionChanged(const FGraphPanelSelectionSet&
 
 void SRigDependencyGraph::Tick(const FGeometry& AllottedGeometry, const double InCurrentTime, const float InDeltaTime)
 {
+	SCompoundWidget::Tick(AllottedGeometry, InCurrentTime, InDeltaTime);
+	
 	GraphObj->Tick(InDeltaTime);
 	
 	if (GraphObj->NeedsRefreshLayout())
@@ -742,7 +745,10 @@ void SRigDependencyGraph::Tick(const FGeometry& AllottedGeometry, const double I
 			}
 		}
 	}
+}
 
+void SRigDependencyGraph::OnGraphEditorTick(const FGeometry& AllottedGeometry, const double InCurrentTime, const float InDeltaTime)
+{
 	// flashlight to lighten up nodes in proximity of the mouse
 	const FGeometry PaintGeometry = GraphEditor->GetPaintSpaceGeometry();
 	const FVector2D MousePosition = FSlateApplication::Get().GetCursorPos();
@@ -766,16 +772,12 @@ void SRigDependencyGraph::Tick(const FGeometry& AllottedGeometry, const double I
 				if (GraphEditor->GetBoundsForNode(Node, Bounds, 0.f))
 				{
 					const FVector2f Center = Bounds.GetCenter();
-					const FVector2f Size = Bounds.GetSize();
-					if (!Size.IsNearlyZero())
+					const float Radius = 64.f * FlashLightRadius;
+					const float Ratio = FVector2f::DistSquared(Center, WidgetPosition) / (Radius * Radius);
+					if (Ratio < 1.f)
 					{
-						const float Radius = Size.Length() * 0.5f * FlashLightRadius;
-						const float Ratio = FVector2f::DistSquared(Center, WidgetPosition) / (Radius * Radius);
-						if (Ratio < 1.f)
-						{
-							const float HalfRatio = 1.f - FMath::Clamp(Ratio - 0.5f, 0.f, .5f) / 0.5f;
-							Node->OverrideFadeOutState(FMath::Lerp(PreviousFadedOutState, 1.f, HalfRatio));
-						}
+						const float HalfRatio = 1.f - FMath::Clamp(Ratio - 0.5f, 0.f, .5f) / 0.5f;
+						Node->OverrideFadeOutState(FMath::Lerp(PreviousFadedOutState, 1.f, HalfRatio));
 					}
 				}
 			}
diff --git a/Engine/Plugins/Animation/ControlRig/Source/ControlRigEditor/Private/RigDependencyGraph/SRigDependencyGraph.h b/Engine/Plugins/Animation/ControlRig/Source/ControlRigEditor/Private/RigDependencyGraph/SRigDependencyGraph.h
--- a/Engine/Plugins/Animation/ControlRig/Source/ControlRigEditor/Private/RigDependencyGraph/SRigDependencyGraph.h
+++ b/Engine/Plugins/Animation/ControlRig/Source/ControlRigEditor/Private/RigDependencyGraph/SRigDependencyGraph.h
@@ -52,6 +52,8 @@ public:
 
 private:
 
+	void OnGraphEditorTick(const FGeometry& AllottedGeometry, const double InCurrentTime, const float InDeltaTime);
+
 	/** Called to create context menu when right-clicking on graph */
 	FActionMenuContent OnCreateGraphActionMenu(UEdGraph* InGraph, const FVector2f& InNodePosition, const TArray<UEdGraphPin*>& InDraggedPins, bool bAutoExpand, SGraphEditor::FActionMenuClosed InOnMenuClosed);
 
diff --git a/Engine/Plugins/Runtime/RigVM/Source/RigVMDeveloper/Private/EdGraph/RigVMEdGraphNode.cpp b/Engine/Plugins/Runtime/RigVM/Source/RigVMDeveloper/Private/EdGraph/RigVMEdGraphNode.cpp
--- a/Engine/Plugins/Runtime/RigVM/Source/RigVMDeveloper/Private/EdGraph/RigVMEdGraphNode.cpp
+++ b/Engine/Plugins/Runtime/RigVM/Source/RigVMDeveloper/Private/EdGraph/RigVMEdGraphNode.cpp
@@ -491,6 +491,15 @@ int32 URigVMEdGraphNode::GetInstructionIndex(bool bAsInput) const
 	return INDEX_NONE;
 }
 
+bool URigVMEdGraphNode::IsExcludedFromPreview() const
+{
+	if(const URigVMNode* ModelNode = Cast<URigVMNode>(GetModelNode()))
+	{
+		return ModelNode->IsExcludedByEarlyExit();
+	}
+	return false;
+}
+
 const FRigVMTemplate* URigVMEdGraphNode::GetTemplate() const
 {
 	if(CachedTemplate == nullptr)
@@ -1491,6 +1500,35 @@ void URigVMEdGraphNode::SetModelNode(URigVMNode* InModelNode)
 	//PostReconstructNode();
 }
 
+bool URigVMEdGraphNode::IsFadedOut() const
+{
+	return GetFadedOutState() < 1.f - SMALL_NUMBER;
+}
+
+float URigVMEdGraphNode::GetFadedOutState() const
+{
+#if WITH_EDITOR
+	if (FadedOutOverride.IsSet())
+	{
+		return FadedOutOverride.GetValue();
+	}
+	if(GetInstructionIndex(true) == INDEX_NONE)
+	{
+		return 0.25f;
+	}
+	if (IsExcludedFromPreview())
+	{
+		return 0.4f;
+	}
+#endif
+	return 1.f;
+}
+
+void URigVMEdGraphNode::OverrideFadeOutState(TOptional<float> InFadedOutState)
+{
+	FadedOutOverride = InFadedOutState;
+}
+
 bool URigVMEdGraphNode::ModelPinRemoved_Internal(const URigVMPin* InModelPin)
 {
 	auto RemoveGraphPin = [this](const URigVMPin* InModelPin, bool bAsInput)
@@ -1628,27 +1666,10 @@ UClass* URigVMEdGraphNode::GetRigVMGeneratedClass() const
 	return nullptr;
 }
 
-FLinearColor URigVMEdGraphNode::GetNodeOpacityColor() const
-{
-	if (URigVMNode* ModelNode = GetModelNode())
-	{
-		if (Cast<URigVMVariableNode>(ModelNode))
-		{
-			return FLinearColor::White;
-		}
-
-		if(GetInstructionIndex(true) == INDEX_NONE)
-		{
-			return FLinearColor(0.35f, 0.35f, 0.35f, 0.35f);
-		}
-	}
-	return FLinearColor::White;
-}
-
 FLinearColor URigVMEdGraphNode::GetNodeTitleColor() const
 {
 	// return a darkened version of the default node's color
-	return CachedTitleColor * GetNodeOpacityColor();
+	return CachedTitleColor;
 }
 
 FLinearColor URigVMEdGraphNode::GetNodeBodyTintColor() const
@@ -1663,7 +1684,7 @@ FLinearColor URigVMEdGraphNode::GetNodeBodyTintColor() const
 		return HighlightingColor.GetValue();
 	}
 #endif
-	return CachedNodeColor * GetNodeOpacityColor();
+	return CachedNodeColor;
 }
 
 bool URigVMEdGraphNode::ShowPaletteIconOnNode() const
diff --git a/Engine/Plugins/Runtime/RigVM/Source/RigVMDeveloper/Private/RigVMModel/RigVMNode.cpp b/Engine/Plugins/Runtime/RigVM/Source/RigVMDeveloper/Private/RigVMModel/RigVMNode.cpp
--- a/Engine/Plugins/Runtime/RigVM/Source/RigVMDeveloper/Private/RigVMModel/RigVMNode.cpp
+++ b/Engine/Plugins/Runtime/RigVM/Source/RigVMDeveloper/Private/RigVMModel/RigVMNode.cpp
@@ -28,6 +28,7 @@ URigVMNode::URigVMNode()
 , NodeColor(FLinearColor::White)
 , NodeColorType(ERigVMNodeColorType::FromMetadata)
 , bHasEarlyExitMarker(false)
+, bIsExcludedByEarlyExit(false)
 , NodeVersion(0)
 #if WITH_EDITOR
 , ProfilingHash(0)
@@ -1437,6 +1438,16 @@ TArray<FRigVMUserWorkflow> URigVMNode::GetSupportedWorkflows(ERigVMUserWorkflowT
 	return URigVMUserWorkflowRegistry::Get()->GetWorkflows(InType, Struct, InSubject);
 }
 
+bool URigVMNode::IsExcludedByEarlyExit() const
+{
+	return bIsExcludedByEarlyExit;
+}
+
+void URigVMNode::SetIsExcludedByEarlyExit(bool bIsExcluded)
+{
+	bIsExcludedByEarlyExit = bIsExcluded;
+}
+
 bool URigVMNode::IsAggregate() const
 {
 #if UE_RIGVM_AGGREGATE_NODES_ENABLED
diff --git a/Engine/Plugins/Runtime/RigVM/Source/RigVMDeveloper/Public/EdGraph/RigVMEdGraphNode.h b/Engine/Plugins/Runtime/RigVM/Source/RigVMDeveloper/Public/EdGraph/RigVMEdGraphNode.h
--- a/Engine/Plugins/Runtime/RigVM/Source/RigVMDeveloper/Public/EdGraph/RigVMEdGraphNode.h
+++ b/Engine/Plugins/Runtime/RigVM/Source/RigVMDeveloper/Public/EdGraph/RigVMEdGraphNode.h
@@ -195,6 +195,7 @@ public:
 	UE_API void HandleInsertArrayElement(FString InPinPath);
 
 	UE_API int32 GetInstructionIndex(bool bAsInput) const;
+	UE_API bool IsExcludedFromPreview() const;
 
 	UE_API const FRigVMTemplate* GetTemplate() const;
 
@@ -240,9 +241,19 @@ public:
 
 	const TArray<URigVMPin*>& GetInputPins() const { return InputPins; }
 
-protected:
+	// Returns true if this Node should be faded out
+	UE_API bool IsFadedOut() const;
+
+	// Returns 1.0 if the node is 100% opaque, and 0.0 if the node isn't visible.
+	UE_API float GetFadedOutState() const;
+
+	// Overrides the fade out state of this node
+	UE_API void OverrideFadeOutState(TOptional<float> InFadedOutState);
 
-	UE_API FLinearColor GetNodeOpacityColor() const;
+	// Removes any user based fade out state
+	UE_API void ResetFadedOutState() { OverrideFadeOutState({}); }
+
+protected:
 
 	/** Helper function for AllocateDefaultPins */
 	UE_API void UpdatePinLists();
@@ -324,6 +335,9 @@ private:
 	FSimpleDelegate RequestRenameDelegate;
 	mutable bool bRenameIsPending;
 
+	mutable TOptional<bool> bIsFadedOut;
+	mutable TOptional<float> FadedOutOverride;
+
 	friend class SRigVMGraphNode;
 	friend class FRigVMFunctionArgumentLayout;
 	friend class FRigVMEdGraphDetailCustomization;
diff --git a/Engine/Plugins/Runtime/RigVM/Source/RigVMDeveloper/Public/RigVMModel/RigVMNode.h b/Engine/Plugins/Runtime/RigVM/Source/RigVMDeveloper/Public/RigVMModel/RigVMNode.h
--- a/Engine/Plugins/Runtime/RigVM/Source/RigVMDeveloper/Public/RigVMModel/RigVMNode.h
+++ b/Engine/Plugins/Runtime/RigVM/Source/RigVMDeveloper/Public/RigVMModel/RigVMNode.h
@@ -352,10 +352,16 @@ public:
 	void SetHasBreakpoint(const bool bValue) {}
 
 	UFUNCTION(BlueprintCallable, Category = RigVMNode)
-	bool HasEarlyExitMarker() const { return bHasEarlyExitMarker; }
+	UE_API bool HasEarlyExitMarker() const { return bHasEarlyExitMarker; }
 
 	UFUNCTION(BlueprintCallable, Category = RigVMNode)
-	void SetHasEarlyExitMarker(const bool bValue) { bHasEarlyExitMarker = bValue; }
+	UE_API void SetHasEarlyExitMarker(const bool bValue) { bHasEarlyExitMarker = bValue; }
+
+	UFUNCTION(BlueprintCallable, Category = RigVMNode)
+	UE_API bool IsExcludedByEarlyExit() const;
+
+	UFUNCTION(BlueprintCallable, Category = RigVMNode)
+	UE_API void SetIsExcludedByEarlyExit(bool bIsExcluded);
 
 	UE_DEPRECATED(5.7, "Please, use HasEarlyExitMarker")
 	UFUNCTION(BlueprintCallable, Category = RigVMNode)
@@ -498,6 +504,9 @@ protected:
 	UPROPERTY(transient)
 	bool bHasEarlyExitMarker;
 
+	UPROPERTY(transient)
+	bool bIsExcludedByEarlyExit;
+
 	UPROPERTY()
 	TArray<FString> TraitRootPinNames;
 
diff --git a/Engine/Plugins/Runtime/RigVM/Source/RigVMEditor/Private/EdGraph/RigVMEdGraphConnectionDrawingPolicy.cpp b/Engine/Plugins/Runtime/RigVM/Source/RigVMEditor/Private/EdGraph/RigVMEdGraphConnectionDrawingPolicy.cpp
--- a/Engine/Plugins/Runtime/RigVM/Source/RigVMEditor/Private/EdGraph/RigVMEdGraphConnectionDrawingPolicy.cpp
+++ b/Engine/Plugins/Runtime/RigVM/Source/RigVMEditor/Private/EdGraph/RigVMEdGraphConnectionDrawingPolicy.cpp
@@ -351,5 +351,8 @@ void FRigVMEdGraphConnectionDrawingPolicy::DetermineWiringStyle(UEdGraphPin* Out
 				Params.WireColor = Params.WireColor * 0.5f;
 			}
 		}
+
+		const float FadedState = FMath::Min(OutputNode->GetFadedOutState(), InputNode->GetFadedOutState());
+		Params.WireColor.A *= FadedState;
 	}
 }
diff --git a/Engine/Plugins/Runtime/RigVM/Source/RigVMEditor/Private/Editor/RigVMEditor.cpp b/Engine/Plugins/Runtime/RigVM/Source/RigVMEditor/Private/Editor/RigVMEditor.cpp
--- a/Engine/Plugins/Runtime/RigVM/Source/RigVMEditor/Private/Editor/RigVMEditor.cpp
+++ b/Engine/Plugins/Runtime/RigVM/Source/RigVMEditor/Private/Editor/RigVMEditor.cpp
@@ -4288,6 +4288,63 @@ void FRigVMEditorBase::OnPIEStopped(bool bSimulation)
 	}
 }
 
+void FRigVMEditorBase::OnGraphEditorTick(const FGeometry& AllottedGeometry, const double InCurrentTime, const float InDeltaTime,
+	TSharedRef<SGraphEditor> InGraphEditor, URigVMEdGraph* InGraph)
+{
+	// flashlight to lighten up nodes in proximity of the mouse
+	const FGeometry PaintGeometry = InGraphEditor->GetPaintSpaceGeometry();
+	const FVector2D MousePosition = FSlateApplication::Get().GetCursorPos();
+	const float FlashLightRadius = 5.f * 64.f;
+
+	if(PaintGeometry.IsUnderLocation(MousePosition) && FlashLightRadius > SMALL_NUMBER)
+	{
+		FVector2f GraphEditorLocation = FVector2f::ZeroVector;
+		float GraphEditorZoomAmount = 0;
+		InGraphEditor->GetViewLocation(GraphEditorLocation, GraphEditorZoomAmount);
+
+		const FVector2f WidgetPosition = PaintGeometry.AbsoluteToLocal(MousePosition) / GraphEditorZoomAmount + GraphEditorLocation;
+		for (const TObjectPtr<UEdGraphNode>& Node : InGraph->Nodes)
+		{
+			URigVMEdGraphNode* RigVMEdGraphNode = Cast<URigVMEdGraphNode>(Node);
+			if (RigVMEdGraphNode == nullptr)
+			{
+				continue;
+			}
+
+			RigVMEdGraphNode->ResetFadedOutState();
+
+			const float PreviousFadedOutState = RigVMEdGraphNode->GetFadedOutState();
+			if (PreviousFadedOutState < 1.f - SMALL_NUMBER)
+			{
+				FSlateRect Bounds;
+				if (InGraphEditor->GetBoundsForNode(Node, Bounds, 0.f))
+				{
+					const FVector2f Center = Bounds.GetCenter();
+					const float Radius = FlashLightRadius;
+					const float Ratio = FVector2f::DistSquared(Center, WidgetPosition) / (Radius * Radius);
+					if (Ratio < 1.f)
+					{
+						const float HalfRatio = 1.f - FMath::Clamp(Ratio - 0.5f, 0.f, .5f) / 0.5f;
+						RigVMEdGraphNode->OverrideFadeOutState(FMath::Lerp(PreviousFadedOutState, 1.f, HalfRatio));
+					}
+				}
+			}
+		}
+	}
+	else
+	{
+		for (const TObjectPtr<UEdGraphNode>& Node : InGraph->Nodes)
+		{
+			URigVMEdGraphNode* RigVMEdGraphNode = Cast<URigVMEdGraphNode>(Node);
+			if (RigVMEdGraphNode == nullptr)
+			{
+				continue;
+			}
+			RigVMEdGraphNode->ResetFadedOutState();
+		}
+	}
+}
+
 bool FRigVMEditorBase::IsEditingSingleBlueprint() const
 {
 	return GetRigVMAsset() != nullptr;
diff --git a/Engine/Plugins/Runtime/RigVM/Source/RigVMEditor/Private/Editor/RigVMLegacyEditor.cpp b/Engine/Plugins/Runtime/RigVM/Source/RigVMEditor/Private/Editor/RigVMLegacyEditor.cpp
--- a/Engine/Plugins/Runtime/RigVM/Source/RigVMEditor/Private/Editor/RigVMLegacyEditor.cpp
+++ b/Engine/Plugins/Runtime/RigVM/Source/RigVMEditor/Private/Editor/RigVMLegacyEditor.cpp
@@ -166,6 +166,23 @@ TSharedRef<SGraphEditor> FRigVMLegacyEditor::CreateGraphEditorWidget(TSharedRef<
 {
 	TSharedRef<SGraphEditor> GraphEditor = FBlueprintEditor::CreateGraphEditorWidget(InTabInfo, InGraph);
 	GraphEditor->GetGraphPanel()->SetZoomLevelsContainer<FRigVMEditorZoomLevelsContainer>();
+
+	TWeakObjectPtr<URigVMEdGraph> WeakGraph(Cast<URigVMEdGraph>(InGraph));
+	if (WeakGraph.IsValid())
+	{
+		TWeakPtr<FRigVMEditorBase> WeakThisEditor = SharedThis(this).ToWeakPtr();
+		TWeakPtr<SGraphEditor> WeakGraphEditor = GraphEditor.ToWeakPtr();
+		
+		GraphEditor->SetOnTick(FOnGraphEditorTickDelegate::CreateLambda([WeakThisEditor, WeakGraphEditor, WeakGraph]
+			(const FGeometry& AllottedGeometry, const double InCurrentTime, const float InDeltaTime)
+			{
+				if (WeakThisEditor.IsValid() && WeakGraphEditor.IsValid() && WeakGraph.IsValid())
+				{
+					WeakThisEditor.Pin()->OnGraphEditorTick(AllottedGeometry, InCurrentTime, InDeltaTime, WeakGraphEditor.Pin().ToSharedRef(), WeakGraph.Get());
+				}
+			}
+		));
+	}
 	return GraphEditor;
 }
 
diff --git a/Engine/Plugins/Runtime/RigVM/Source/RigVMEditor/Private/Editor/RigVMNewEditor.cpp b/Engine/Plugins/Runtime/RigVM/Source/RigVMEditor/Private/Editor/RigVMNewEditor.cpp
--- a/Engine/Plugins/Runtime/RigVM/Source/RigVMEditor/Private/Editor/RigVMNewEditor.cpp
+++ b/Engine/Plugins/Runtime/RigVM/Source/RigVMEditor/Private/Editor/RigVMNewEditor.cpp
@@ -1262,6 +1262,23 @@ TSharedRef<SGraphEditor> FRigVMNewEditor::CreateGraphEditorWidget(TSharedRef<FTa
 		.AssetEditorToolkit(this->GetHostingApp());
 		//@TODO: Crashes in command list code during the callback .OnGraphModuleReloaded(FEdGraphEvent::CreateSP(this, &FBlueprintEditor::ChangeOpenGraphInDocumentEditorWidget, WeakParent))
 		;
+
+	TWeakObjectPtr<URigVMEdGraph> WeakGraph(Cast<URigVMEdGraph>(InGraph));
+	if (WeakGraph.IsValid())
+	{
+		TWeakPtr<FRigVMEditorBase> WeakThisEditor = SharedThis(this).ToWeakPtr();
+		TWeakPtr<SGraphEditor> WeakGraphEditor = Editor.ToWeakPtr();
+		
+		Editor->SetOnTick(FOnGraphEditorTickDelegate::CreateLambda([WeakThisEditor, WeakGraphEditor, WeakGraph]
+			(const FGeometry& AllottedGeometry, const double InCurrentTime, const float InDeltaTime)
+			{
+				if (WeakThisEditor.IsValid() && WeakGraphEditor.IsValid() && WeakGraph.IsValid())
+				{
+					WeakThisEditor.Pin()->OnGraphEditorTick(AllottedGeometry, InCurrentTime, InDeltaTime, WeakGraphEditor.Pin().ToSharedRef(), WeakGraph.Get());
+				}
+			}
+		));
+	}
 		
 	//OnSetPinVisibility.AddSP(&Editor.Get(), &SGraphEditor::SetPinVisibility);
 
diff --git a/Engine/Plugins/Runtime/RigVM/Source/RigVMEditor/Private/Widgets/SRigVMGraphNode.cpp b/Engine/Plugins/Runtime/RigVM/Source/RigVMEditor/Private/Widgets/SRigVMGraphNode.cpp
--- a/Engine/Plugins/Runtime/RigVM/Source/RigVMEditor/Private/Widgets/SRigVMGraphNode.cpp
+++ b/Engine/Plugins/Runtime/RigVM/Source/RigVMEditor/Private/Widgets/SRigVMGraphNode.cpp
@@ -322,6 +322,45 @@ void SRigVMGraphNode::Construct( const FArguments& InArgs )
 
 	const FSlateBrush* ImageBrush = FSlateIcon(TEXT("RigVMEditorStyle"), TEXT("RigVM.Bug.Dot")).GetIcon();
 
+	PreviewHereIndicatorWidget =
+		SNew(SButton)
+		.Visibility(EVisibility::Visible)
+		.ButtonStyle( FAppStyle::Get(), "SimpleButton" )
+		.ToolTipText(LOCTEXT("PreviewHereTooltip", "Toggles the Preview Here functionality.\nBy pressing this you can skip the execution of consecutive nodes to preview the results of the graph up to this point."))
+		.ContentPadding(0.0f)
+		.OnHovered_Lambda([this]
+		{
+			bPreviewHereIndicatorHovered = true;
+		})
+		.OnUnhovered_Lambda([this]
+		{
+			bPreviewHereIndicatorHovered = false;
+		})
+		.OnClicked_Lambda([this]() -> FReply
+		{
+			if (ModelNode.IsValid() && Blueprint.IsValid())
+			{
+				if (ModelNode->HasEarlyExitMarker())
+				{
+					Blueprint->ResetEarlyExitInstruction();
+				}
+				else
+				{
+					Blueprint->SetEarlyExitInstruction(ModelNode.Get());
+				}
+				return FReply::Handled();
+			}
+			return FReply::Unhandled();
+		})
+		.Content()
+		[
+			SNew(SImage)
+			.Image(FAppStyle::GetBrush(TEXT("Icons.Preview")))
+			.ColorAndOpacity(this, &SRigVMGraphNode::GetPreviewHereColor)
+			.Visibility(EVisibility::Visible)
+		];
+	bPreviewHereIndicatorHovered = false;
+
 	VisualDebugIndicatorWidget =
 		SNew(SImage)
 		.Image(ImageBrush)
@@ -543,6 +582,19 @@ TOptional<FSlateColor> SRigVMGraphNode::GetPinTextColor(const SGraphPin* InGraph
 	return Result;
 }
 
+int32 SRigVMGraphNode::OnPaint(const FPaintArgs& Args, const FGeometry& AllottedGeometry, const FSlateRect& MyCullingRect,
+	FSlateWindowElementList& OutDrawElements, int32 LayerId, const FWidgetStyle& InWidgetStyle, bool bParentEnabled) const
+{
+	FWidgetStyle WidgetStyle = InWidgetStyle;
+	
+	if (const URigVMEdGraphNode* DependencyGraphNode = Cast<URigVMEdGraphNode>(GraphNode))
+	{
+		WidgetStyle.BlendOpacity(DependencyGraphNode->GetFadedOutState());
+	}
+	
+	return SGraphNode::OnPaint(Args, AllottedGeometry, MyCullingRect, OutDrawElements, LayerId, WidgetStyle, bParentEnabled);
+}
+
 void SRigVMGraphNode::CreateAddPinButton()
 {
 	if (LeftNodeBox.IsValid() && ModelNode.IsValid())
@@ -1225,6 +1277,18 @@ TArray<FOverlayWidgetInfo> SRigVMGraphNode::GetOverlayWidgets(bool bSelected, co
 	if (ModelNode.IsValid())
 	{
 		float HorizontalOffset = 0;
+
+		if (ModelNode->IsMutable())
+		{
+			const FVector2f ImageSize = PreviewHereIndicatorWidget->GetDesiredSize();
+
+			FOverlayWidgetInfo Info;
+			Info.OverlayOffset = FVector2f(WidgetSize.X - ImageSize.X - 6.f - HorizontalOffset, 6.f);
+			HorizontalOffset += ImageSize.X + 6.f;
+			Info.Widget = PreviewHereIndicatorWidget;
+
+			Widgets.Add(Info);
+		}
 		
 		if(CVarRigVMEnablePinOverrides.GetValueOnAnyThread())
 		{
@@ -2913,6 +2977,26 @@ EOverrideWidgetStatus::Type SRigVMGraphNode::GetOverrideStatus() const
 	return EOverrideWidgetStatus::ChangedHere;
 }
 
+FSlateColor SRigVMGraphNode::GetPreviewHereColor() const
+{
+	if (const URigVMNode* CurrentModelNode = ModelNode.Get())
+	{
+		if (CurrentModelNode->HasEarlyExitMarker())
+		{
+			return FSlateColor(EStyleColor::AccentBlue);
+		}
+		else if (CurrentModelNode->IsExcludedByEarlyExit())
+		{
+			return FSlateColor(EStyleColor::AccentBlack);
+		}
+	}
+	if (bPreviewHereIndicatorHovered)
+	{
+		return FSlateColor(EStyleColor::Hover);
+	}
+	return FSlateColor::UseForeground();
+}
+
 int32 SRigVMGraphNode::CanReorderExposedPin(const URigVMPin* InPinToReorder, const URigVMPin* InTargetPin, EItemDropZone InDropZone)
 {
 	if (!InPinToReorder || !InTargetPin || InPinToReorder == InTargetPin)
diff --git a/Engine/Plugins/Runtime/RigVM/Source/RigVMEditor/Public/Editor/RigVMEditor.h b/Engine/Plugins/Runtime/RigVM/Source/RigVMEditor/Public/Editor/RigVMEditor.h
--- a/Engine/Plugins/Runtime/RigVM/Source/RigVMEditor/Public/Editor/RigVMEditor.h
+++ b/Engine/Plugins/Runtime/RigVM/Source/RigVMEditor/Public/Editor/RigVMEditor.h
@@ -551,6 +551,11 @@ protected:
 
 	UE_API void OnPIEStopped(bool bSimulation);
 
+public:
+	UE_API void OnGraphEditorTick(const FGeometry& AllottedGeometry, const double InCurrentTime, const float InDeltaTime,
+		TSharedRef<SGraphEditor> InGraphEditor, URigVMEdGraph* InGraph);
+
+protected:
 	/** Our currently running rig vm instance */
 	//TObjectPtr<URigVMHost> Host;
 
diff --git a/Engine/Plugins/Runtime/RigVM/Source/RigVMEditor/Public/Widgets/SRigVMGraphNode.h b/Engine/Plugins/Runtime/RigVM/Source/RigVMEditor/Public/Widgets/SRigVMGraphNode.h
--- a/Engine/Plugins/Runtime/RigVM/Source/RigVMEditor/Public/Widgets/SRigVMGraphNode.h
+++ b/Engine/Plugins/Runtime/RigVM/Source/RigVMEditor/Public/Widgets/SRigVMGraphNode.h
@@ -6,6 +6,7 @@
 #include "Widgets/DeclarativeSyntaxSupport.h"
 #include "SGraphNode.h"
 #include "Widgets/Images/SImage.h"
+#include "Widgets/Input/SButton.h"
 #include "Widgets/Text/STextBlock.h"
 #include "Widgets/Views/STableRow.h"
 #include "RigVMModel/RigVMPin.h"
@@ -107,6 +108,7 @@ public:
 	UE_API void UpdateStandardNode();
 	UE_API void UpdateCompactNode();
 	UE_API virtual TOptional<FSlateColor> GetPinTextColor(const SGraphPin* InGraphPin) const override;
+	UE_API virtual int32 OnPaint(const FPaintArgs& Args, const FGeometry& AllottedGeometry, const FSlateRect& MyCullingRect, FSlateWindowElementList& OutDrawElements, int32 LayerId, const FWidgetStyle& InWidgetStyle, bool bParentEnabled) const override;
 
 	UE_API void CreateAddPinButton();
 
@@ -152,6 +154,7 @@ protected:
 	UE_API ERigVMNodeDefaultValueOverrideState::Type GetPinDefaultValueOverrideState() const;
 	UE_API bool IsCategoryPin(const SGraphPin* InGraphPin) const;
 	UE_API EOverrideWidgetStatus::Type GetOverrideStatus() const;
+	UE_API FSlateColor GetPreviewHereColor() const;
 
 	int32 CanReorderExposedPin(const URigVMPin* InPinToReorder, const URigVMPin* InTargetPin, EItemDropZone InDropZone);
 	void OnReorderExposedPin(const URigVMPin* InPinToReorder, const URigVMPin* InTargetPin, EItemDropZone InDropZone);
@@ -161,6 +164,8 @@ protected:
 
 	int32 NodeErrorType;
 
+	TSharedPtr<SButton> PreviewHereIndicatorWidget;
+	bool bPreviewHereIndicatorHovered;
 	TSharedPtr<SImage> VisualDebugIndicatorWidget;
 	TSharedPtr<STextBlock> InstructionCountTextBlockWidget;
 	TSharedPtr<STextBlock> InstructionDurationTextBlockWidget;
diff --git a/Engine/Source/Editor/GraphEditor/Private/GraphEditorModule.cpp b/Engine/Source/Editor/GraphEditor/Private/GraphEditorModule.cpp
--- a/Engine/Source/Editor/GraphEditor/Private/GraphEditorModule.cpp
+++ b/Engine/Source/Editor/GraphEditor/Private/GraphEditorModule.cpp
@@ -72,7 +72,8 @@ TSharedRef<SGraphEditor> FGraphEditorModule::PRIVATE_MakeGraphEditor(
 	TAttribute<int32> FocusedDiffResult,
 	FSimpleDelegate InOnNavigateHistoryBack,
 	FSimpleDelegate InOnNavigateHistoryForward,
-	TAttribute<bool> ShowGraphStateOverlay)
+	TAttribute<bool> ShowGraphStateOverlay,
+	FOnGraphEditorTickDelegate InOnTickDelegate)
 {
 	return
 		SNew(SGraphEditorImpl)
@@ -88,6 +89,7 @@ TSharedRef<SGraphEditor> FGraphEditorModule::PRIVATE_MakeGraphEditor(
 		.FocusedDiffResult(FocusedDiffResult)
 		.OnNavigateHistoryBack(InOnNavigateHistoryBack)
 		.OnNavigateHistoryForward(InOnNavigateHistoryForward)
+		.OnTick(InOnTickDelegate)
 		.ShowGraphStateOverlay(ShowGraphStateOverlay);
 }
 
diff --git a/Engine/Source/Editor/GraphEditor/Private/SGraphEditorImpl.cpp b/Engine/Source/Editor/GraphEditor/Private/SGraphEditorImpl.cpp
--- a/Engine/Source/Editor/GraphEditor/Private/SGraphEditorImpl.cpp
+++ b/Engine/Source/Editor/GraphEditor/Private/SGraphEditorImpl.cpp
@@ -635,6 +635,7 @@ void SGraphEditorImpl::Construct( const FArguments& InArgs )
 	OnNavigateHistoryBack = InArgs._OnNavigateHistoryBack;
 	OnNavigateHistoryForward = InArgs._OnNavigateHistoryForward;
 	OnNodeSpawnedByKeymap = InArgs._GraphEvents.OnNodeSpawnedByKeymap;
+	OnTickDelegate = InArgs._OnTick;
 
 	NumNodesAddedSinceLastPointerPosition = 0;
 
@@ -899,6 +900,11 @@ void SGraphEditorImpl::Tick( const FGeometry& AllottedGeometry, const double InC
 			FocusLockedEditorHere();
 		}
 	}
+
+	if (OnTickDelegate.IsBound())
+	{
+		OnTickDelegate.Execute(AllottedGeometry, InCurrentTime, InDeltaTime);
+	}
 }
 
 void SGraphEditorImpl::OnClosedActionMenu()
@@ -1939,6 +1945,11 @@ SGraphPanel* SGraphEditorImpl::GetGraphPanel() const
 	return GraphPanel.Get();
 }
 
+void SGraphEditorImpl::SetOnTick(FOnGraphEditorTickDelegate InOnTickDelegate)
+{
+	OnTickDelegate = InOnTickDelegate;
+}
+
 /////////////////////////////////////////////////////
 
 
diff --git a/Engine/Source/Editor/GraphEditor/Private/SGraphEditorImpl.h b/Engine/Source/Editor/GraphEditor/Private/SGraphEditorImpl.h
--- a/Engine/Source/Editor/GraphEditor/Private/SGraphEditorImpl.h
+++ b/Engine/Source/Editor/GraphEditor/Private/SGraphEditorImpl.h
@@ -223,6 +223,7 @@ public:
 		SLATE_ARGUMENT(bool, AutoExpandActionMenu)
 		SLATE_EVENT(FSimpleDelegate, OnNavigateHistoryBack)
 		SLATE_EVENT(FSimpleDelegate, OnNavigateHistoryForward)
+		SLATE_EVENT(FOnGraphEditorTickDelegate, OnTick);
 
 		/** Show overlay elements for the graph state such as the PIE and read-only borders and text */
 		SLATE_ATTRIBUTE(bool, ShowGraphStateOverlay)
@@ -290,6 +291,12 @@ private:
 	/** Invoked when a node is created by a keymap */
 	FOnNodeSpawnedByKeymap OnNodeSpawnedByKeymap;
 
+	/**
+	 * Invoked during the tick of the graph editor. 
+	 * Can be used for custom layout code / node UI relevant code
+	 */
+	FOnGraphEditorTickDelegate OnTickDelegate;
+
 public:
 	virtual ~SGraphEditorImpl();
 
@@ -362,6 +369,8 @@ public:
 	virtual UEdGraphNode* GetSingleSelectedNode() const override;
 
 	virtual SGraphPanel* GetGraphPanel() const override;
+
+	virtual void SetOnTick(FOnGraphEditorTickDelegate InOnTickDelegate) override;
 	
 	// End of SGraphEditor interface
 protected:
diff --git a/Engine/Source/Editor/GraphEditor/Public/GraphEditorModule.h b/Engine/Source/Editor/GraphEditor/Public/GraphEditorModule.h
--- a/Engine/Source/Editor/GraphEditor/Public/GraphEditorModule.h
+++ b/Engine/Source/Editor/GraphEditor/Public/GraphEditorModule.h
@@ -57,7 +57,8 @@ private:
 		TAttribute<int32> FocusedDiffResult,
 		FSimpleDelegate InOnNavigateHistoryBack,
 		FSimpleDelegate InOnNavigateHistoryForward,
-		TAttribute<bool> ShowGraphStateOverlay);
+		TAttribute<bool> ShowGraphStateOverlay,
+		FOnGraphEditorTickDelegate InOnTickDelegate);
 
 private:
 	/** All extender delegates for the graph menus */
diff --git a/Engine/Source/Editor/UnrealEd/Private/GraphEditor.cpp b/Engine/Source/Editor/UnrealEd/Private/GraphEditor.cpp
--- a/Engine/Source/Editor/UnrealEd/Private/GraphEditor.cpp
+++ b/Engine/Source/Editor/UnrealEd/Private/GraphEditor.cpp
@@ -44,7 +44,8 @@ void SGraphEditor::ConstructImplementation( const FArguments& InArgs )
 		InArgs._FocusedDiffResult,
 		InArgs._OnNavigateHistoryBack,
 		InArgs._OnNavigateHistoryForward,
-		InArgs._ShowGraphStateOverlay
+		InArgs._ShowGraphStateOverlay,
+		InArgs._OnTick
 		);
 
 	Implementation->AssetEditorToolkit = InArgs._AssetEditorToolkit;
diff --git a/Engine/Source/Editor/UnrealEd/Public/GraphEditor.h b/Engine/Source/Editor/UnrealEd/Public/GraphEditor.h
--- a/Engine/Source/Editor/UnrealEd/Public/GraphEditor.h
+++ b/Engine/Source/Editor/UnrealEd/Public/GraphEditor.h
@@ -44,9 +44,12 @@ struct FPropertyChangedEvent;
 struct FSlateBrush;
 struct Rect;
 
+class SGraphEditor;
+
 DECLARE_DELEGATE_ThreeParams( FOnNodeTextCommitted, const FText&, ETextCommit::Type, UEdGraphNode* );
 DECLARE_DELEGATE_RetVal_ThreeParams( bool, FOnNodeVerifyTextCommit, const FText&, UEdGraphNode*, FText& );
 DECLARE_MULTICAST_DELEGATE(FOnGraphContentMenuDismissed);
+DECLARE_DELEGATE_ThreeParams(FOnGraphEditorTickDelegate, const FGeometry& AllottedGeometry, const double InCurrentTime, const float InDeltaTime);
 
 typedef TSet<class UObject*> FGraphPanelSelectionSet;
 
@@ -215,6 +218,7 @@ PRAGMA_ENABLE_DEPRECATION_WARNINGS
 		SLATE_ARGUMENT( TWeakPtr<FAssetEditorToolkit>, AssetEditorToolkit)
 		SLATE_EVENT(FSimpleDelegate, OnNavigateHistoryBack)
 		SLATE_EVENT(FSimpleDelegate, OnNavigateHistoryForward)
+		SLATE_EVENT(FOnGraphEditorTickDelegate, OnTick);
 
 		/** Show overlay elements for the graph state such as the PIE and read-only borders and text */
 		SLATE_ATTRIBUTE(bool, ShowGraphStateOverlay)				
@@ -723,6 +727,15 @@ PRAGMA_ENABLE_DEPRECATION_WARNINGS
 		return nullptr;
 	}
 
+	/** Sets the on tick delegate for the graph editor */
+	virtual void SetOnTick(FOnGraphEditorTickDelegate InOnTickDelegate)
+	{
+		if (Implementation.IsValid())
+		{
+			Implementation->SetOnTick(InOnTickDelegate);
+		}
+	}
+
 protected:
 	/** Invoked when the underlying Graph is being changed. */
 	virtual void OnGraphChanged(const struct FEdGraphEditAction& InAction)
