RigVM: remove misleading comment in RigVMEdGraphNode

#rnx

diff --git a/Engine/Plugins/Runtime/RigVM/Source/RigVMDeveloper/Private/EdGraph/RigVMEdGraphNode.cpp b/Engine/Plugins/Runtime/RigVM/Source/RigVMDeveloper/Private/EdGraph/RigVMEdGraphNode.cpp
--- a/Engine/Plugins/Runtime/RigVM/Source/RigVMDeveloper/Private/EdGraph/RigVMEdGraphNode.cpp
+++ b/Engine/Plugins/Runtime/RigVM/Source/RigVMDeveloper/Private/EdGraph/RigVMEdGraphNode.cpp
@@ -1668,7 +1668,6 @@ UClass* URigVMEdGraphNode::GetRigVMGeneratedClass() const
 
 FLinearColor URigVMEdGraphNode::GetNodeTitleColor() const
 {
-	// return a darkened version of the default node's color
 	return CachedTitleColor;
 }
 
