RigVM: Follow up to 44856852

#rnx

diff --git a/Engine/Plugins/Runtime/RigVM/Source/RigVMDeveloper/Public/EdGraph/RigVMEdGraphNode.h b/Engine/Plugins/Runtime/RigVM/Source/RigVMDeveloper/Public/EdGraph/RigVMEdGraphNode.h
--- a/Engine/Plugins/Runtime/RigVM/Source/RigVMDeveloper/Public/EdGraph/RigVMEdGraphNode.h
+++ b/Engine/Plugins/Runtime/RigVM/Source/RigVMDeveloper/Public/EdGraph/RigVMEdGraphNode.h
@@ -251,7 +251,7 @@ public:
 	UE_API void OverrideFadeOutState(TOptional<float> InFadedOutState);
 
 	// Removes any user based fade out state
-	UE_API void ResetFadedOutState() { OverrideFadeOutState({}); }
+	void ResetFadedOutState() { OverrideFadeOutState({}); }
 
 protected:
 
