Control Rig: Fix swap modules not considering cooked rig modules
#rnx

diff --git a/Engine/Plugins/Animation/ControlRig/Source/ControlRig/Private/ModularRigModel.cpp b/Engine/Plugins/Animation/ControlRig/Source/ControlRig/Private/ModularRigModel.cpp
--- a/Engine/Plugins/Animation/ControlRig/Source/ControlRig/Private/ModularRigModel.cpp
+++ b/Engine/Plugins/Animation/ControlRig/Source/ControlRig/Private/ModularRigModel.cpp
@@ -447,11 +447,13 @@ bool FModularRigModel::IsModuleParentedTo(const FRigModuleReference* InChildModu
 TArray<const FRigModuleReference*> FModularRigModel::FindModuleInstancesOfClass(const FString& InModuleClassPath) const
 {
 	TArray<const FRigModuleReference*> Result;
-	ForEachModule([&Result, InModuleClassPath](const FRigModuleReference* Module) -> bool
+	FString ModuleClassPath = InModuleClassPath;
+	ModuleClassPath.RemoveFromEnd(TEXT("_C"));
+	ForEachModule([&Result, ModuleClassPath](const FRigModuleReference* Module) -> bool
 	{
 		FString PackageName = Module->Class.ToSoftObjectPath().GetAssetPathString();
 		PackageName.RemoveFromEnd(TEXT("_C"));
-		if (PackageName == InModuleClassPath)
+		if (PackageName == ModuleClassPath)
 		{
 			Result.Add(Module);
 		}
diff --git a/Engine/Plugins/Runtime/RigVM/Source/RigVMEditor/Private/RigVMEditorBlueprintLibrary.cpp b/Engine/Plugins/Runtime/RigVM/Source/RigVMEditor/Private/RigVMEditorBlueprintLibrary.cpp
--- a/Engine/Plugins/Runtime/RigVM/Source/RigVMEditor/Private/RigVMEditorBlueprintLibrary.cpp
+++ b/Engine/Plugins/Runtime/RigVM/Source/RigVMEditor/Private/RigVMEditorBlueprintLibrary.cpp
@@ -296,7 +296,7 @@ TArray<FAssetData> URigVMEditorBlueprintLibrary::GetAssetsWithFilter_ForBlueprin
 }
 
 TArray<FAssetData> URigVMEditorBlueprintLibrary::GetAssetsWithFilter(
-		TSubclassOf<URigVMBlueprint> InClass,
+		UClass* InClass,
 		FRigVMAssetDataFilter InAssetDataFilter)
 {
 	FARFilter Filter;
@@ -320,9 +320,16 @@ TArray<FAssetData> URigVMEditorBlueprintLibrary::GetAssetsWithFilter(
 	}
 	else
 	{
+		if (!InClass->IsChildOf(URigVMBlueprint::StaticClass()) &&
+			!InClass->IsChildOf(URigVMBlueprintGeneratedClass::StaticClass()))
+		{
+			return TArray<FAssetData>();
+		}
+		
 		Filter.ClassPaths.Add(InClass->GetClassPathName());
 	}
 
+
 	const FAssetRegistryModule& AssetRegistryModule = FModuleManager::LoadModuleChecked<FAssetRegistryModule>(TEXT("AssetRegistry"));
 
 	// Collect a full list of assets with the specified class
diff --git a/Engine/Plugins/Runtime/RigVM/Source/RigVMEditor/Private/Widgets/SRigVMSwapAssetReferencesWidget.cpp b/Engine/Plugins/Runtime/RigVM/Source/RigVMEditor/Private/Widgets/SRigVMSwapAssetReferencesWidget.cpp
--- a/Engine/Plugins/Runtime/RigVM/Source/RigVMEditor/Private/Widgets/SRigVMSwapAssetReferencesWidget.cpp
+++ b/Engine/Plugins/Runtime/RigVM/Source/RigVMEditor/Private/Widgets/SRigVMSwapAssetReferencesWidget.cpp
@@ -308,8 +308,7 @@ TArray<TSharedRef<FRigVMTreeNode>> SRigVMSwapAssetReferencesWidget::GetAssetNode
 {
 	TArray<TSharedRef<FRigVMTreeNode>> Nodes;
 
-	const TArray<FAssetData> Assets = URigVMEditorBlueprintLibrary::GetAssetsWithFilter(nullptr,
-		FRigVMAssetDataFilter::CreateLambda([this, InPhase] (const FAssetData& AssetData)
+	auto AssetFilter = [this, InPhase] (const FAssetData& AssetData)
 		{
 			TArray<FRigVMAssetDataFilter>& Filters = (InPhase == PHASE_PICKSOURCE) ? SourceAssetFilters : TargetAssetFilters;
 			for (const FRigVMAssetDataFilter& Filter : Filters)
@@ -320,18 +319,15 @@ TArray<TSharedRef<FRigVMTreeNode>> SRigVMSwapAssetReferencesWidget::GetAssetNode
 				}
 			}
 			return true;
-		}));
+		};
+	
+	TArray<FAssetData> Assets = URigVMEditorBlueprintLibrary::GetAssetsWithFilter(nullptr, FRigVMAssetDataFilter::CreateLambda(AssetFilter));
+	Assets.Append(URigVMEditorBlueprintLibrary::GetAssetsWithFilter(URigVMBlueprintGeneratedClass::StaticClass(), FRigVMAssetDataFilter::CreateLambda(AssetFilter)));
 	for(const FAssetData& Asset : Assets)
 	{
-		const FString PackagePath = Asset.PackageName.ToString();
-		const FAssetData AssetData = FRigVMTreeContext::FindAssetFromAnyPath(PackagePath, false);
-		if(AssetData.IsValid())
-		{
-			Nodes.Add(FRigVMTreePackageNode::Create(AssetData)->ToSharedRef());
-		}
-		else
+		if(Asset.IsValid())
 		{
-			continue;
+			Nodes.Add(FRigVMTreePackageNode::Create(Asset)->ToSharedRef());
 		}
 	}
 
diff --git a/Engine/Plugins/Runtime/RigVM/Source/RigVMEditor/Public/RigVMEditorBlueprintLibrary.h b/Engine/Plugins/Runtime/RigVM/Source/RigVMEditor/Public/RigVMEditorBlueprintLibrary.h
--- a/Engine/Plugins/Runtime/RigVM/Source/RigVMEditor/Public/RigVMEditorBlueprintLibrary.h
+++ b/Engine/Plugins/Runtime/RigVM/Source/RigVMEditor/Public/RigVMEditorBlueprintLibrary.h
@@ -131,7 +131,7 @@ public:
 		FRigVMAssetDataFilterDynamic InAssetDataFilter);
 
 	static UE_API TArray<FAssetData> GetAssetsWithFilter(
-		TSubclassOf<URigVMBlueprint> InClass,
+		UClass* InClass,
 		FRigVMAssetDataFilter InAssetDataFilter);
 };
 
