RigVM module unloading optimizations

- Optimize FRigVMRegistry_NoLock::OnModulesUnloadedInternal by parallelizing function and type lookups
- Optimize FRigVMTemplate::HandleTypeRemoval by using a type index to permutation index lookup cache

#tests RigVM test suite, game client test on console

#was-backout 44822964

diff --git a/Engine/Plugins/Runtime/RigVM/Source/RigVM/Private/RigVMCore/RigVMRegistry.cpp b/Engine/Plugins/Runtime/RigVM/Source/RigVM/Private/RigVMCore/RigVMRegistry.cpp
--- a/Engine/Plugins/Runtime/RigVM/Source/RigVM/Private/RigVMCore/RigVMRegistry.cpp
+++ b/Engine/Plugins/Runtime/RigVM/Source/RigVM/Private/RigVMCore/RigVMRegistry.cpp
@@ -7,6 +7,7 @@
 #include "RigVMTypeUtils.h"
 #include "RigVMStringUtils.h"
 #include "RigVMModule.h"
+#include "Async/ParallelFor.h"
 #include "Animation/AttributeTypes.h"
 #include "AssetRegistry/AssetRegistryModule.h"
 #include "Engine/UserDefinedEnum.h"
@@ -926,7 +927,10 @@ bool FRigVMRegistry_NoLock::OnModulesUnloaded_NoLock(TConstArrayView<FName> Modu
 
 bool FRigVMRegistry_NoLock::OnModulesUnloadedInternal(const TArray<FString>& ModulePrefixes)
 {
+	TRACE_CPUPROFILER_EVENT_SCOPE(FRigVMRegistry_NoLock::OnModulesUnloadedInternal);
+
 	bool bRegistryChanged = false;
+	FCriticalSection ParallelForDataLock;
 
 	// remove the dispatch factories defined in the plugin
 	TArray<UScriptStruct*> FactoriesToRemove;
@@ -983,46 +987,68 @@ bool FRigVMRegistry_NoLock::OnModulesUnloadedInternal(const TArray<FString>& Mod
 		}
 	}
 
-	// remove the functions defined in the plugin
-	for (int32 Index = 0; Index < Functions.Num(); Index++)
+	// remove the functions defined in the plugin - parellize the search as it can last for a significant time
 	{
-		if (!Functions[Index].IsValid())
-		{
-			continue;
-		}
+		TRACE_CPUPROFILER_EVENT_SCOPE(FRigVMRegistry_NoLock::OnModulesUnloadedInternal::Functions);
 
-		if (Functions[Index].Struct == nullptr)
+		TArray<int32> FunctionIndicesToRemove;
+		ParallelFor(Functions.Num(), [this, &ModulePrefixes, &ParallelForDataLock, &FunctionIndicesToRemove](int32 Index)
 		{
-			continue;
-		}
+			if (!Functions[Index].IsValid())
+			{
+				return;
+			}
 
-		if (!IsWithinPlugin(ModulePrefixes, Functions[Index].Struct))
-		{
-			continue;
-		}
+			if (Functions[Index].Struct == nullptr)
+			{
+				return;
+			}
+
+			if (!IsWithinPlugin(ModulePrefixes, Functions[Index].Struct))
+			{
+				return;
+			}
+
+			FScopeLock Lock(&ParallelForDataLock);
+			FunctionIndicesToRemove.Add(Index);
+		});
 
-		if (RemoveFunction_NoLock(Index))
+		for (int32 Index : FunctionIndicesToRemove)
 		{
-			bRegistryChanged = true;
+			if (RemoveFunction_NoLock(Index))
+			{
+				bRegistryChanged = true;
+			}
 		}
-	}	
+	}
 
-	// remove the types defined in the plugin
-	for (int32 Index = 0; Index < Types.Num(); Index++)
+	// remove the types defined in the plugin - parellize the search as it can last for a very significant time
 	{
-		if (Types[Index].IsUnknownType())
-		{
-			continue;
-		}
+		TRACE_CPUPROFILER_EVENT_SCOPE(FRigVMRegistry_NoLock::OnModulesUnloadedInternal::Types);
 
-		if (!IsWithinPlugin(ModulePrefixes, Types[Index].Type.CPPTypeObject))
+		TArray<int32> TypeIndicesToRemove;
+		ParallelFor(Types.Num(), [this, &ModulePrefixes, &ParallelForDataLock, &TypeIndicesToRemove](int32 Index)
 		{
-			continue;
-		}
-		
-		if (RemoveType_NoLock(Index))
+			if (Types[Index].IsUnknownType())
+			{
+				return;
+			}
+			
+			if (!IsWithinPlugin(ModulePrefixes, Types[Index].Type.CPPTypeObject))
+			{
+				return;
+			}
+
+			FScopeLock Lock(&ParallelForDataLock);
+			TypeIndicesToRemove.Add(Index);
+		});
+
+		for (int32 Index : TypeIndicesToRemove)
 		{
-			bRegistryChanged = true;
+			if (RemoveType_NoLock(Index))
+			{
+				bRegistryChanged = true;
+			}
 		}
 	}
 
diff --git a/Engine/Plugins/Runtime/RigVM/Source/RigVM/Private/RigVMCore/RigVMTemplate.cpp b/Engine/Plugins/Runtime/RigVM/Source/RigVM/Private/RigVMCore/RigVMTemplate.cpp
--- a/Engine/Plugins/Runtime/RigVM/Source/RigVM/Private/RigVMCore/RigVMTemplate.cpp
+++ b/Engine/Plugins/Runtime/RigVM/Source/RigVM/Private/RigVMCore/RigVMTemplate.cpp
@@ -9,6 +9,7 @@
 #include "Algo/Accumulate.h"
 #include "Algo/ForEach.h"
 #include "Algo/Sort.h"
+#include "Async/ParallelFor.h"
 #include "UObject/UObjectIterator.h"
 
 #include UE_INLINE_GENERATED_CPP_BY_NAME(RigVMTemplate)
@@ -1586,7 +1587,8 @@ bool FRigVMTemplate::Merge(const FRigVMTemplate& InOther, const FRigVMRegistryHa
 	Permutations.Add(InOther.Permutations[0]);
 
 	UpdateTypesHashToPermutation(Permutations.Num()-1, InRegistry);
-	
+	bTypesIndicesToPermutationIndicesDirty = true;
+
 	return true;
 }
 
@@ -1835,6 +1837,8 @@ const FRigVMFunction* FRigVMTemplate::GetOrCreatePermutation_NoLock(int32 InInde
 		}
 	}
 	
+	bTypesIndicesToPermutationIndicesDirty = true;
+
 	return nullptr;
 }
 
@@ -2478,6 +2482,9 @@ bool FRigVMTemplate::UpdateArgumentTypes_Impl(
 	{
 		return false;
 	}
+
+	bTypesIndicesToPermutationIndicesDirty = true;
+	
 	return true;
 }
 
@@ -2595,35 +2602,73 @@ bool FRigVMTemplate::HandlePropagatedArgumentType(const TRigVMTypeIndex InTypeIn
 		}
 	}
 
+	bTypesIndicesToPermutationIndicesDirty = true;
+
 	return bResult;
 }
 
+void FRigVMTemplate::UpdateTypesIndicesToPermutationIndices(FRigVMRegistryHandle& InRegistry)
+{
+	TRACE_CPUPROFILER_EVENT_SCOPE(FRigVMTemplate::UpdateTypesIndicesToPermutationIndices);
+
+	if (bTypesIndicesToPermutationIndicesDirty)
+	{
+		TypesIndicesToPermutationIndices.Empty();
+		RemovedPermutationIndices.Empty();
+
+		FCriticalSection TypesIndicesToPermutationIndicesLock;
+
+		ParallelFor(NumPermutations_NoLock(InRegistry), [this, &InRegistry, &TypesIndicesToPermutationIndicesLock](int32 PermutationIndex)
+		{
+			const FRigVMTemplateTypeMap TypeMap = GetTypesForPermutation_NoLock(PermutationIndex, InRegistry);
+
+			for (const TPair<FName, TRigVMTypeIndex>& Pair : TypeMap)
+			{
+				FScopeLock Lock(&TypesIndicesToPermutationIndicesLock);
+				TArray<int32>& PermutationIndices = TypesIndicesToPermutationIndices.FindOrAdd(Pair.Value);
+				PermutationIndices.AddUnique(PermutationIndex);
+			}
+		});
+
+		bTypesIndicesToPermutationIndicesDirty = false;
+	}
+}
+
 void FRigVMTemplate::HandleTypeRemoval(TRigVMTypeIndex InTypeIndex, FRigVMRegistryHandle& InRegistry)
 {
+	TRACE_CPUPROFILER_EVENT_SCOPE(FRigVMTemplate::HandleTypeRemoval);
+
 	check(InRegistry.GetRegistry() == OwnerRegistry);
 
 	InvalidateHash();
 
-	const int32 CurrentNumPermutations = NumPermutations_NoLock(InRegistry);
-	for (int32 PermutationIndex = 0; PermutationIndex < CurrentNumPermutations; PermutationIndex++)
-	{
-		const FRigVMTemplateTypeMap TypeMap = GetTypesForPermutation_NoLock(PermutationIndex, InRegistry);
+	UpdateTypesIndicesToPermutationIndices(InRegistry);
 
-		for (const TPair<FName, TRigVMTypeIndex>& Pair : TypeMap)
+	TArray<int32> PermutationIndices;
+	TArray<int32>* CachedPermutationIndices = TypesIndicesToPermutationIndices.Find(InTypeIndex);
+	if (CachedPermutationIndices)
+	{
+		for (int32 PermutationIndex : *CachedPermutationIndices)
 		{
-			if (Pair.Value == InTypeIndex)
+			if (!RemovedPermutationIndices.Contains(PermutationIndex))
 			{
-				if (const FRigVMFunction* Permutation = GetPermutation_NoLock(PermutationIndex, InRegistry))
-				{
-					const int32 FunctionIndex = Permutation->Index;
-					InRegistry->RemoveFunction_NoLock(FunctionIndex);
-					Permutations[PermutationIndex] = INDEX_NONE;
-				}
-				break;
+				PermutationIndices.Add(PermutationIndex);
 			}
 		}
 	}
 
+	for (int32 PermutationIndex : PermutationIndices)
+	{
+		if (const FRigVMFunction* Permutation = GetPermutation_NoLock(PermutationIndex, InRegistry))
+		{
+			const int32 FunctionIndex = Permutation->Index;
+			InRegistry->RemoveFunction_NoLock(FunctionIndex);
+			Permutations[PermutationIndex] = INDEX_NONE;
+		}
+
+		RemovedPermutationIndices.Add(PermutationIndex);
+	}
+
 	check(!TypeIndicesToIgnore.Contains(InTypeIndex));
 	TypeIndicesToIgnore.Add(InTypeIndex);
 	TypeNamesToIgnore.Add(InRegistry->GetType_NoLock(InTypeIndex).CPPType);
diff --git a/Engine/Plugins/Runtime/RigVM/Source/RigVM/Public/RigVMCore/RigVMTemplate.h b/Engine/Plugins/Runtime/RigVM/Source/RigVM/Public/RigVMCore/RigVMTemplate.h
--- a/Engine/Plugins/Runtime/RigVM/Source/RigVM/Public/RigVMCore/RigVMTemplate.h
+++ b/Engine/Plugins/Runtime/RigVM/Source/RigVM/Public/RigVMCore/RigVMTemplate.h
@@ -615,6 +615,7 @@ private:
 
 	static FLinearColor GetColorFromMetadata(FString InMetadata);
 
+	void UpdateTypesIndicesToPermutationIndices(FRigVMRegistryHandle& InRegistry);
 	void InvalidateHash() { Hash = UINT32_MAX; }
 	RIGVM_API const TArray<FRigVMExecuteArgument>& GetExecuteArguments(const FRigVMDispatchContext& InContext) const;
 	RIGVM_API const TArray<FRigVMExecuteArgument>& GetExecuteArguments_NoLock(const FRigVMDispatchContext& InContext, const FRigVMRegistryHandle& InRegistry) const;
@@ -641,6 +642,9 @@ private:
 	mutable TArray<FRigVMExecuteArgument> ExecuteArguments;
 	TArray<int32> Permutations;
 	TMap<uint32, int32> TypesHashToPermutation;
+	TMap<TRigVMTypeIndex, TArray<int32>> TypesIndicesToPermutationIndices;
+	TArray<int32> RemovedPermutationIndices;
+	bool bTypesIndicesToPermutationIndicesDirty;
 	mutable uint32 Hash;
 	FRigVMRegistry_NoLock* OwnerRegistry;
 
