Fix failed assertion on double-deletion of RigVM template type

#rnx

diff --git a/Engine/Plugins/Runtime/RigVM/Source/RigVM/Private/RigVMCore/RigVMTemplate.cpp b/Engine/Plugins/Runtime/RigVM/Source/RigVM/Private/RigVMCore/RigVMTemplate.cpp
--- a/Engine/Plugins/Runtime/RigVM/Source/RigVM/Private/RigVMCore/RigVMTemplate.cpp
+++ b/Engine/Plugins/Runtime/RigVM/Source/RigVM/Private/RigVMCore/RigVMTemplate.cpp
@@ -2639,6 +2639,10 @@ void FRigVMTemplate::HandleTypeRemoval(TRigVMTypeIndex InTypeIndex, FRigVMRegist
 	TRACE_CPUPROFILER_EVENT_SCOPE(FRigVMTemplate::HandleTypeRemoval);
 
 	check(InRegistry.GetRegistry() == OwnerRegistry);
+	if (IsIgnoredType(InTypeIndex))
+	{
+		return;
+	}
 
 	InvalidateHash();
 
@@ -2669,7 +2673,6 @@ void FRigVMTemplate::HandleTypeRemoval(TRigVMTypeIndex InTypeIndex, FRigVMRegist
 		RemovedPermutationIndices.Add(PermutationIndex);
 	}
 
-	check(!TypeIndicesToIgnore.Contains(InTypeIndex));
 	TypeIndicesToIgnore.Add(InTypeIndex);
 	TypeNamesToIgnore.Add(InRegistry->GetType_NoLock(InTypeIndex).CPPType);
 
