RigVM Editor: Disable fading out of unrelated nodes by default

#rnx

diff --git a/Engine/Plugins/Runtime/RigVM/Source/RigVM/Private/RigVMSettings.cpp b/Engine/Plugins/Runtime/RigVM/Source/RigVM/Private/RigVMSettings.cpp
--- a/Engine/Plugins/Runtime/RigVM/Source/RigVM/Private/RigVMSettings.cpp
+++ b/Engine/Plugins/Runtime/RigVM/Source/RigVM/Private/RigVMSettings.cpp
@@ -13,7 +13,8 @@ URigVMEditorSettings::URigVMEditorSettings(const FObjectInitializer& ObjectIniti
 #endif
 {
 #if WITH_EDITORONLY_DATA
-	bHighlightSimilarNodes = true;
+	bHighlightSimilarNodes = false;
+	bFadeOutUnrelatedNodes = false;
 	bAutoLinkMutableNodes = false;
 #endif
 }
diff --git a/Engine/Plugins/Runtime/RigVM/Source/RigVM/Public/RigVMSettings.h b/Engine/Plugins/Runtime/RigVM/Source/RigVM/Public/RigVMSettings.h
--- a/Engine/Plugins/Runtime/RigVM/Source/RigVM/Public/RigVMSettings.h
+++ b/Engine/Plugins/Runtime/RigVM/Source/RigVM/Public/RigVMSettings.h
@@ -33,6 +33,10 @@ public:
 	UPROPERTY(EditAnywhere, config, Category = Interaction)
 	bool bHighlightSimilarNodes;
 
+	// fade out unrelated nodes
+	UPROPERTY(EditAnywhere, config, Category = Interaction)
+	bool bFadeOutUnrelatedNodes;
+
 	// When this is checked mutable nodes (nodes with an execute pin)
 	// will be hooked up automatically.
 	UPROPERTY(EditAnywhere, config, Category = Interaction)
diff --git a/Engine/Plugins/Runtime/RigVM/Source/RigVMDeveloper/Private/EdGraph/RigVMEdGraphNode.cpp b/Engine/Plugins/Runtime/RigVM/Source/RigVMDeveloper/Private/EdGraph/RigVMEdGraphNode.cpp
--- a/Engine/Plugins/Runtime/RigVM/Source/RigVMDeveloper/Private/EdGraph/RigVMEdGraphNode.cpp
+++ b/Engine/Plugins/Runtime/RigVM/Source/RigVMDeveloper/Private/EdGraph/RigVMEdGraphNode.cpp
@@ -1514,7 +1514,10 @@ float URigVMEdGraphNode::GetFadedOutState() const
 	}
 	if(GetInstructionIndex(true) == INDEX_NONE)
 	{
-		return 0.25f;
+		if (GetDefault<URigVMEditorSettings>()->bFadeOutUnrelatedNodes)
+    	{
+			return 0.25f;
+		}
 	}
 	if (IsExcludedFromPreview())
 	{
diff --git a/Engine/Plugins/Runtime/RigVM/Source/RigVMEditor/Private/Editor/RigVMEditor.cpp b/Engine/Plugins/Runtime/RigVM/Source/RigVMEditor/Private/Editor/RigVMEditor.cpp
--- a/Engine/Plugins/Runtime/RigVM/Source/RigVMEditor/Private/Editor/RigVMEditor.cpp
+++ b/Engine/Plugins/Runtime/RigVM/Source/RigVMEditor/Private/Editor/RigVMEditor.cpp
@@ -4288,8 +4288,18 @@ void FRigVMEditorBase::OnPIEStopped(bool bSimulation)
 	}
 }
 
+void FRigVMEditorBase::ToggleFadeOutUnrelateNodes()
+{
+	GetMutableDefault<URigVMEditorSettings>()->bFadeOutUnrelatedNodes = !IsToggleFadeOutUnrelatedNodesChecked();
+}
+
+bool FRigVMEditorBase::IsToggleFadeOutUnrelatedNodesChecked() const
+{
+	return GetDefault<URigVMEditorSettings>()->bFadeOutUnrelatedNodes;
+}
+
 void FRigVMEditorBase::OnGraphEditorTick(const FGeometry& AllottedGeometry, const double InCurrentTime, const float InDeltaTime,
-	TSharedRef<SGraphEditor> InGraphEditor, URigVMEdGraph* InGraph)
+                                         TSharedRef<SGraphEditor> InGraphEditor, URigVMEdGraph* InGraph)
 {
 	// flashlight to lighten up nodes in proximity of the mouse
 	const FGeometry PaintGeometry = InGraphEditor->GetPaintSpaceGeometry();
diff --git a/Engine/Plugins/Runtime/RigVM/Source/RigVMEditor/Private/Editor/RigVMEditorCommands.cpp b/Engine/Plugins/Runtime/RigVM/Source/RigVMEditor/Private/Editor/RigVMEditorCommands.cpp
--- a/Engine/Plugins/Runtime/RigVM/Source/RigVMEditor/Private/Editor/RigVMEditorCommands.cpp
+++ b/Engine/Plugins/Runtime/RigVM/Source/RigVMEditor/Private/Editor/RigVMEditorCommands.cpp
@@ -12,6 +12,7 @@ void FRigVMEditorCommands::RegisterCommands()
 	UI_COMMAND(SaveOnCompile_Always, "Always", "Sets the save-on-compile option to 'Always', meaning that your Blueprints will be saved whenever they are compiled (even if there were errors)", EUserInterfaceActionType::RadioButton, FInputChord());
 	UI_COMMAND(JumpToErrorNode, "Jump to Error Node", "When enabled, then the Blueprint will snap focus to nodes producing an error during compilation", EUserInterfaceActionType::ToggleButton, FInputChord());
 
+	UI_COMMAND(ToggleFadeUnrelated, "FadeUnrelated", "Fade out unrelated nodes", EUserInterfaceActionType::ToggleButton, FInputChord());
 	UI_COMMAND(EditGlobalOptions, "Settings", "Edit Class Settings", EUserInterfaceActionType::ToggleButton, FInputChord());
 	UI_COMMAND(EditClassDefaults, "Defaults", "Edit the initial values of your class.", EUserInterfaceActionType::ToggleButton, FInputChord());
 
diff --git a/Engine/Plugins/Runtime/RigVM/Source/RigVMEditor/Private/Editor/RigVMEditorCommands.h b/Engine/Plugins/Runtime/RigVM/Source/RigVMEditor/Private/Editor/RigVMEditorCommands.h
--- a/Engine/Plugins/Runtime/RigVM/Source/RigVMEditor/Private/Editor/RigVMEditorCommands.h
+++ b/Engine/Plugins/Runtime/RigVM/Source/RigVMEditor/Private/Editor/RigVMEditorCommands.h
@@ -26,6 +26,7 @@ public:
 	TSharedPtr<FUICommandInfo> JumpToErrorNode;
 
 	/** Edit global options */
+	TSharedPtr<FUICommandInfo> ToggleFadeUnrelated;
 	TSharedPtr<FUICommandInfo> EditGlobalOptions;
 	TSharedPtr<FUICommandInfo> EditClassDefaults;
 	
diff --git a/Engine/Plugins/Runtime/RigVM/Source/RigVMEditor/Private/Editor/RigVMNewEditor.cpp b/Engine/Plugins/Runtime/RigVM/Source/RigVMEditor/Private/Editor/RigVMNewEditor.cpp
--- a/Engine/Plugins/Runtime/RigVM/Source/RigVMEditor/Private/Editor/RigVMNewEditor.cpp
+++ b/Engine/Plugins/Runtime/RigVM/Source/RigVMEditor/Private/Editor/RigVMNewEditor.cpp
@@ -622,6 +622,22 @@ void FRigVMNewEditor::AddCompileWidget(FToolBarBuilder& ToolbarBuilder)
 
 void FRigVMNewEditor::AddSettingsAndDefaultWidget(FToolBarBuilder& ToolbarBuilder)
 {
+	ToolbarBuilder.AddToolBarButton(
+		FUIAction(FExecuteAction::CreateLambda([this]()
+		{
+			ToggleFadeOutUnrelateNodes();
+		}),
+		FCanExecuteAction(),
+		FGetActionCheckState::CreateLambda([this]() -> ECheckBoxState
+		{
+			return IsToggleFadeOutUnrelatedNodesChecked() ? ECheckBoxState::Checked : ECheckBoxState::Unchecked;
+		})),
+		NAME_None,
+		LOCTEXT("HideUnrelated", "Hide Unrelated"),
+		LOCTEXT("HideUnrelatedTooltip", "Fades out unrelated nodes"),
+		FSlateIcon(FAppStyle::GetAppStyleSetName(), "GraphEditor.ToggleHideUnrelatedNodes")
+	);
+
 	ToolbarBuilder.AddToolBarButton(FRigVMEditorCommands::Get().EditGlobalOptions);
 	ToolbarBuilder.AddToolBarButton(FRigVMEditorCommands::Get().EditClassDefaults);
 }
diff --git a/Engine/Plugins/Runtime/RigVM/Source/RigVMEditor/Public/Editor/RigVMEditor.h b/Engine/Plugins/Runtime/RigVM/Source/RigVMEditor/Public/Editor/RigVMEditor.h
--- a/Engine/Plugins/Runtime/RigVM/Source/RigVMEditor/Public/Editor/RigVMEditor.h
+++ b/Engine/Plugins/Runtime/RigVM/Source/RigVMEditor/Public/Editor/RigVMEditor.h
@@ -551,6 +551,9 @@ protected:
 
 	UE_API void OnPIEStopped(bool bSimulation);
 
+	UE_API virtual void ToggleFadeOutUnrelateNodes();
+	UE_API virtual bool IsToggleFadeOutUnrelatedNodesChecked() const;
+
 public:
 	UE_API void OnGraphEditorTick(const FGeometry& AllottedGeometry, const double InCurrentTime, const float InDeltaTime,
 		TSharedRef<SGraphEditor> InGraphEditor, URigVMEdGraph* InGraph);
diff --git a/Engine/Plugins/Runtime/RigVM/Source/RigVMEditor/Public/Editor/RigVMLegacyEditor.h b/Engine/Plugins/Runtime/RigVM/Source/RigVMEditor/Public/Editor/RigVMLegacyEditor.h
--- a/Engine/Plugins/Runtime/RigVM/Source/RigVMEditor/Public/Editor/RigVMLegacyEditor.h
+++ b/Engine/Plugins/Runtime/RigVM/Source/RigVMEditor/Public/Editor/RigVMLegacyEditor.h
@@ -127,6 +127,15 @@ public:
 	RIGVMEDITOR_API virtual void OnClose() override;
 	virtual FPreviewScene* GetPreviewScene() override { return &PreviewScene; }
 
+	RIGVMEDITOR_API virtual void ToggleHideUnrelatedNodes() override
+	{
+		ToggleFadeOutUnrelateNodes();
+	}
+
+	RIGVMEDITOR_API virtual bool IsToggleHideUnrelatedNodesChecked() const override
+	{
+		return IsToggleFadeOutUnrelatedNodesChecked();
+	}
 	
 protected:
 	RIGVMEDITOR_API virtual void Tick(float DeltaTime) override;
diff --git a/Engine/Source/Editor/Kismet/Public/BlueprintEditor.h b/Engine/Source/Editor/Kismet/Public/BlueprintEditor.h
--- a/Engine/Source/Editor/Kismet/Public/BlueprintEditor.h
+++ b/Engine/Source/Editor/Kismet/Public/BlueprintEditor.h
@@ -1549,8 +1549,8 @@ protected:
 
 public:
 	/** Make nodes which are unrelated to the selected nodes fade out */
-	UE_API void ToggleHideUnrelatedNodes();
-	UE_API bool IsToggleHideUnrelatedNodesChecked() const;
+	UE_API virtual void ToggleHideUnrelatedNodes();
+	UE_API virtual bool IsToggleHideUnrelatedNodesChecked() const;
 	
 	UE_DEPRECATED(5.0, "The Toggle Hide Unrelated Nodes button is always shown now")
 	UE_API bool ShouldShowToggleHideUnrelatedNodes(bool bIsToolbar) const;
