Restore the RigVMTemplate optimization that previously caused issues

- No longer results in broken or skipped permutation index removals
- Nearly 10x speedup compared to the previous, simpler optimization
- Add comparison code to confirm the new behavior in debug

#rnx

diff --git a/Engine/Plugins/Runtime/RigVM/Source/RigVM/Private/RigVMCore/RigVMTemplate.cpp b/Engine/Plugins/Runtime/RigVM/Source/RigVM/Private/RigVMCore/RigVMTemplate.cpp
--- a/Engine/Plugins/Runtime/RigVM/Source/RigVM/Private/RigVMCore/RigVMTemplate.cpp
+++ b/Engine/Plugins/Runtime/RigVM/Source/RigVM/Private/RigVMCore/RigVMTemplate.cpp
@@ -1586,8 +1586,8 @@ bool FRigVMTemplate::Merge(const FRigVMTemplate& InOther, const FRigVMRegistryHa
 
 	Permutations.Add(InOther.Permutations[0]);
 
+	UpdateTypesIndicesToPermutationIndices(Permutations.Num() - 1, InRegistry);
 	UpdateTypesHashToPermutation(Permutations.Num()-1, InRegistry);
-	bTypesIndicesToPermutationIndicesDirty = true;
 
 	return true;
 }
@@ -1801,6 +1801,7 @@ const FRigVMFunction* FRigVMTemplate::GetOrCreatePermutation_NoLock(int32 InInde
 			}
 
 			Permutations[InIndex] = FunctionIndex;
+			UpdateTypesIndicesToPermutationIndices(InIndex, InRegistry);
 
 			if (FunctionIndex  == InRegistry->Functions.Num())
 			{
@@ -1836,8 +1837,6 @@ const FRigVMFunction* FRigVMTemplate::GetOrCreatePermutation_NoLock(int32 InInde
 			return &InRegistry->Functions[FunctionIndex];
 		}
 	}
-	
-	bTypesIndicesToPermutationIndicesDirty = true;
 
 	return nullptr;
 }
@@ -2482,8 +2481,6 @@ bool FRigVMTemplate::UpdateArgumentTypes_Impl(
 	{
 		return false;
 	}
-
-	bTypesIndicesToPermutationIndicesDirty = true;
 	
 	return true;
 }
@@ -2594,43 +2591,30 @@ bool FRigVMTemplate::HandlePropagatedArgumentType(const TRigVMTypeIndex InTypeIn
 			for(int32 PermutationIndex = 0; PermutationIndex < Permutations.Num(); PermutationIndex++)
 			{
 				Permutations[PermutationIndex] = INDEX_NONE;
+				UpdateTypesIndicesToPermutationIndices(PermutationIndex, InRegistry);
 			}
 			for(int32 FunctionIndex : ValidFunctionIndices)
 			{
-				Permutations[OwnerRegistry->Functions[FunctionIndex].PermutationIndex] = FunctionIndex;
+				int32 PermutationIndex = OwnerRegistry->Functions[FunctionIndex].PermutationIndex;
+				Permutations[PermutationIndex] = FunctionIndex;
+				UpdateTypesIndicesToPermutationIndices(PermutationIndex, InRegistry);
 			}
 		}
 	}
 
-	bTypesIndicesToPermutationIndicesDirty = true;
-
 	return bResult;
 }
 
-void FRigVMTemplate::UpdateTypesIndicesToPermutationIndices(FRigVMRegistryHandle& InRegistry)
+void FRigVMTemplate::UpdateTypesIndicesToPermutationIndices(int32 InPermutationIndex, const FRigVMRegistryHandle& InRegistry)
 {
 	TRACE_CPUPROFILER_EVENT_SCOPE(FRigVMTemplate::UpdateTypesIndicesToPermutationIndices);
 
-	if (bTypesIndicesToPermutationIndicesDirty)
-	{
-		TypesIndicesToPermutationIndices.Empty();
-		RemovedPermutationIndices.Empty();
-
-		FCriticalSection TypesIndicesToPermutationIndicesLock;
-
-		ParallelFor(NumPermutations_NoLock(InRegistry), [this, &InRegistry, &TypesIndicesToPermutationIndicesLock](int32 PermutationIndex)
-		{
-			const FRigVMTemplateTypeMap TypeMap = GetTypesForPermutation_NoLock(PermutationIndex, InRegistry);
-
-			for (const TPair<FName, TRigVMTypeIndex>& Pair : TypeMap)
-			{
-				FScopeLock Lock(&TypesIndicesToPermutationIndicesLock);
-				TArray<int32>& PermutationIndices = TypesIndicesToPermutationIndices.FindOrAdd(Pair.Value);
-				PermutationIndices.AddUnique(PermutationIndex);
-			}
-		});
+	const FRigVMTemplateTypeMap TypeMap = GetTypesForPermutation_NoLock(InPermutationIndex, InRegistry);
 
-		bTypesIndicesToPermutationIndicesDirty = false;
+	for (const TPair<FName, TRigVMTypeIndex>& Pair : TypeMap)
+	{
+		TArray<int32>& PermutationIndices = TypesIndicesToPermutationIndices.FindOrAdd(Pair.Value);
+		PermutationIndices.AddUnique(InPermutationIndex);
 	}
 }
 
@@ -2646,8 +2630,6 @@ void FRigVMTemplate::HandleTypeRemoval(TRigVMTypeIndex InTypeIndex, FRigVMRegist
 
 	InvalidateHash();
 
-	UpdateTypesIndicesToPermutationIndices(InRegistry);
-
 	TArray<int32> PermutationIndices;
 	TArray<int32>* CachedPermutationIndices = TypesIndicesToPermutationIndices.Find(InTypeIndex);
 	if (CachedPermutationIndices)
@@ -2661,6 +2643,50 @@ void FRigVMTemplate::HandleTypeRemoval(TRigVMTypeIndex InTypeIndex, FRigVMRegist
 		}
 	}
 
+#if UE_BUILD_DEBUG
+
+	// The fast path above is tricky, so we'll also run the slow path, and confirm we found the right indices
+	TArray<int32> SlowPermutationIndices;
+	const int32 CurrentNumPermutations = NumPermutations_NoLock(InRegistry);
+	for (int32 PermutationIndex = 0; PermutationIndex < CurrentNumPermutations; PermutationIndex++)
+	{
+		const FRigVMTemplateTypeMap TypeMap = GetTypesForPermutation_NoLock(PermutationIndex, InRegistry);
+
+		for (const TPair<FName, TRigVMTypeIndex>& Pair : TypeMap)
+		{
+			if (Pair.Value == InTypeIndex)
+			{
+				if (const FRigVMFunction* Permutation = GetPermutation_NoLock(PermutationIndex, InRegistry))
+				{
+					SlowPermutationIndices.Add(PermutationIndex);
+				}
+				break;
+			}
+		}
+	}
+
+	// Verify our indices
+	if (PermutationIndices.Num() != SlowPermutationIndices.Num())
+	{
+		UE_LOG(LogRigVM, Error, TEXT("HandleTypeRemoval: bad permutation count %d v %d (%d)"),
+		PermutationIndices.Num(), SlowPermutationIndices.Num(), PermutationIndices.Num() > 0 ? PermutationIndices[0] : SlowPermutationIndices[0]);
+	}
+	else
+	{
+		for (int32 TestIndex = 0; TestIndex < PermutationIndices.Num(); TestIndex++)
+		{
+			if (PermutationIndices[TestIndex] != SlowPermutationIndices[TestIndex])
+			{
+				UE_LOG(LogRigVM, Error, TEXT("HandleTypeRemoval: bad permutation index %d v %d"),
+					PermutationIndices[TestIndex], SlowPermutationIndices[TestIndex]);
+			}
+		}
+	}
+
+	PermutationIndices = SlowPermutationIndices;
+
+#endif // UE_BUILD_DEBUG
+
 	for (int32 PermutationIndex : PermutationIndices)
 	{
 		if (const FRigVMFunction* Permutation = GetPermutation_NoLock(PermutationIndex, InRegistry))
diff --git a/Engine/Plugins/Runtime/RigVM/Source/RigVM/Public/RigVMCore/RigVMTemplate.h b/Engine/Plugins/Runtime/RigVM/Source/RigVM/Public/RigVMCore/RigVMTemplate.h
--- a/Engine/Plugins/Runtime/RigVM/Source/RigVM/Public/RigVMCore/RigVMTemplate.h
+++ b/Engine/Plugins/Runtime/RigVM/Source/RigVM/Public/RigVMCore/RigVMTemplate.h
@@ -615,7 +615,7 @@ private:
 
 	static FLinearColor GetColorFromMetadata(FString InMetadata);
 
-	void UpdateTypesIndicesToPermutationIndices(FRigVMRegistryHandle& InRegistry);
+	void UpdateTypesIndicesToPermutationIndices(int32 InPermutationIndex, const FRigVMRegistryHandle& InRegistry);
 	void InvalidateHash() { Hash = UINT32_MAX; }
 	RIGVM_API const TArray<FRigVMExecuteArgument>& GetExecuteArguments(const FRigVMDispatchContext& InContext) const;
 	RIGVM_API const TArray<FRigVMExecuteArgument>& GetExecuteArguments_NoLock(const FRigVMDispatchContext& InContext, const FRigVMRegistryHandle& InRegistry) const;
@@ -644,7 +644,6 @@ private:
 	TMap<uint32, int32> TypesHashToPermutation;
 	TMap<TRigVMTypeIndex, TArray<int32>> TypesIndicesToPermutationIndices;
 	TArray<int32> RemovedPermutationIndices;
-	bool bTypesIndicesToPermutationIndicesDirty;
 	mutable uint32 Hash;
 	FRigVMRegistry_NoLock* OwnerRegistry;
 
