Control Rig / RigVM: Preview Here functionality first pass

#rnx

diff --git a/Engine/Plugins/Animation/ControlRig/Source/ControlRig/Private/ControlRig.cpp b/Engine/Plugins/Animation/ControlRig/Source/ControlRig/Private/ControlRig.cpp
--- a/Engine/Plugins/Animation/ControlRig/Source/ControlRig/Private/ControlRig.cpp
+++ b/Engine/Plugins/Animation/ControlRig/Source/ControlRig/Private/ControlRig.cpp
@@ -935,20 +935,13 @@ bool UControlRig::Execute(const FName& InEventName)
 #if WITH_EDITOR
 	ExtendedExecuteContext.SetInstructionVisitInfo(&InstructionVisitInfo);
 
-	if (IsInDebugMode())
-	{
-		if (UControlRig* CDO = GetClass()->GetDefaultObject<UControlRig>())
-		{
-			// Copy the breakpoints. This will not override the state of the breakpoints
-			DebugInfo.CopyFrom(CDO->DebugInfo);
-		}
-
-		ExtendedExecuteContext.SetDebugInfo(&DebugInfo);
-	}
-	else
+	if (UControlRig* CDO = GetClass()->GetDefaultObject<UControlRig>())
 	{
-		ExtendedExecuteContext.SetDebugInfo(nullptr);
+		// Copy the breakpoints. This will not override the state of the breakpoints
+		DebugInfo.CopyFrom(CDO->DebugInfo);
 	}
+
+	ExtendedExecuteContext.SetDebugInfo(&DebugInfo);
 	
 	if (IsProfilingEnabled())
     {
diff --git a/Engine/Plugins/Animation/ControlRig/Source/ControlRigEditor/Private/RigDependencyGraph/RigDependencyGraphNode.cpp b/Engine/Plugins/Animation/ControlRig/Source/ControlRigEditor/Private/RigDependencyGraph/RigDependencyGraphNode.cpp
--- a/Engine/Plugins/Animation/ControlRig/Source/ControlRigEditor/Private/RigDependencyGraph/RigDependencyGraphNode.cpp
+++ b/Engine/Plugins/Animation/ControlRig/Source/ControlRigEditor/Private/RigDependencyGraph/RigDependencyGraphNode.cpp
@@ -390,7 +390,7 @@ float URigDependencyGraphNode::GetFadedOutState() const
 	{
 		return FadedOutOverride.GetValue();
 	}
-	return bIsFadedOut.Get(false) ? 0.1f : 1.f;
+	return bIsFadedOut.Get(false) ? 0.25f : 1.f;
 }
 
 void URigDependencyGraphNode::OverrideFadeOutState(TOptional<float> InFadedOutState)
diff --git a/Engine/Plugins/Runtime/RigVM/Source/RigVM/Private/RigVMCore/RigVM.cpp b/Engine/Plugins/Runtime/RigVM/Source/RigVM/Private/RigVMCore/RigVM.cpp
--- a/Engine/Plugins/Runtime/RigVM/Source/RigVM/Private/RigVMCore/RigVM.cpp
+++ b/Engine/Plugins/Runtime/RigVM/Source/RigVM/Private/RigVMCore/RigVM.cpp
@@ -1637,19 +1637,66 @@ ERigVMExecuteResult URigVM::ExecuteVM(FRigVMExtendedExecuteContext& Context, con
 	
 #endif
 
-	Context.CurrentExecuteResult = ExecuteInstructions(Context, ContextPublicData.InstructionIndex, Instructions.Num() - 1);
+	// find the last instruction
+	int32 LastInstruction = Instructions.Num() - 1;
+	for (int32 InstructionIndex = ContextPublicData.InstructionIndex + 1; InstructionIndex < Instructions.Num(); InstructionIndex++)
+	{
+		if (Instructions[InstructionIndex].OpCode == ERigVMOpCode::Exit)
+		{
+			LastInstruction = InstructionIndex;
+			break;
+		}
+	}
+
+#if WITH_EDITOR
+	FRigVMDebugInfo* RigVMDebugInfo = Context.GetRigVMDebugInfo();
+	FRigVMDebugInfo PreviousDebugInfo;
+	if(bIsRootEntry && RigVMDebugInfo)
+	{
+		PreviousDebugInfo.CopyFrom(*RigVMDebugInfo);
+		
+		if (!RigVMDebugInfo->IsEmpty())
+		{
+			// if the instruction we are about to early exit from is outside of the set of instructions
+			// we are about to run - disable the debug info.
+			const int32 EarlyExistInstructionIndex = RigVMDebugInfo->GetInstructionForEarlyExit();
+			if (EarlyExistInstructionIndex != FMath::Clamp(EarlyExistInstructionIndex, ContextPublicData.InstructionIndex, LastInstruction))
+			{
+				RigVMDebugInfo->Reset();
+			}
+			else
+			{
+				LastInstruction = EarlyExistInstructionIndex;
+			}
+		}
+		
+		if (!RigVMDebugInfo->IsEmpty())
+		{
+			RigVMDebugInfo->ExecutionHalted().Broadcast(INDEX_NONE, nullptr, InEntryName);
+		}
+	}
+#endif
+
+	Context.CurrentExecuteResult = ExecuteInstructions(Context, ContextPublicData.InstructionIndex, LastInstruction);
 
 #if WITH_EDITOR
 	if(bIsRootEntry)
 	{
 		Context.CurrentVMMemory.Reset();
 		
-		if (FRigVMDebugInfo* RigVMDebugInfo = Context.GetRigVMDebugInfo())
+		if (RigVMDebugInfo)
 		{
 			if (!RigVMDebugInfo->IsEmpty())
 			{
-				RigVMDebugInfo->ExecutionHalted().Broadcast(INDEX_NONE, nullptr, InEntryName);
+				if (Context.CurrentExecuteResult != ERigVMExecuteResult::Failed)
+				{
+					Context.CurrentExecuteResult = ERigVMExecuteResult::Halted;
+					RigVMDebugInfo->ExecutionHalted().Broadcast(RigVMDebugInfo->GetInstructionForEarlyExit(), nullptr, InEntryName);
+				}
 			}
+
+			// restore it
+			RigVMDebugInfo->CopyFrom(PreviousDebugInfo);
 		}
 	}
 #endif
@@ -1680,7 +1727,6 @@ ERigVMExecuteResult URigVM::ExecuteInstructions(FRigVMExtendedExecuteContext& Co
 #if WITH_EDITOR
 	TArray<FName>& FunctionNames = GetFunctionNames();
 	FRigVMInstructionVisitInfo* RigVMInstructionVisitInfo = Context.GetRigVMInstructionVisitInfo();
-	FRigVMDebugInfo* RigVMDebugInfo = Context.GetRigVMDebugInfo();
 	FRigVMProfilingInfo* RigVMProfilingInfo = Context.GetRigVMProfilingInfo();
 #endif
 	
@@ -1951,13 +1997,6 @@ ERigVMExecuteResult URigVM::ExecuteInstructions(FRigVMExtendedExecuteContext& Co
 					}
 
 #if WITH_EDITOR					
-					if (RigVMDebugInfo != nullptr)
-					{
-						if (RigVMDebugInfo && !RigVMDebugInfo->IsEmpty())
-						{
-							RigVMDebugInfo->ExecutionHalted().Broadcast(INDEX_NONE, nullptr, Context.CurrentEntryName);
-						}
-					}
 #if UE_RIGVM_DEBUG_EXECUTION
 					if (CVarControlRigDebugAllVMExecutions->GetBool() || ContextPublicData.bDebugExecution)
 					{
@@ -1999,14 +2038,6 @@ ERigVMExecuteResult URigVM::ExecuteInstructions(FRigVMExtendedExecuteContext& Co
 					{
 						return Context.CurrentExecuteResult = ExecuteResult;
 					}
-
-#if WITH_EDITOR
-					// if we are halted at a break point we need to exit here
-					if (RigVMDebugInfo && RigVMDebugInfo->ShouldEarlyExitBeforeInstruction(ContextPublicData.InstructionIndex))
-					{
-						return Context.CurrentExecuteResult = ERigVMExecuteResult::Halted;
-					}
-#endif
 				}
 					
 				ContextPublicData.InstructionIndex++;
@@ -2164,21 +2195,6 @@ ERigVMExecuteResult URigVM::ExecuteInstructions(FRigVMExtendedExecuteContext& Co
 				RigVMProfilingInfo->AddOverallCycles(Cycles);
 			}
 		}
-
-		if (RigVMDebugInfo && RigVMDebugInfo->ShouldEarlyExitBeforeInstruction(ContextPublicData.InstructionIndex))
-		{
-#if UE_RIGVM_DEBUG_EXECUTION
-			if (CVarControlRigDebugAllVMExecutions->GetBool() || ContextPublicData.bDebugExecution)
-			{
-				ContextPublicData.Log(EMessageSeverity::Info, ContextPublicData.DebugMemoryString);
-			}
-#endif
-			// we'll recursively exit all invoked
-			// entries here.
-			RigVMDebugInfo->ExecutionHalted().Broadcast(ContextPublicData.InstructionIndex, nullptr, Context.CurrentEntryName);
-			return Context.CurrentExecuteResult = ERigVMExecuteResult::Halted;
-		}
-
 #if UE_RIGVM_DEBUG_EXECUTION
 		if (CVarControlRigDebugAllVMExecutions->GetBool() || ContextPublicData.bDebugExecution)
 		{
diff --git a/Engine/Plugins/Runtime/RigVM/Source/RigVM/Private/RigVMHost.cpp b/Engine/Plugins/Runtime/RigVM/Source/RigVM/Private/RigVMHost.cpp
--- a/Engine/Plugins/Runtime/RigVM/Source/RigVM/Private/RigVMHost.cpp
+++ b/Engine/Plugins/Runtime/RigVM/Source/RigVM/Private/RigVMHost.cpp
@@ -484,19 +484,12 @@ bool URigVMHost::Execute(const FName& InEventName)
 #if WITH_EDITOR
 	ExtendedExecuteContext.SetInstructionVisitInfo(&InstructionVisitInfo);
 
-	if (IsInDebugMode())
+	if (CDO)
 	{
-		if (CDO)
-		{
-			DebugInfo.CopyFrom(CDO->DebugInfo);
-		}
-
-		ExtendedExecuteContext.SetDebugInfo(&DebugInfo);
-	}
-	else
-	{
-		ExtendedExecuteContext.SetDebugInfo(nullptr);
+		DebugInfo.CopyFrom(CDO->DebugInfo);
 	}
+
+	ExtendedExecuteContext.SetDebugInfo(&DebugInfo);
 	
 	if (IsProfilingEnabled())
 	{
diff --git a/Engine/Plugins/Runtime/RigVM/Source/RigVMDeveloper/Private/RigVMAsset.cpp b/Engine/Plugins/Runtime/RigVM/Source/RigVMDeveloper/Private/RigVMAsset.cpp
--- a/Engine/Plugins/Runtime/RigVM/Source/RigVMDeveloper/Private/RigVMAsset.cpp
+++ b/Engine/Plugins/Runtime/RigVM/Source/RigVMDeveloper/Private/RigVMAsset.cpp
@@ -77,7 +77,7 @@ static TArray<UClass*> GetClassObjectsInPackage(UPackage* InPackage)
 
 FRigVMEdGraphDisplaySettings::FRigVMEdGraphDisplaySettings(): bShowNodeInstructionIndex(false)
                                                             , bShowNodeRunCounts(false)
-                                                            , NodeRunLowerBound(1)
+                                                            , NodeRunLowerBound(0)
                                                             , NodeRunLimit(256)
                                                             , MinMicroSeconds(0.0)
                                                             , MaxMicroSeconds(1.0)
@@ -927,9 +927,10 @@ const FLazyName& IRigVMAssetInterface::GetPanelPinFactoryName() const
 
 bool IRigVMAssetInterface::SetEarlyExitInstruction(int32 InInstructionIndex)
 {
+	const bool bResetResult = ResetEarlyExitInstruction();
 	if (InInstructionIndex == INDEX_NONE)
 	{
-		return ResetEarlyExitInstruction();
+		return bResetResult;
 	}
 
 	const FRigVMByteCode* ByteCode = GetController()->GetCurrentByteCode();
@@ -948,6 +949,68 @@ bool IRigVMAssetInterface::SetEarlyExitInstruction(int32 InInstructionIndex)
 	if (URigVMHost* DebuggedHost = GetDebuggedRigVMHost())
 	{
 		DebuggedHost->GetDebugInfo().CopyFrom(GetDebugInfo());
+		DebuggedHost->SetIsInDebugMode(true);
+	}
+
+	URigVMNode* Node = Cast<URigVMNode>(ByteCode->GetSubjectForInstruction(InInstructionIndex));
+	if (Node)
+	{
+		Node->SetHasEarlyExitMarker(true);
+
+		TArray<URigVMNode*> ExcludedNodes = {Node};
+		for (int32 ExcludedNodeIndex = 0; ExcludedNodeIndex < ExcludedNodes.Num(); ExcludedNodeIndex++)
+		{
+			URigVMNode* ExcludedNode = ExcludedNodes[ExcludedNodeIndex];
+			if (!ExcludedNode->IsMutable())
+			{
+				continue;
+			}
+
+			if (ExcludedNode != Node)
+			{
+				ExcludedNode->SetIsExcludedByEarlyExit(true);
+			}
+
+			const TArray<URigVMNode*> LinkedNodes = ExcludedNode->GetLinkedTargetNodes();
+			for (URigVMNode* LinkedNode : LinkedNodes)
+			{
+				if (!LinkedNode->IsMutable())
+				{
+					continue;
+				}
+				ExcludedNodes.AddUnique(LinkedNode);
+			}
+		}
+
+		for (int32 ExcludedNodeIndex = 0; ExcludedNodeIndex < ExcludedNodes.Num(); ExcludedNodeIndex++)
+		{
+			URigVMNode* ExcludedNode = ExcludedNodes[ExcludedNodeIndex];
+
+			if (!ExcludedNode->IsMutable())
+			{
+				const TArray<URigVMNode*> LinkedTargetNodes = ExcludedNode->GetLinkedTargetNodes();
+				const bool bHasIncludedTargetNode = LinkedTargetNodes.ContainsByPredicate([](const URigVMNode* LinkedTargetNode) -> bool {
+					return !LinkedTargetNode->IsExcludedByEarlyExit();
+				});
+
+				if (bHasIncludedTargetNode)
+				{
+					continue;
+				}
+
+				ExcludedNode->SetIsExcludedByEarlyExit(true);
+			}
+
+			const TArray<URigVMNode*> LinkedSourceNodes = ExcludedNode->GetLinkedSourceNodes();
+			for (URigVMNode* LinkedNode : LinkedSourceNodes)
+			{
+				if (LinkedNode->IsMutable())
+				{
+					continue;
+				}
+				ExcludedNodes.AddUnique(LinkedNode);
+			}
+		}
 	}
 	
 	return true;
@@ -972,7 +1035,14 @@ bool IRigVMAssetInterface::SetEarlyExitInstruction(URigVMNode* InNodeToExitEarly
 		return false;
 	}
 
-	return SetEarlyExitInstruction(InstructionIndices.Last());
+	// find the last consecutive instruction for a node in the first block
+	int32 Instruction = InstructionIndices[0];
+	while (Instruction < InstructionIndices.Num() - 1 && Instruction == InstructionIndices[Instruction + 1] - 1)
+	{
+		Instruction++;
+	}
+	
+	return SetEarlyExitInstruction(Instruction);
 }
 
 bool IRigVMAssetInterface::ResetEarlyExitInstruction()
@@ -984,6 +1054,16 @@ bool IRigVMAssetInterface::ResetEarlyExitInstruction()
 		DebuggedHost->GetDebugInfo().CopyFrom(GetDebugInfo());
 	}
 
+	TArray<URigVMGraph*> AllGraphs = GetAllModels();
+	for (URigVMGraph* Graph : AllGraphs)
+	{
+		for (URigVMNode* Node : Graph->GetNodes())
+		{
+			Node->SetHasEarlyExitMarker(false);
+			Node->SetIsExcludedByEarlyExit(false);
+		}
+	}
+
 	return true;
 }
 
@@ -1517,6 +1597,8 @@ void IRigVMAssetInterface::RecompileVM()
 		TGuardValue<bool> ReentrantGuardSelf(bSuspendModelNotificationsForSelf, true);
 		TGuardValue<bool> ReentrantGuardOthers(RigVMClient->bSuspendModelNotificationsForOthers, true);
 
+		ResetEarlyExitInstruction();
+		
 		PreCompile();
 
 		GetObject()->Modify(false);
@@ -2442,6 +2524,7 @@ void IRigVMAssetInterface::SetObjectBeingDebugged(UObject* NewObject)
 		PreviousRigBeingDebugged->RigVMLog = nullptr;
 #if WITH_EDITOR
 		PreviousRigBeingDebugged->bIsBeingDebugged = false;
+		PreviousRigBeingDebugged->GetDebugInfo().Reset();
 #endif
 	}
 
@@ -2452,6 +2535,8 @@ void IRigVMAssetInterface::SetObjectBeingDebugged(UObject* NewObject)
 	{
 		NewRigBeingDebugged->bIsBeingDebugged = true;
 	}
+	
+	ResetEarlyExitInstruction();
 #endif
 }
 
@@ -2986,6 +3071,7 @@ void IRigVMAssetInterface::HandleModifiedEvent(ERigVMGraphNotifType InNotifType,
 					}
 
 					RequestAutoVMRecompilation();
+					ResetEarlyExitInstruction();
 
 					(void)MarkPackageDirty();
 					MarkBlueprintAsStructurallyModified();
@@ -3018,6 +3104,7 @@ void IRigVMAssetInterface::HandleModifiedEvent(ERigVMGraphNotifType InNotifType,
 			case ERigVMGraphNotifType::PinDirectionChanged:
 			{
 				RequestAutoVMRecompilation();
+				ResetEarlyExitInstruction();
 				(void)MarkPackageDirty();
 
 				// we don't need to mark the blueprint as modified since we only
diff --git a/Engine/Plugins/Runtime/RigVM/Source/RigVMEditor/Private/RigVMEditorModule.cpp b/Engine/Plugins/Runtime/RigVM/Source/RigVMEditor/Private/RigVMEditorModule.cpp
--- a/Engine/Plugins/Runtime/RigVM/Source/RigVMEditor/Private/RigVMEditorModule.cpp
+++ b/Engine/Plugins/Runtime/RigVM/Source/RigVMEditor/Private/RigVMEditorModule.cpp
@@ -896,7 +896,6 @@ void FRigVMEditorModule::GetNodeDebugContextMenuActions(TScriptInterface<IRigVMC
 		bool bAnyNodeIsEarlyExit = false;
 
 		const URigVMGraph* Model = ModelNode->GetGraph();
-		URigVMController* Controller = RigVMBlueprint->GetController(Model);
 
 		TArray<URigVMNode*> SelectedNodes;
 		TArray<FName> SelectedNodeNames = Model->GetSelectNodes();
@@ -921,12 +920,12 @@ void FRigVMEditorModule::GetNodeDebugContextMenuActions(TScriptInterface<IRigVMC
 			LOCTEXT("PreviewHere", "Preview Here"),
 			LOCTEXT("PreviewHere_Tooltip", "Preview the result of the node here"),
 			FSlateIcon(),
-			FUIAction(FExecuteAction::CreateLambda([Controller, SelectedNodes, RigVMBlueprint]()
+			FUIAction(FExecuteAction::CreateLambda([SelectedNodes, RigVMBlueprint]()
 			{
 				for (URigVMNode* SelectedNode : SelectedNodes)
 				{
-					// todooooo
-					//if (RigVMBlueprint->AddBreakpoint(SelectedNode))
+					RigVMBlueprint->SetEarlyExitInstruction(SelectedNode);
+					break;
 				}
 			})));
 		}
@@ -937,16 +936,9 @@ void FRigVMEditorModule::GetNodeDebugContextMenuActions(TScriptInterface<IRigVMC
 			LOCTEXT("StopPreview", "Stop Preview Here"),
 			LOCTEXT("StopPreview_Tooltip", "Stops the preview from here functionality and runs the whole graph"),
 			FSlateIcon(),
-			FUIAction(FExecuteAction::CreateLambda([Controller, SelectedNodes, RigVMBlueprint]()
+			FUIAction(FExecuteAction::CreateLambda([RigVMBlueprint]()
 			{
-				for (URigVMNode* SelectedNode : SelectedNodes)
-				{
-					if (SelectedNode->HasEarlyExitMarker())
-					{
-						// todooooo
-						//if (RigVMBlueprint->RemoveBreakpoint(SelectedNode))
-					}
-				}                            
+					RigVMBlueprint->ResetEarlyExitInstruction();
 			})));
 		}
 	}
diff --git a/Engine/Plugins/Runtime/RigVM/Source/RigVMEditor/Private/Widgets/SRigVMExecutionStackView.cpp b/Engine/Plugins/Runtime/RigVM/Source/RigVMEditor/Private/Widgets/SRigVMExecutionStackView.cpp
--- a/Engine/Plugins/Runtime/RigVM/Source/RigVMEditor/Private/Widgets/SRigVMExecutionStackView.cpp
+++ b/Engine/Plugins/Runtime/RigVM/Source/RigVMEditor/Private/Widgets/SRigVMExecutionStackView.cpp
@@ -290,7 +290,6 @@ void SRigVMExecutionStackView::Construct( const FArguments& InArgs, TSharedRef<I
 	CommandList = MakeShared<FUICommandList>();
 	bSuspendModelNotifications = false;
 	bSuspendControllerSelection = false;
-	HaltedAtInstruction = INDEX_NONE;
 
 	BindCommands();
 
@@ -783,8 +782,6 @@ void SRigVMExecutionStackView::RefreshTreeView(URigVM* InVM, FRigVMExtendedExecu
 						}
 					}
 				}
-
-				Host->GetDebugInfo().ExecutionHalted().AddSP(this, &SRigVMExecutionStackView::HandleExecutionHalted);
 			}
 		}
 	}
@@ -1068,30 +1065,6 @@ void SRigVMExecutionStackView::OnVMCompiled(UObject* InCompiledObject, URigVM* I
 	}
 }
 
-void SRigVMExecutionStackView::HandleExecutionHalted(const int32 InHaltedAtInstruction, UObject* InNode, const FName& InEntryName)
-{
-	if (HaltedAtInstruction == InHaltedAtInstruction)
-	{
-		return;
-	}
-
-	if (TSharedPtr<IRigVMEditor> Editor = WeakEditor.Pin())
-	{
-		if (InHaltedAtInstruction == INDEX_NONE && InEntryName == Editor->GetRigVMHost()->GetEventQueue().Last())
-		{
-			HaltedAtInstruction = InHaltedAtInstruction;
-			return;
-		}
-	}
-	
-	if (InHaltedAtInstruction != INDEX_NONE && Operators.Num() > InHaltedAtInstruction)
-	{
-		HaltedAtInstruction = InHaltedAtInstruction;
-		TreeView->SetSelection(Operators[InHaltedAtInstruction]);
-		TreeView->SetScrollOffset(static_cast<float>(FMath::Max(InHaltedAtInstruction-5, 0)));
-	}	
-}
-
 void SRigVMExecutionStackView::OnFilterTextChanged(const FText& SearchText)
 {
 	FilterText = SearchText;
diff --git a/Engine/Plugins/Runtime/RigVM/Source/RigVMEditor/Private/Widgets/SRigVMGraphNode.cpp b/Engine/Plugins/Runtime/RigVM/Source/RigVMEditor/Private/Widgets/SRigVMGraphNode.cpp
--- a/Engine/Plugins/Runtime/RigVM/Source/RigVMEditor/Private/Widgets/SRigVMGraphNode.cpp
+++ b/Engine/Plugins/Runtime/RigVM/Source/RigVMEditor/Private/Widgets/SRigVMGraphNode.cpp
@@ -356,6 +356,7 @@ void SRigVMGraphNode::Construct( const FArguments& InArgs )
 		[
 			SNew(SImage)
 			.Image(FAppStyle::GetBrush(TEXT("Icons.Preview")))
+			.DesiredSizeOverride(FVector2D(16, 16))
 			.ColorAndOpacity(this, &SRigVMGraphNode::GetPreviewHereColor)
 			.Visibility(EVisibility::Visible)
 		];
@@ -587,11 +588,11 @@ int32 SRigVMGraphNode::OnPaint(const FPaintArgs& Args, const FGeometry& Allotted
 {
 	FWidgetStyle WidgetStyle = InWidgetStyle;
 	
-	if (const URigVMEdGraphNode* DependencyGraphNode = Cast<URigVMEdGraphNode>(GraphNode))
+	if (const URigVMEdGraphNode* RigVMEdGraphNode = Cast<URigVMEdGraphNode>(GraphNode))
 	{
-		WidgetStyle.BlendOpacity(DependencyGraphNode->GetFadedOutState());
+		WidgetStyle.BlendOpacity(RigVMEdGraphNode->GetFadedOutState());
 	}
-	
+		
 	return SGraphNode::OnPaint(Args, AllottedGeometry, MyCullingRect, OutDrawElements, LayerId, WidgetStyle, bParentEnabled);
 }
 
@@ -1283,7 +1284,7 @@ TArray<FOverlayWidgetInfo> SRigVMGraphNode::GetOverlayWidgets(bool bSelected, co
 			const FVector2f ImageSize = PreviewHereIndicatorWidget->GetDesiredSize();
 
 			FOverlayWidgetInfo Info;
-			Info.OverlayOffset = FVector2f(WidgetSize.X - ImageSize.X - 6.f - HorizontalOffset, 6.f);
+			Info.OverlayOffset = FVector2f(WidgetSize.X - ImageSize.X - 6.f - HorizontalOffset, 4.f);
 			HorizontalOffset += ImageSize.X + 6.f;
 			Info.Widget = PreviewHereIndicatorWidget;
 
diff --git a/Engine/Plugins/Runtime/RigVM/Source/RigVMEditor/Public/Widgets/SRigVMExecutionStackView.h b/Engine/Plugins/Runtime/RigVM/Source/RigVMEditor/Public/Widgets/SRigVMExecutionStackView.h
--- a/Engine/Plugins/Runtime/RigVM/Source/RigVMEditor/Public/Widgets/SRigVMExecutionStackView.h
+++ b/Engine/Plugins/Runtime/RigVM/Source/RigVMEditor/Public/Widgets/SRigVMExecutionStackView.h
@@ -115,9 +115,6 @@ private:
 
 	UE_API void OnVMCompiled(UObject* InCompiledObject, URigVM* InCompiledVM, FRigVMExtendedExecuteContext& InVMContext);
 
-	//* Focus on the instruction when the execution is halted */
-	UE_API void HandleExecutionHalted(const int32 HaltedAtInstruction, UObject* InNode, const FName& InEntryName);
-
 	/** Search box widget */
 	TSharedPtr<SSearchBox> FilterBox;
 	FText FilterText;
@@ -144,8 +141,6 @@ private:
 
 	TArray<TSharedPtr<FRigStackEntry>> Operators;
 
-	int32 HaltedAtInstruction;
-
 	FDelegateHandle OnModelModified;
 	FDelegateHandle OnHostInitializedHandle;
 	FDelegateHandle OnVMCompiledHandle;
