RigVM: implement middle mouse selection pattern from dependency graph also in main graph

#rnx

diff --git a/Engine/Plugins/Runtime/RigVM/Source/RigVMDeveloper/Private/RigVMModel/RigVMController.cpp b/Engine/Plugins/Runtime/RigVM/Source/RigVMDeveloper/Private/RigVMModel/RigVMController.cpp
--- a/Engine/Plugins/Runtime/RigVM/Source/RigVMDeveloper/Private/RigVMModel/RigVMController.cpp
+++ b/Engine/Plugins/Runtime/RigVM/Source/RigVMDeveloper/Private/RigVMModel/RigVMController.cpp
@@ -7663,6 +7663,172 @@ bool URigVMController::SetNodeSelection(const TArray<FName>& InNodeNames, bool b
 	return bSelectionChanged;
 }
 
+bool URigVMController::SelectLinkedNodes(const TArray<FName>& InNodeNames, bool bSelectSourceNodes, bool bClearSelection, bool bSetupUndoRedo,
+	bool bPrintPythonCommand)
+{
+	if (!IsValidGraph())
+	{
+		return false;
+	}
+
+	if (!bIsTransacting && !IsGraphEditable())
+	{
+		return false;
+	}
+
+	URigVMGraph* Graph = GetGraph();
+	check(Graph);
+
+	TArray<FName> NewSelection;
+	for(const FName& NodeName : InNodeNames)
+	{
+		NewSelection.AddUnique(NodeName);
+	} 
+	
+	for (int32 Index = 0; Index < NewSelection.Num(); Index++)
+	{
+		const URigVMNode* Node = Graph->FindNodeByName(NewSelection[Index]);
+		if (!Node)
+		{
+			continue;
+		}
+
+		const bool bIsLinkedNode = !InNodeNames.Contains(Node->GetFName());
+		if (bSelectSourceNodes || bIsLinkedNode)
+		{
+			const bool bSkipMutableNodes = bIsLinkedNode && !bSelectSourceNodes;
+			const TArray<URigVMNode*> LinkedNodes = Node->GetLinkedSourceNodes();
+			for (const URigVMNode* LinkedNode : LinkedNodes)
+			{
+				if(bSkipMutableNodes && LinkedNode->IsMutable())
+				{
+					continue;
+				}
+				NewSelection.AddUnique(LinkedNode->GetFName());
+			}
+		}
+		if (!bSelectSourceNodes)
+		{
+			const TArray<URigVMNode*> LinkedNodes = Node->GetLinkedTargetNodes();
+			for (const URigVMNode* LinkedNode : LinkedNodes)
+			{
+				NewSelection.AddUnique(LinkedNode->GetFName());
+			}
+		}
+	}
+
+	if(!bClearSelection)
+	{
+		const TArray<FName> PreviousSelection = Graph->GetSelectNodes();	
+		for(const FName& NodeName : PreviousSelection)
+		{
+			NewSelection.AddUnique(NodeName);
+		} 
+	}
+
+	const bool bSelectionChanged = SetNodeSelection(NewSelection, bSetupUndoRedo, false);
+
+	if (bPrintPythonCommand)
+	{
+		FString ArrayStr = TEXT("[");
+		for (auto It = InNodeNames.CreateConstIterator(); It; ++It)
+		{
+			ArrayStr += TEXT("'") + GetSchema()->GetSanitizedNodeName(It->ToString()) + TEXT("'");
+			if (It.GetIndex() < InNodeNames.Num() - 1)
+			{
+				ArrayStr += TEXT(", ");
+			}
+		}
+		ArrayStr += TEXT("]");
+
+		const FString GraphName = GetSchema()->GetSanitizedGraphName(GetGraph()->GetGraphName());
+
+		RigVMPythonUtils::Print(GetSchema()->GetGraphOuterName(GetGraph()), 
+							FString::Printf(TEXT("blueprint.get_controller_by_name('%s').select_linked_nodes(%s, %s, %s)"),
+											*GraphName,
+											*ArrayStr,
+											(bSelectSourceNodes) ? TEXT("True") : TEXT("False"),
+											(bClearSelection) ? TEXT("True") : TEXT("False")));
+	}
+
+	return bSelectionChanged;
+}
+
+bool URigVMController::SelectNodeIslands(const TArray<FName>& InNodeNames, bool bClearSelection, bool bSetupUndoRedo, bool bPrintPythonCommand)
+{
+	if (!IsValidGraph())
+	{
+		return false;
+	}
+
+	if (!bIsTransacting && !IsGraphEditable())
+	{
+		return false;
+	}
+
+	URigVMGraph* Graph = GetGraph();
+	check(Graph);
+
+	TArray<FName> NewSelection;
+	for(const FName& NodeName : InNodeNames)
+	{
+		NewSelection.AddUnique(NodeName);
+	} 
+	
+	if(!bClearSelection)
+	{
+		const TArray<FName> PreviousSelection = Graph->GetSelectNodes();	
+		for(const FName& NodeName : PreviousSelection)
+		{
+			NewSelection.AddUnique(NodeName);
+		} 
+	}
+
+	for (int32 Index = 0; Index < NewSelection.Num(); Index++)
+	{
+		const URigVMNode* Node = Graph->FindNodeByName(NewSelection[Index]);
+		if (!Node)
+		{
+			continue;
+		}
+		
+		TArray<URigVMNode*> LinkedNodes;
+		LinkedNodes.Append(Node->GetLinkedSourceNodes());
+		LinkedNodes.Append(Node->GetLinkedTargetNodes());
+		
+		for (const URigVMNode* LinkedNode : LinkedNodes)
+		{
+			NewSelection.AddUnique(LinkedNode->GetFName());
+		}
+	}
+
+	const bool bSelectionChanged = SetNodeSelection(NewSelection, bSetupUndoRedo, false);
+
+	if (bPrintPythonCommand)
+	{
+		FString ArrayStr = TEXT("[");
+		for (auto It = InNodeNames.CreateConstIterator(); It; ++It)
+		{
+			ArrayStr += TEXT("'") + GetSchema()->GetSanitizedNodeName(It->ToString()) + TEXT("'");
+			if (It.GetIndex() < InNodeNames.Num() - 1)
+			{
+				ArrayStr += TEXT(", ");
+			}
+		}
+		ArrayStr += TEXT("]");
+
+		const FString GraphName = GetSchema()->GetSanitizedGraphName(GetGraph()->GetGraphName());
+
+		RigVMPythonUtils::Print(GetSchema()->GetGraphOuterName(GetGraph()), 
+							FString::Printf(TEXT("blueprint.get_controller_by_name('%s').select_node_islands(%s, %s)"),
+											*GraphName,
+											*ArrayStr,
+											(bClearSelection) ? TEXT("True") : TEXT("False")));
+	}
+
+	return bSelectionChanged;
+}
+
 bool URigVMController::SetNodePosition(URigVMNode* InNode, const FVector2D& InPosition, bool bSetupUndoRedo, bool bMergeUndoAction, bool bPrintPythonCommand)
 {
 	if (!bIsTransacting && !IsGraphEditable())
diff --git a/Engine/Plugins/Runtime/RigVM/Source/RigVMDeveloper/Public/EdGraph/RigVMEdGraph.h b/Engine/Plugins/Runtime/RigVM/Source/RigVMDeveloper/Public/EdGraph/RigVMEdGraph.h
--- a/Engine/Plugins/Runtime/RigVM/Source/RigVMDeveloper/Public/EdGraph/RigVMEdGraph.h
+++ b/Engine/Plugins/Runtime/RigVM/Source/RigVMDeveloper/Public/EdGraph/RigVMEdGraph.h
@@ -8,6 +8,7 @@
 #include "RigVMModel/RigVMGraph.h"
 #include "RigVMCore/RigVM.h"
 #include "RigVMModel/RigVMClient.h"
+#include "Input/Events.h"
 #include "RigVMEdGraph.generated.h"
 
 #define UE_API RIGVMDEVELOPER_API
@@ -129,7 +130,7 @@ private:
 	FRigVMStringTag Tag;
 };
 
-DECLARE_MULTICAST_DELEGATE_OneParam(FRigVMEdGraphNodeClicked, URigVMEdGraphNode*);
+DECLARE_MULTICAST_DELEGATE_ThreeParams(FRigVMEdGraphNodeClicked, URigVMEdGraphNode*, const FGeometry&, const FPointerEvent&);
 
 UCLASS(MinimalAPI)
 class URigVMEdGraph : public UEdGraph, public IRigVMEditorSideObject
diff --git a/Engine/Plugins/Runtime/RigVM/Source/RigVMDeveloper/Public/RigVMModel/RigVMController.h b/Engine/Plugins/Runtime/RigVM/Source/RigVMDeveloper/Public/RigVMModel/RigVMController.h
--- a/Engine/Plugins/Runtime/RigVM/Source/RigVMDeveloper/Public/RigVMModel/RigVMController.h
+++ b/Engine/Plugins/Runtime/RigVM/Source/RigVMDeveloper/Public/RigVMModel/RigVMController.h
@@ -728,6 +728,14 @@ public:
 	UFUNCTION(BlueprintCallable, Category = RigVMController)
 	UE_API bool SetNodeSelection(const TArray<FName>& InNodeNames, bool bSetupUndoRedo = true, bool bPrintPythonCommand = false);
 
+	// Selects the linked nodes given the input node names
+	UFUNCTION(BlueprintCallable, Category = RigVMController)
+	UE_API bool SelectLinkedNodes(const TArray<FName>& InNodeNames, bool bSelectSourceNodes, bool bClearSelection, bool bSetupUndoRedo = true, bool bPrintPythonCommand = false);
+
+	// Selects the node islands given the input node names
+	UFUNCTION(BlueprintCallable, Category = RigVMController)
+	UE_API bool SelectNodeIslands(const TArray<FName>& InNodeNames, bool bClearSelection, bool bSetupUndoRedo = true, bool bPrintPythonCommand = false);
+
 	// Sets the position of a node in the graph.
 	// This causes a NodePositionChanged modified event.
 	UFUNCTION(BlueprintCallable, Category = RigVMController)
diff --git a/Engine/Plugins/Runtime/RigVM/Source/RigVMEditor/Private/Editor/RigVMEditor.cpp b/Engine/Plugins/Runtime/RigVM/Source/RigVMEditor/Private/Editor/RigVMEditor.cpp
--- a/Engine/Plugins/Runtime/RigVM/Source/RigVMEditor/Private/Editor/RigVMEditor.cpp
+++ b/Engine/Plugins/Runtime/RigVM/Source/RigVMEditor/Private/Editor/RigVMEditor.cpp
@@ -3651,7 +3651,7 @@ void FRigVMEditorBase::HandleVariableDroppedFromBlueprint(UObject* InSubject, FP
 	);
 }
 
-void FRigVMEditorBase::OnGraphNodeClicked(URigVMEdGraphNode* InNode)
+void FRigVMEditorBase::OnGraphNodeClicked(URigVMEdGraphNode* InNode, const FGeometry& InNodeGeometry, const FPointerEvent& InMouseEvent)
 {
 	if (InNode)
 	{
@@ -3659,6 +3659,24 @@ void FRigVMEditorBase::OnGraphNodeClicked(URigVMEdGraphNode* InNode)
 		{
 			SetDetailViewForGraph(InNode->GetModel());
 		}
+
+		if (InMouseEvent.IsMouseButtonDown(EKeys::MiddleMouseButton))
+		{
+			if (URigVMController* Controller = GetFocusedController())
+			{
+				const bool ClearSelection = !InMouseEvent.GetModifierKeys().IsShiftDown();
+				const bool SelectSourceNodes = !InMouseEvent.GetModifierKeys().IsControlDown();
+				const bool SelectIsland = InMouseEvent.GetModifierKeys().IsAltDown();
+				if (SelectIsland)
+				{
+					Controller->SelectNodeIslands({InNode->GetModelNodeName()}, ClearSelection);
+				}
+				else
+				{
+					Controller->SelectLinkedNodes({InNode->GetModelNodeName()}, SelectSourceNodes, ClearSelection, true);
+				}
+			}
+		}
 	}
 }
 
diff --git a/Engine/Plugins/Runtime/RigVM/Source/RigVMEditor/Private/Widgets/SRigVMGraphNode.cpp b/Engine/Plugins/Runtime/RigVM/Source/RigVMEditor/Private/Widgets/SRigVMGraphNode.cpp
--- a/Engine/Plugins/Runtime/RigVM/Source/RigVMEditor/Private/Widgets/SRigVMGraphNode.cpp
+++ b/Engine/Plugins/Runtime/RigVM/Source/RigVMEditor/Private/Widgets/SRigVMGraphNode.cpp
@@ -998,7 +998,7 @@ FReply SRigVMGraphNode::OnMouseButtonDown(const FGeometry& MyGeometry, const FPo
 	{
 		if (URigVMEdGraph* RigGraph = Cast<URigVMEdGraph>(RigNode->GetGraph()))
 		{
-			RigGraph->OnGraphNodeClicked.Broadcast(RigNode);
+			RigGraph->OnGraphNodeClicked.Broadcast(RigNode, MyGeometry, MouseEvent);
 		}
 	}
 
diff --git a/Engine/Plugins/Runtime/RigVM/Source/RigVMEditor/Public/Editor/RigVMEditor.h b/Engine/Plugins/Runtime/RigVM/Source/RigVMEditor/Public/Editor/RigVMEditor.h
--- a/Engine/Plugins/Runtime/RigVM/Source/RigVMEditor/Public/Editor/RigVMEditor.h
+++ b/Engine/Plugins/Runtime/RigVM/Source/RigVMEditor/Public/Editor/RigVMEditor.h
@@ -530,7 +530,7 @@ protected:
 	
 	UE_API virtual void HandleRefreshEditorFromBlueprint(TRigVMAssetInterfacePtr InBlueprint);
 	UE_API void HandleVariableDroppedFromBlueprint(UObject* InSubject, FProperty* InVariableToDrop, const FVector2D& InDropPosition, const FVector2D& InScreenPosition);
-	UE_API void OnGraphNodeClicked(URigVMEdGraphNode* InNode);
+	UE_API void OnGraphNodeClicked(URigVMEdGraphNode* InNode, const FGeometry& InNodeGeometry, const FPointerEvent& InMouseEvent);
 	UE_API void OnNodeDoubleClicked(TRigVMAssetInterfacePtr InBlueprint, URigVMNode* InNode);
 	UE_API void OnGraphImported(UEdGraph* InEdGraph);
 	UE_API bool OnActionMatchesName(FEdGraphSchemaAction* InAction, const FName& InName) const;
