RigVM: Improve instruction index cache for ed graph nodes.
Make flashlight functionality optional

#rnx

diff --git a/Engine/Plugins/Runtime/RigVM/Source/RigVM/Private/RigVMSettings.cpp b/Engine/Plugins/Runtime/RigVM/Source/RigVM/Private/RigVMSettings.cpp
--- a/Engine/Plugins/Runtime/RigVM/Source/RigVM/Private/RigVMSettings.cpp
+++ b/Engine/Plugins/Runtime/RigVM/Source/RigVM/Private/RigVMSettings.cpp
@@ -15,6 +15,7 @@ URigVMEditorSettings::URigVMEditorSettings(const FObjectInitializer& ObjectIniti
 #if WITH_EDITORONLY_DATA
 	bHighlightSimilarNodes = false;
 	bFadeOutUnrelatedNodes = false;
+	bUseFlashLight = true;
 	bAutoLinkMutableNodes = false;
 #endif
 }
diff --git a/Engine/Plugins/Runtime/RigVM/Source/RigVM/Public/RigVMSettings.h b/Engine/Plugins/Runtime/RigVM/Source/RigVM/Public/RigVMSettings.h
--- a/Engine/Plugins/Runtime/RigVM/Source/RigVM/Public/RigVMSettings.h
+++ b/Engine/Plugins/Runtime/RigVM/Source/RigVM/Public/RigVMSettings.h
@@ -37,6 +37,10 @@ public:
 	UPROPERTY(EditAnywhere, config, Category = Interaction)
 	bool bFadeOutUnrelatedNodes;
 
+	// use a flash light to brighten up nodes
+	UPROPERTY(EditAnywhere, config, Category = Interaction)
+	bool bUseFlashLight;
+
 	// When this is checked mutable nodes (nodes with an execute pin)
 	// will be hooked up automatically.
 	UPROPERTY(EditAnywhere, config, Category = Interaction)
diff --git a/Engine/Plugins/Runtime/RigVM/Source/RigVMDeveloper/Private/EdGraph/RigVMEdGraph.cpp b/Engine/Plugins/Runtime/RigVM/Source/RigVMDeveloper/Private/EdGraph/RigVMEdGraph.cpp
--- a/Engine/Plugins/Runtime/RigVM/Source/RigVMDeveloper/Private/EdGraph/RigVMEdGraph.cpp
+++ b/Engine/Plugins/Runtime/RigVM/Source/RigVMDeveloper/Private/EdGraph/RigVMEdGraph.cpp
@@ -889,6 +889,16 @@ int32 URigVMEdGraph::GetInstructionIndex(const URigVMEdGraphNode* InNode, bool b
 					return i;
 				}
 			}
+
+			if(bAsInput)
+			{
+				Indices.FindOrAdd(InModelNode).Key = INDEX_NONE;
+			}
+			else
+			{
+				Indices.FindOrAdd(InModelNode).Value = INDEX_NONE;
+			}
+
 			return INDEX_NONE;
 		}
 	};
diff --git a/Engine/Plugins/Runtime/RigVM/Source/RigVMEditor/Private/Editor/RigVMEditor.cpp b/Engine/Plugins/Runtime/RigVM/Source/RigVMEditor/Private/Editor/RigVMEditor.cpp
--- a/Engine/Plugins/Runtime/RigVM/Source/RigVMEditor/Private/Editor/RigVMEditor.cpp
+++ b/Engine/Plugins/Runtime/RigVM/Source/RigVMEditor/Private/Editor/RigVMEditor.cpp
@@ -4322,9 +4322,10 @@ void FRigVMEditorBase::OnGraphEditorTick(const FGeometry& AllottedGeometry, cons
 	// flashlight to lighten up nodes in proximity of the mouse
 	const FGeometry PaintGeometry = InGraphEditor->GetPaintSpaceGeometry();
 	const FVector2D MousePosition = FSlateApplication::Get().GetCursorPos();
+	const bool bUseFlashLight = GetDefault<URigVMEditorSettings>()->bUseFlashLight;
 	const float FlashLightRadius = 5.f * 64.f;
 
-	if(PaintGeometry.IsUnderLocation(MousePosition) && FlashLightRadius > SMALL_NUMBER)
+	if(bUseFlashLight && PaintGeometry.IsUnderLocation(MousePosition) && FlashLightRadius > SMALL_NUMBER)
 	{
 		FVector2f GraphEditorLocation = FVector2f::ZeroVector;
 		float GraphEditorZoomAmount = 0;
