Fix UAF crash due to RigVM assuming Blueprints are used in graph binding
Fix Crash on function add / undo due to assumption that controller is valid.

#tests Editor modify RigVM Assets

diff --git a/Engine/Plugins/Runtime/RigVM/Source/RigVMDeveloper/Private/EdGraph/RigVMEdGraph.cpp b/Engine/Plugins/Runtime/RigVM/Source/RigVMDeveloper/Private/EdGraph/RigVMEdGraph.cpp
--- a/Engine/Plugins/Runtime/RigVM/Source/RigVMDeveloper/Private/EdGraph/RigVMEdGraph.cpp
+++ b/Engine/Plugins/Runtime/RigVM/Source/RigVMDeveloper/Private/EdGraph/RigVMEdGraph.cpp
@@ -903,11 +903,14 @@ int32 URigVMEdGraph::GetInstructionIndex(const URigVMEdGraphNode* InNode, bool b
 		}
 	};
 
-	if (const FRigVMByteCode* ByteCode = GetController()->GetCurrentByteCode())
+	if (URigVMController* Controller = GetController())
 	{
-		const int32 SourceInstructionIndex = Local::GetInstructionIndex(InNode->GetModelNode(), ByteCode, CachedInstructionIndices, true);
-		const int32 TargetInstructionIndex = Local::GetInstructionIndex(InNode->GetModelNode(), ByteCode, CachedInstructionIndices, false);
-		return bAsInput ? SourceInstructionIndex : TargetInstructionIndex;
+		if (const FRigVMByteCode* ByteCode = Controller->GetCurrentByteCode())
+		{
+			const int32 SourceInstructionIndex = Local::GetInstructionIndex(InNode->GetModelNode(), ByteCode, CachedInstructionIndices, true);
+			const int32 TargetInstructionIndex = Local::GetInstructionIndex(InNode->GetModelNode(), ByteCode, CachedInstructionIndices, false);
+			return bAsInput ? SourceInstructionIndex : TargetInstructionIndex;
+		}
 	}
 
 	return INDEX_NONE;
diff --git a/Engine/Plugins/Runtime/RigVM/Source/RigVMEditor/Private/Widgets/SRigVMGraphPinVariableBinding.cpp b/Engine/Plugins/Runtime/RigVM/Source/RigVMEditor/Private/Widgets/SRigVMGraphPinVariableBinding.cpp
--- a/Engine/Plugins/Runtime/RigVM/Source/RigVMEditor/Private/Widgets/SRigVMGraphPinVariableBinding.cpp
+++ b/Engine/Plugins/Runtime/RigVM/Source/RigVMEditor/Private/Widgets/SRigVMGraphPinVariableBinding.cpp
@@ -57,12 +57,15 @@ void SRigVMGraphVariableBinding::Construct(const FArguments& InArgs)
 		nullptr,
 		FMenuExtensionDelegate::CreateSP(this, &SRigVMGraphVariableBinding::FillLocalVariableMenu));
 
-	TArray<FBindingContextStruct> BindingContexts;
-	BindingContexts.Add(Blueprint->GetRigVMGeneratedClass());
-	this->ChildSlot
-	[
-		PropertyAccessEditor.MakePropertyBindingWidget(BindingContexts, BindingArgs)
-	];
+	if (Blueprint)
+	{
+		TArray<FBindingContextStruct> BindingContexts;
+		BindingContexts.Add(Blueprint->GetRigVMGeneratedClass());
+		this->ChildSlot
+		[
+			PropertyAccessEditor.MakePropertyBindingWidget(BindingContexts, BindingArgs)
+		];
+	}
 }
 
 FText SRigVMGraphVariableBinding::GetBindingText(URigVMPin* ModelPin) const
