RigVM: Fix RigVMMemory Hashes

#rnx

diff --git a/Engine/Plugins/Runtime/RigVM/Source/RigVM/Private/RigVMCore/RigVMMemoryStorageStruct.cpp b/Engine/Plugins/Runtime/RigVM/Source/RigVM/Private/RigVMCore/RigVMMemoryStorageStruct.cpp
--- a/Engine/Plugins/Runtime/RigVM/Source/RigVM/Private/RigVMCore/RigVMMemoryStorageStruct.cpp
+++ b/Engine/Plugins/Runtime/RigVM/Source/RigVM/Private/RigVMCore/RigVMMemoryStorageStruct.cpp
@@ -637,8 +637,8 @@ void FRigVMMemoryStorageStruct::SetDefaultValues(const TArray<FRigVMPropertyDesc
 		const FString Name = RigVMDescriptor.Name.ToString();
 
 		FBlake3 Builder;
-		Builder.Update(*Name, Name.Len());
-		Builder.Update(*RigVMDescriptor.CPPType, RigVMDescriptor.CPPType.Len());
+		Builder.Update(*Name, Name.Len() * sizeof(TCHAR));
+		Builder.Update(*RigVMDescriptor.CPPType, RigVMDescriptor.CPPType.Len()  * sizeof(TCHAR));
 		FBlake3Hash Hash = Builder.Finalize();
 		Result.ID = FGuid::NewGuidFromHash(Hash);
 	}
