RigVM: Fix multi-threading issues accessing function-static containers.
- Also fixed zero-pad string printing of integers to avoid a large switch statement.

#rnx

diff --git a/Engine/Plugins/Runtime/RigVM/Source/RigVM/Private/RigVMFunctions/RigVMDispatch_Array.cpp b/Engine/Plugins/Runtime/RigVM/Source/RigVM/Private/RigVMFunctions/RigVMDispatch_Array.cpp
--- a/Engine/Plugins/Runtime/RigVM/Source/RigVM/Private/RigVMFunctions/RigVMDispatch_Array.cpp
+++ b/Engine/Plugins/Runtime/RigVM/Source/RigVM/Private/RigVMFunctions/RigVMDispatch_Array.cpp
@@ -1514,7 +1514,7 @@ const TArray<FRigVMTemplateArgumentInfo>& FRigVMDispatch_ArrayIterator::GetArgum
 const TArray<FRigVMExecuteArgument>& FRigVMDispatch_ArrayIterator::GetExecuteArguments_Impl(
 	const FRigVMDispatchContext& InContext) const
 {
-	static TArray<FRigVMExecuteArgument> ExecuteArguments;
+	thread_local TArray<FRigVMExecuteArgument> ExecuteArguments;
 	if(ExecuteArguments.IsEmpty())
 	{
 		ExecuteArguments = FRigVMDispatch_ArrayBaseMutable::GetExecuteArguments_Impl(InContext);
diff --git a/Engine/Plugins/Runtime/RigVM/Source/RigVM/Private/RigVMFunctions/RigVMDispatch_Switch.cpp b/Engine/Plugins/Runtime/RigVM/Source/RigVM/Private/RigVMFunctions/RigVMDispatch_Switch.cpp
--- a/Engine/Plugins/Runtime/RigVM/Source/RigVM/Private/RigVMFunctions/RigVMDispatch_Switch.cpp
+++ b/Engine/Plugins/Runtime/RigVM/Source/RigVM/Private/RigVMFunctions/RigVMDispatch_Switch.cpp
@@ -131,8 +131,8 @@ void FRigVMDispatch_SwitchInt32::Execute(FRigVMExtendedExecuteContext& InContext
 	check(Handles[1].IsName());
 #endif
 
-	const int32 Index = *(const int32*)Handles[0].GetInputData();
-	FName& BlockToRun = *(FName*)Handles[1].GetOutputData();
+	const int32 Index = *reinterpret_cast<const int32*>(Handles[0].GetInputData());
+	FName& BlockToRun = *reinterpret_cast<FName*>(Handles[1].GetOutputData());
 
 	if(BlockToRun.IsNone())
 	{
@@ -146,7 +146,7 @@ void FRigVMDispatch_SwitchInt32::Execute(FRigVMExtendedExecuteContext& InContext
 
 FName FRigVMDispatch_SwitchInt32::GetCaseName(int32 InIndex)
 {
-	static TMap<int32, FName> NamesFromInt;
+	thread_local TMap<int32, FName> NamesFromInt;
 	if (FName* Found = NamesFromInt.Find(InIndex))
 	{
 		return *Found;
@@ -159,7 +159,7 @@ FName FRigVMDispatch_SwitchInt32::GetCaseName(int32 InIndex)
 
 FName FRigVMDispatch_SwitchInt32::GetCaseDisplayName(int32 InIndex)
 {
-	static TMap<int32, FName> NamesFromInt;
+	thread_local TMap<int32, FName> NamesFromInt;
     if (FName* Found = NamesFromInt.Find(InIndex))
     {
     	return *Found;
diff --git a/Engine/Plugins/Runtime/RigVM/Source/RigVM/Private/RigVMFunctions/RigVMFunction_String.cpp b/Engine/Plugins/Runtime/RigVM/Source/RigVM/Private/RigVMFunctions/RigVMFunction_String.cpp
--- a/Engine/Plugins/Runtime/RigVM/Source/RigVM/Private/RigVMFunctions/RigVMFunction_String.cpp
+++ b/Engine/Plugins/Runtime/RigVM/Source/RigVM/Private/RigVMFunctions/RigVMFunction_String.cpp
@@ -152,104 +152,13 @@ FRigVMFunction_StringJoin_Execute()
 
 FRigVMFunction_StringPadInteger_Execute()
 {
-	static constexpr TCHAR Format02Digits[] = TEXT("%02d");
-	static constexpr TCHAR Format03Digits[] = TEXT("%03d");
-	static constexpr TCHAR Format04Digits[] = TEXT("%04d");
-	static constexpr TCHAR Format05Digits[] = TEXT("%05d");
-	static constexpr TCHAR Format06Digits[] = TEXT("%06d");
-	static constexpr TCHAR Format07Digits[] = TEXT("%07d");
-	static constexpr TCHAR Format08Digits[] = TEXT("%08d");
-	static constexpr TCHAR Format09Digits[] = TEXT("%09d");
-	static constexpr TCHAR Format10Digits[] = TEXT("%10d");
-	static constexpr TCHAR Format11Digits[] = TEXT("%11d");
-	static constexpr TCHAR Format12Digits[] = TEXT("%12d");
-	static constexpr TCHAR Format13Digits[] = TEXT("%13d");
-	static constexpr TCHAR Format14Digits[] = TEXT("%14d");
-	static constexpr TCHAR Format15Digits[] = TEXT("%15d");
-	static constexpr TCHAR Format16Digits[] = TEXT("%16d");
-
-	switch(Digits)
+	if (Digits >= 2 && Digits <= 16)
 	{
-		case 2:
-		{
-			Result = FString::Printf(Format02Digits, Value); 
-			break;
-		}
-		case 3:
-		{
-			Result = FString::Printf(Format03Digits, Value); 
-			break;
-		}
-		case 4:
-		{
-			Result = FString::Printf(Format04Digits, Value); 
-			break;
-		}
-		case 5:
-		{
-			Result = FString::Printf(Format05Digits, Value); 
-			break;
-		}
-		case 6:
-		{
-			Result = FString::Printf(Format06Digits, Value); 
-			break;
-		}
-		case 7:
-		{
-			Result = FString::Printf(Format07Digits, Value); 
-			break;
-		}
-		case 8:
-		{
-			Result = FString::Printf(Format08Digits, Value); 
-			break;
-		}
-		case 9:
-		{
-			Result = FString::Printf(Format09Digits, Value); 
-			break;
-		}
-		case 10:
-		{
-			Result = FString::Printf(Format10Digits, Value); 
-			break;
-		}
-		case 11:
-		{
-			Result = FString::Printf(Format11Digits, Value); 
-			break;
-		}
-		case 12:
-		{
-			Result = FString::Printf(Format12Digits, Value); 
-			break;
-		}
-		case 13:
-		{
-			Result = FString::Printf(Format13Digits, Value); 
-			break;
-		}
-		case 14:
-		{
-			Result = FString::Printf(Format14Digits, Value); 
-			break;
-		}
-		case 15:
-		{
-			Result = FString::Printf(Format15Digits, Value); 
-			break;
-		}
-		case 16:
-		{
-			Result = FString::Printf(Format16Digits, Value); 
-			break;
-		}
-		default:
-		{
-			Result = FString::FormatAsNumber(Value);
-			break;
-		}
+		Result = FString::Printf(TEXT("%0*d"), Digits, Value); 
+	}
+	else
+	{
+		Result = FString::FormatAsNumber(Value);
 	}
 }
 
