Fix minor code review feedback on const.

#tests nonunity

diff --git a/Engine/Plugins/Runtime/RigVM/Source/RigVMDeveloper/Private/EdGraph/RigVMEdGraph.cpp b/Engine/Plugins/Runtime/RigVM/Source/RigVMDeveloper/Private/EdGraph/RigVMEdGraph.cpp
--- a/Engine/Plugins/Runtime/RigVM/Source/RigVMDeveloper/Private/EdGraph/RigVMEdGraph.cpp
+++ b/Engine/Plugins/Runtime/RigVM/Source/RigVMDeveloper/Private/EdGraph/RigVMEdGraph.cpp
@@ -903,7 +903,7 @@ int32 URigVMEdGraph::GetInstructionIndex(const URigVMEdGraphNode* InNode, bool b
 		}
 	};
 
-	if (URigVMController* Controller = GetController())
+	if (const URigVMController* Controller = GetController())
 	{
 		if (const FRigVMByteCode* ByteCode = Controller->GetCurrentByteCode())
 		{
