Control Rig: Fix crash when trying to find public function headers in the generated class
#rnx

diff --git a/Engine/Plugins/Runtime/RigVM/Source/RigVMDeveloper/Private/RigVMModel/RigVMBuildData.cpp b/Engine/Plugins/Runtime/RigVM/Source/RigVMDeveloper/Private/RigVMModel/RigVMBuildData.cpp
--- a/Engine/Plugins/Runtime/RigVM/Source/RigVMDeveloper/Private/RigVMModel/RigVMBuildData.cpp
+++ b/Engine/Plugins/Runtime/RigVM/Source/RigVMDeveloper/Private/RigVMModel/RigVMBuildData.cpp
@@ -221,8 +221,7 @@ TArray<FRigVMGraphFunctionHeader> URigVMBuildData::GetFunctionHeadersForAsset(co
 	const FString PublicGraphFunctionsString = InAssetData.GetTagValueRef<FString>(TEXT("PublicGraphFunctions"));
 	if(!PublicGraphFunctionsString.IsEmpty())
 	{
-		static const FName FunctionsPropertyName = TEXT("PublicGraphFunctions");
-		const FArrayProperty* PublicGraphFunctionsProperty = CastField<FArrayProperty>(InAssetData.GetClass()->FindPropertyByName(FunctionsPropertyName));
+		const FArrayProperty* PublicGraphFunctionsProperty = CastField<FArrayProperty>(FRigVMGraphFunctionHeaderArray::StaticStruct()->FindPropertyByName(TEXT("Headers")));
 		PublicGraphFunctionsProperty->ImportText_Direct(*PublicGraphFunctionsString, &Result, nullptr, EPropertyPortFlags::PPF_None);
 	}
 	return Result;
