Control Rig: Fix possible crash when trying to find public function headers in the generated class (2nd try)
#rnx

diff --git a/Engine/Plugins/Runtime/RigVM/Source/RigVMDeveloper/Private/RigVMModel/RigVMBuildData.cpp b/Engine/Plugins/Runtime/RigVM/Source/RigVMDeveloper/Private/RigVMModel/RigVMBuildData.cpp
--- a/Engine/Plugins/Runtime/RigVM/Source/RigVMDeveloper/Private/RigVMModel/RigVMBuildData.cpp
+++ b/Engine/Plugins/Runtime/RigVM/Source/RigVMDeveloper/Private/RigVMModel/RigVMBuildData.cpp
@@ -896,7 +896,7 @@ TArray<FRigVMGraphFunctionIdentifier> URigVMBuildData::GetAllFunctionIdentifiers
 		}
 		if (!PublicGraphFunctionsString.IsEmpty())
 		{
-			const FArrayProperty* Property = CastField<FArrayProperty>(AssetData.GetClass()->FindPropertyByName(PublicGraphFunctionsName));
+			const FArrayProperty* Property = CastField<FArrayProperty>(FRigVMGraphFunctionHeaderArray::StaticStruct()->FindPropertyByName(TEXT("Headers")));
 			if (Property)
 			{
 				TArray<FRigVMGraphFunctionHeader> PublicFunctions;
