Control Rig: Fix modular rigs not showing up to bake into control rig
#rnx

diff --git a/Engine/Plugins/Runtime/RigVM/Source/RigVMDeveloper/Private/RigVMAsset.cpp b/Engine/Plugins/Runtime/RigVM/Source/RigVMDeveloper/Private/RigVMAsset.cpp
--- a/Engine/Plugins/Runtime/RigVM/Source/RigVMDeveloper/Private/RigVMAsset.cpp
+++ b/Engine/Plugins/Runtime/RigVM/Source/RigVMDeveloper/Private/RigVMAsset.cpp
@@ -1136,12 +1136,7 @@ void IRigVMAssetInterface::PreSave(FObjectPreSaveContext ObjectSaveContext)
 {
 	GetRigVMClient()->PreSave(ObjectSaveContext);
 
-	GetSupportedEventNames().Reset();
-	if (URigVM* VM = GetVM(true))
-	{
-		GetSupportedEventNames() = VM->GetEntryNames();
-	}
-	
+	UpdateSupportedEventNames();
 	if (TScriptInterface<IRigVMGraphFunctionHost> FunctionHost = GetRigVMClientHost()->GetRigVMGraphFunctionHost())
 	{
 		FRigVMGraphFunctionStore* Store = FunctionHost->GetRigVMGraphFunctionStore();
diff --git a/Engine/Plugins/Runtime/RigVM/Source/RigVMDeveloper/Private/RigVMBlueprintLegacy.cpp b/Engine/Plugins/Runtime/RigVM/Source/RigVMDeveloper/Private/RigVMBlueprintLegacy.cpp
--- a/Engine/Plugins/Runtime/RigVM/Source/RigVMDeveloper/Private/RigVMBlueprintLegacy.cpp
+++ b/Engine/Plugins/Runtime/RigVM/Source/RigVMDeveloper/Private/RigVMBlueprintLegacy.cpp
@@ -87,6 +87,15 @@ void URigVMBlueprint::BeginDestroy()
 	Clear();
 }
 
+void URigVMBlueprint::UpdateSupportedEventNames()
+{
+	SupportedEventNames.Reset();
+	if (URigVMHost* CDO = Cast<URigVMHost>(GetDefaultsObject()))
+	{
+		SupportedEventNames = CDO->GetSupportedEvents();
+	}
+}
+
 UObject* URigVMBlueprint::GetDefaultsObject()
 {
 	if (GeneratedClass)
diff --git a/Engine/Plugins/Runtime/RigVM/Source/RigVMDeveloper/Public/RigVMAsset.h b/Engine/Plugins/Runtime/RigVM/Source/RigVMDeveloper/Public/RigVMAsset.h
--- a/Engine/Plugins/Runtime/RigVM/Source/RigVMDeveloper/Public/RigVMAsset.h
+++ b/Engine/Plugins/Runtime/RigVM/Source/RigVMDeveloper/Public/RigVMAsset.h
@@ -396,7 +396,8 @@ protected:
 	virtual bool& IsReferencedObjectPathsStored() = 0;
 	virtual void SetReferencedObjectPathsStored(bool bValue) { IsReferencedObjectPathsStored() = bValue; }
 	virtual TArray<FSoftObjectPath>& GetReferencedObjectPaths() = 0;
-	virtual TArray<FName>& GetSupportedEventNames() = 0;
+	virtual TArray<FName> GetSupportedEventNames() = 0;
+	virtual void UpdateSupportedEventNames() = 0;
 	virtual TArray<FRigVMReferenceNodeData>& GetFunctionReferenceNodeData() = 0;
 	virtual void NotifyGraphRenamedSuper(class UEdGraph* Graph, FName OldName, FName NewName) = 0;
 	virtual URigVMHost* CreateRigVMHostSuper(UObject* InOuter) = 0;
diff --git a/Engine/Plugins/Runtime/RigVM/Source/RigVMDeveloper/Public/RigVMBlueprintLegacy.h b/Engine/Plugins/Runtime/RigVM/Source/RigVMDeveloper/Public/RigVMBlueprintLegacy.h
--- a/Engine/Plugins/Runtime/RigVM/Source/RigVMDeveloper/Public/RigVMBlueprintLegacy.h
+++ b/Engine/Plugins/Runtime/RigVM/Source/RigVMDeveloper/Public/RigVMBlueprintLegacy.h
@@ -147,7 +147,8 @@ protected:
 	virtual TArray<TObjectPtr<UEdGraph>>& GetFunctionGraphs() override { return FunctionGraphs; }
 	virtual bool& IsReferencedObjectPathsStored() override { return ReferencedObjectPathsStored; }
 	virtual TArray<FSoftObjectPath>& GetReferencedObjectPaths() override { return ReferencedObjectPaths; }
-	virtual TArray<FName>& GetSupportedEventNames() override { return SupportedEventNames; }
+	virtual TArray<FName> GetSupportedEventNames() override { return SupportedEventNames; }
+	virtual void UpdateSupportedEventNames() override;
 	virtual TArray<FRigVMReferenceNodeData>& GetFunctionReferenceNodeData() override { return FunctionReferenceNodeData; }
 	virtual void NotifyGraphRenamedSuper(class UEdGraph* Graph, FName OldName, FName NewName) override { UBlueprint::NotifyGraphRenamed(Graph, OldName, NewName); }
 	virtual UObject* GetDefaultsObject() override;
