Control Rig: Right clicking on a node now shows a Find References by Name option

diff --git a/Engine/Plugins/Runtime/RigVM/Source/RigVMEditor/Private/RigVMEditorModule.cpp b/Engine/Plugins/Runtime/RigVM/Source/RigVMEditor/Private/RigVMEditorModule.cpp
--- a/Engine/Plugins/Runtime/RigVM/Source/RigVMEditor/Private/RigVMEditorModule.cpp
+++ b/Engine/Plugins/Runtime/RigVM/Source/RigVMEditor/Private/RigVMEditorModule.cpp
@@ -649,6 +649,7 @@ void FRigVMEditorModule::GetNodeContextMenuActions(TScriptInterface<IRigVMClient
 		}
 	}
 
+	GetNodeActionsContextMenuActions(RigVMClientHost, EdGraphNode, ModelNode, Menu);
 	GetNodeWorkflowContextMenuActions(RigVMClientHost, EdGraphNode, ModelNode, Menu);
 	GetNodeEventsContextMenuActions(RigVMClientHost, EdGraphNode, ModelNode, Menu);
 	GetNodeDefaultValueContextMenuActions(RigVMClientHost, EdGraphNode, ModelNode, Menu);
@@ -662,6 +663,22 @@ void FRigVMEditorModule::GetNodeContextMenuActions(TScriptInterface<IRigVMClient
 	GetNodeDisplayContextMenuActions(RigVMClientHost, EdGraphNode, ModelNode, Menu);
 }
 
+void FRigVMEditorModule::GetNodeActionsContextMenuActions(TScriptInterface<IRigVMClientHost> RigVMClientHost, const URigVMEdGraphNode* EdGraphNode, URigVMNode* ModelNode, UToolMenu* Menu) const
+{
+	if (!EdGraphNode || !ModelNode || !Menu)
+	{
+		return;
+	}
+
+	// This section displays actions that are to be found under 'Node Actions' in other blueprint editors as well
+	FToolMenuSection& NodeActionsSection = Menu->AddSection("RigVMEditorContextMenuNodeActions", LOCTEXT("NodeActionsHeader", "Node Actions"));
+
+	NodeActionsSection.AddMenuEntry(
+		FGraphEditorCommands::Get().FindReferencesByNameLocal, 
+		LOCTEXT("FindReferencesByNameLabel", "Find References By Name")
+	);
+}
+
 void FRigVMEditorModule::GetNodeWorkflowContextMenuActions(TScriptInterface<IRigVMClientHost> RigVMClientHost, const URigVMEdGraphNode* EdGraphNode, URigVMNode* ModelNode, UToolMenu* Menu) const
 {
 	const TArray<FRigVMUserWorkflow> Workflows = ModelNode->GetSupportedWorkflows(ERigVMUserWorkflowType::NodeContext, ModelNode);
diff --git a/Engine/Plugins/Runtime/RigVM/Source/RigVMEditor/Public/RigVMEditorModule.h b/Engine/Plugins/Runtime/RigVM/Source/RigVMEditor/Public/RigVMEditorModule.h
--- a/Engine/Plugins/Runtime/RigVM/Source/RigVMEditor/Public/RigVMEditorModule.h
+++ b/Engine/Plugins/Runtime/RigVM/Source/RigVMEditor/Public/RigVMEditorModule.h
@@ -118,6 +118,7 @@ protected:
 	UE_API virtual UScriptStruct* GetRigVMExecuteContextStruct() const { return FRigVMExecuteContext::StaticStruct(); }
 
 	/** Specific section callbacks for the context menu */
+	UE_API virtual void GetNodeActionsContextMenuActions(TScriptInterface<IRigVMClientHost> RigVMClientHost, const URigVMEdGraphNode* EdGraphNode, URigVMNode* ModelNode, UToolMenu* Menu) const;
 	UE_API virtual void GetNodeWorkflowContextMenuActions(TScriptInterface<IRigVMClientHost> RigVMClientHost, const URigVMEdGraphNode* EdGraphNode, URigVMNode* ModelNode, UToolMenu* Menu) const;
 	UE_API virtual void GetNodeEventsContextMenuActions(TScriptInterface<IRigVMClientHost> RigVMClientHost, const URigVMEdGraphNode* EdGraphNode, URigVMNode* ModelNode, UToolMenu* Menu) const;
 	UE_API virtual void GetNodeDefaultValueContextMenuActions(TScriptInterface<IRigVMClientHost> RigVMClientHost, const URigVMEdGraphNode* EdGraphNode, URigVMNode* ModelNode, UToolMenu* Menu) const;
