Control Rig: Function nodes from other assets now show the asset name instead of the asset path, localize node titles and node subtitles

#prelfight

diff --git a/Engine/Plugins/Runtime/RigVM/Source/RigVMDeveloper/Private/EdGraph/RigVMEdGraphNode.cpp b/Engine/Plugins/Runtime/RigVM/Source/RigVMDeveloper/Private/EdGraph/RigVMEdGraphNode.cpp
--- a/Engine/Plugins/Runtime/RigVM/Source/RigVMDeveloper/Private/EdGraph/RigVMEdGraphNode.cpp
+++ b/Engine/Plugins/Runtime/RigVM/Source/RigVMDeveloper/Private/EdGraph/RigVMEdGraphNode.cpp
@@ -103,15 +103,17 @@ FText URigVMEdGraphNode::GetNodeTitle(ENodeTitleType::Type TitleType) const
 				const FRigVMGraphFunctionHeader& ReferencedHeader = FunctionReferenceNode->GetReferencedFunctionHeader();
 				{
 					TSoftObjectPtr<URigVMFunctionReferenceNode> RefNodePtr(FunctionReferenceNode);
-					const FString& PackagePath = ReferencedHeader.LibraryPointer.GetNodeSoftPath().GetLongPackageName();
+
+					const FSoftObjectPath NodeSoftPath = ReferencedHeader.LibraryPointer.GetNodeSoftPath();
+					const FString& PackagePath = NodeSoftPath.GetLongPackageName();
 					if(PackagePath != RefNodePtr.GetLongPackageName())
 					{
-						SubTitle = FString::Printf(TEXT("From %s"), *PackagePath);
+						const FString AssetName = NodeSoftPath.GetAssetName();
+						SubTitle = FText::Format(LOCTEXT("AssetNameSubTitle", "From {0}"), FText::FromString(AssetName)).ToString();
 					}
 					else
 					{
-						static const FString LocalFunctionString = TEXT("Local Function");
-						SubTitle = LocalFunctionString;
+						SubTitle = LOCTEXT("LocalFunctionSubTitle", "Local Function").ToString();
 					}
 				}
 			}
@@ -120,8 +122,7 @@ FText URigVMEdGraphNode::GetNodeTitle(ENodeTitleType::Type TitleType) const
 			{
 				if(!CollapseNode->IsA<URigVMAggregateNode>())
 				{
-					static const FString CollapseNodeString = TEXT("Collapsed Graph");
-					SubTitle = CollapseNodeString;
+					SubTitle = LOCTEXT("CollapsedGraphSubtitle", "Collapsed Graph").ToString();
 				}
 			}
 
@@ -141,8 +142,7 @@ FText URigVMEdGraphNode::GetNodeTitle(ENodeTitleType::Type TitleType) const
 								FString DefaultValue = NewVariable.DefaultValue;
 								if(DefaultValue.IsEmpty())
 								{
-									static const FString VariableString = TEXT("Variable");
-									SubTitle = VariableString;
+									SubTitle = LOCTEXT("VariableSubtitle", "Variable").ToString();;
 								}
 								else
 								{
@@ -159,7 +159,7 @@ FText URigVMEdGraphNode::GetNodeTitle(ENodeTitleType::Type TitleType) const
 										}
 										DefaultValue = FString::Join(Values, TEXT(","));										
 									}
-									SubTitle = FString::Printf(TEXT("Default %s"), *DefaultValue);
+									SubTitle = FText::Format(LOCTEXT("DefaultValueSubTitle", "Default {0}"), FText::FromString(DefaultValue)).ToString();
 								}
 								break;
 							}
@@ -181,12 +181,12 @@ FText URigVMEdGraphNode::GetNodeTitle(ENodeTitleType::Type TitleType) const
 
 		if(IsDeprecated())
 		{
-			NodeTitle = FText::FromString(FString::Printf(TEXT("%s (Deprecated)"), *NodeTitle.ToString()));
+			NodeTitle = FText::Format(LOCTEXT("DeprecatedNodeTitle", "{0} (Deprecated)"), NodeTitle);
 		}
 
 		if(IsOutDated())
 		{
-			NodeTitle = FText::FromString(FString::Printf(TEXT("%s (OutDated)"), *NodeTitle.ToString()));
+			NodeTitle = FText::Format(LOCTEXT("OutDatedNodeTitle", "{0} (OutDated)"), NodeTitle);
 		}
 
 		FullNodeTitle = NodeTitle;
