RigVM: Solidify localization of registries

#rnx

diff --git a/Engine/Plugins/Runtime/RigVM/Source/RigVM/Private/RigVMCore/RigVM.cpp b/Engine/Plugins/Runtime/RigVM/Source/RigVM/Private/RigVMCore/RigVM.cpp
--- a/Engine/Plugins/Runtime/RigVM/Source/RigVM/Private/RigVMCore/RigVM.cpp
+++ b/Engine/Plugins/Runtime/RigVM/Source/RigVM/Private/RigVMCore/RigVM.cpp
@@ -1483,6 +1483,7 @@ bool URigVM::InitializeInstance(FRigVMExtendedExecuteContext& Context, bool bCop
 #endif // WITH_EDITOR
 	}
 
+	UE::TScopeLock CreateLocalizedRegistryScopeLock(CreateLocalizedRegistryMutex);
 	CacheMemoryHandlesIfRequired(Context);
 
 	return true;
@@ -2320,6 +2321,8 @@ void URigVM::SetPropertyValueFromString(FRigVMExtendedExecuteContext& Context, c
 
 FRigVMRegistry_NoLock* URigVM::GetLocalizedRegistry() const
 {
+	UE::TScopeLock CreateLocalizedRegistryScopeLock(CreateLocalizedRegistryMutex);
+	
 	if(LocalizedRegistry.IsValid())
 	{
 		return LocalizedRegistry.Get();
@@ -2422,14 +2425,23 @@ bool URigVM::CreateLocalizedRegistryIfRequired(bool bForce)
 
 	// for dispatch functions we need to register all of the argument types for the given permutations
 	// otherwise we won't be able to retrieve the permutation from string again later.
-	for(const FRigVMFunction* Function : FunctionsStorage)
 	{
-		if(Function && Function->Factory && Function->OwnerRegistry)
+		FRigVMRegistryReadLock ReadLock;
+		for(const FRigVMFunction* Function : FunctionsStorage)
 		{
-			for(const TRigVMTypeIndex& ArgumentTypeIndex :  Function->GetArgumentTypeIndices())
+			if(Function && Function->Factory)
 			{
-				const FRigVMTemplateArgumentType& Type = Function->OwnerRegistry->GetType_NoLock(ArgumentTypeIndex);
-				TemporaryRegistry->FindOrAddType_NoLock(Type);
+				TArray<TRigVMTypeIndex> ArgumentTypeIndices;
+				Function->GetArgumentTypeIndices_NoLock(ReadLock, ArgumentTypeIndices);
+			
+				for(const TRigVMTypeIndex& ArgumentTypeIndex : ArgumentTypeIndices)
+				{
+					if (ArgumentTypeIndex != INDEX_NONE)
+					{
+						const FRigVMTemplateArgumentType& Type = ReadLock->GetType_NoLock(ArgumentTypeIndex);
+						TemporaryRegistry->FindOrAddType_NoLock(Type);
+					}
+				}
 			}
 		}
 	}
@@ -2463,39 +2475,45 @@ bool URigVM::CreateLocalizedRegistryIfRequired(bool bForce)
 			return false;
 		}
 
+		const TArray<const FRigVMFunction*>& OriginalFunctionsForFactory = Pair.Value; 
+
 		// flatten the argument infos to the permutations used only
 		TArray<FRigVMTemplateArgumentInfo> TemporaryArgumentInfos;
 		{
-			FRigVMRegistryReadLock ReadLock;const TArray<FRigVMTemplateArgumentInfo> FlattenedArgumentInfos = Factory->GetTemplate_NoLock(ReadLock)->GetFlattenedArgumentInfosForFunctions_NoLock(Pair.Value, ReadLock);
+			FRigVMRegistryReadLock ReadLock;
+			const TArray<FRigVMTemplateArgumentInfo> FlattenedArgumentInfos = Factory->GetTemplate_NoLock(ReadLock)->GetFlattenedArgumentInfosForFunctions_NoLock(OriginalFunctionsForFactory, ReadLock);
 			TemporaryArgumentInfos.Reserve(FlattenedArgumentInfos.Num());
 
 			for(const FRigVMTemplateArgumentInfo& ArgumentInfo : FlattenedArgumentInfos)
 			{
-				TArray<TRigVMTypeIndex> TypeIndices = ArgumentInfo.GetArgument_NoLock(ReadLock).TypeIndices;
-				for(int32 TypeIndex = 0; TypeIndex < TypeIndices.Num(); TypeIndex++)
+				const TArray<TRigVMTypeIndex>& OriginalTypeIndices = ArgumentInfo.GetArgument_NoLock(ReadLock).TypeIndices;
+				TArray<TRigVMTypeIndex> TemporaryTypeIndices;
+				TemporaryTypeIndices.Reserve(OriginalTypeIndices.Num());
+				
+				for(int32 TypeIndex = 0; TypeIndex < OriginalTypeIndices.Num(); TypeIndex++)
 				{
-					const FRigVMTemplateArgumentType Type = ReadLock.GetRegistry().GetType_NoLock(TypeIndices[TypeIndex]);
-					if (!Type.IsValid())
+					const FRigVMTemplateArgumentType OriginalType = ReadLock.GetRegistry().GetType_NoLock(OriginalTypeIndices[TypeIndex]);
+					if (!OriginalType.IsValid())
 					{
 						UE_LOG(LogRigVM, Error, TEXT("Cannot register dispatch dispatch '%s' for VM '%s'. Invalid types found.\nCannot localize registry."), *Factory->GetFactoryName().ToString(), *GetPathName());
 						return false;
 					}
 
-					TypeIndices[TypeIndex] = TemporaryRegistry->FindOrAddType_NoLock(Type);
-					if (TypeIndices[TypeIndex] == INDEX_NONE)
+					TemporaryTypeIndices.Add(TemporaryRegistry->FindOrAddType_NoLock(OriginalType));
+					if (TemporaryTypeIndices.Last() == INDEX_NONE)
 					{
-						UE_LOG(LogRigVM, Error, TEXT("Cannot register type '%s' in dispatch '%s' for VM '%s'\nCannot localize registry."), *Type.CPPType.ToString(), *Factory->GetFactoryName().ToString(), *GetPathName());
+						UE_LOG(LogRigVM, Error, TEXT("Cannot register type '%s' in dispatch '%s' for VM '%s'\nCannot localize registry."), *OriginalType.CPPType.ToString(), *Factory->GetFactoryName().ToString(), *GetPathName());
 						return false;
 					}
 
-					const FRigVMTemplateArgumentType TemporaryType = TemporaryRegistry->GetType_NoLock(TypeIndices[TypeIndex]);
-					if(Type != TemporaryType)
+					const FRigVMTemplateArgumentType TemporaryType = TemporaryRegistry->GetType_NoLock(TemporaryTypeIndices.Last());
+					if(TemporaryType != OriginalType)
 					{
-						UE_LOG(LogRigVM, Error, TEXT("Registered type '%s' in dispatch '%s' does not match ('%s') for VM '%s'\nCannot localize registry."), *Type.CPPType.ToString(), *Factory->GetFactoryName().ToString(), *TemporaryType.CPPType.ToString(), *GetPathName());
+						UE_LOG(LogRigVM, Error, TEXT("Registered type '%s' in dispatch '%s' does not match ('%s') for VM '%s'\nCannot localize registry."), *OriginalType.CPPType.ToString(), *Factory->GetFactoryName().ToString(), *TemporaryType.CPPType.ToString(), *GetPathName());
 						return false;
 					}
 				}
-				TemporaryArgumentInfos.Add(FRigVMTemplateArgumentInfo::MakeFlattenedFromTypeIndices(ArgumentInfo.Name, ArgumentInfo.Direction, TypeIndices));
+				TemporaryArgumentInfos.Add(FRigVMTemplateArgumentInfo::MakeFlattenedFromTypeIndices(ArgumentInfo.Name, ArgumentInfo.Direction, TemporaryTypeIndices));
 			}
 		}
 
@@ -2508,15 +2526,43 @@ bool URigVM::CreateLocalizedRegistryIfRequired(bool bForce)
 		{
 			(void)TemporaryFactory->UpdateArgumentNameCache_NoLock(Template->NumArguments(), TemporaryHandle);
 
+			TArray<const FRigVMFunction*> TemporaryFunctionsForFactory;
+			TemporaryFunctionsForFactory.Reserve(OriginalFunctionsForFactory.Num());
+			
 			for (int32 PermutationIndex = 0; PermutationIndex < Template->NumPermutations_NoLock(TemporaryHandle); PermutationIndex++)
 			{
-				if (!Template->GetOrCreatePermutation_NoLock(PermutationIndex, TemporaryHandle))
+				const FRigVMFunction* TemporaryPermutation = Template->GetOrCreatePermutation_NoLock(PermutationIndex, TemporaryHandle);
+				if (!TemporaryPermutation)
 				{
 					FRigVMTemplateTypeMap Types = Template->GetTypesForPermutation_NoLock(PermutationIndex, TemporaryHandle);
 					const FString TypesString = Template->GetStringFromArgumentTypes(Types, TemporaryHandle);
 					UE_LOG(LogRigVM, Error, TEXT("Cannot create permutation [%03d] of template '%s' types '%s' for VM '%s'\nCannot localize registry."), PermutationIndex, *Template->GetNotation().ToString(), *TypesString, *GetPathName());
 					return false;
 				}
+
+				TemporaryFunctionsForFactory.Add(TemporaryPermutation);
+			}
+
+			for (const FRigVMFunction* OriginalFunction : OriginalFunctionsForFactory)
+			{
+				int32 TemporaryFunctionIndex = TemporaryFunctionsForFactory.IndexOfByPredicate([OriginalFunction](const FRigVMFunction* TemporaryFunction)
+				{
+					return OriginalFunction->Name == TemporaryFunction->Name;
+				});
+
+				if (TemporaryFunctionIndex == INDEX_NONE)
+				{
+					for (int32 OriginalFunctionIndex = 0; OriginalFunctionIndex < OriginalFunctionsForFactory.Num(); OriginalFunctionIndex++)
+					{
+						UE_LOG(LogRigVM, Error, TEXT("Original function [%03d] '%s' of template '%s' for VM '%s'."), OriginalFunctionIndex, *OriginalFunctionsForFactory[OriginalFunctionIndex]->Name, *Template->GetNotation().ToString(), *GetPathName());
+					}
+					for (TemporaryFunctionIndex = 0; TemporaryFunctionIndex < TemporaryFunctionsForFactory.Num(); TemporaryFunctionIndex++)
+					{
+						UE_LOG(LogRigVM, Error, TEXT("Temporary function [%03d] '%s' of template '%s' for VM '%s'."), TemporaryFunctionIndex, *TemporaryFunctionsForFactory[TemporaryFunctionIndex]->Name, *Template->GetNotation().ToString(), *GetPathName());
+					}
+					UE_LOG(LogRigVM, Error, TEXT("Missing expected permutation '%s' of template '%s' for VM '%s'.\nCannot localize registry."), *OriginalFunction->Name, *Template->GetNotation().ToString(), *GetPathName());
+					return false;
+				}
 			}
 		}
 	}
diff --git a/Engine/Plugins/Runtime/RigVM/Source/RigVM/Private/RigVMCore/RigVMFunction.cpp b/Engine/Plugins/Runtime/RigVM/Source/RigVM/Private/RigVMCore/RigVMFunction.cpp
--- a/Engine/Plugins/Runtime/RigVM/Source/RigVM/Private/RigVMCore/RigVMFunction.cpp
+++ b/Engine/Plugins/Runtime/RigVM/Source/RigVM/Private/RigVMCore/RigVMFunction.cpp
@@ -133,41 +133,67 @@ const TArray<TRigVMTypeIndex>& FRigVMFunction::GetArgumentTypeIndices_NoLock(FRi
 	
 	if(ArgumentTypeIndices.IsEmpty() && !Arguments.IsEmpty())
 	{
-		if(Struct)
+		GetArgumentTypeIndices_NoLock(InRegistry, ArgumentTypeIndices);
+	}
+	return ArgumentTypeIndices;
+}
+
+void FRigVMFunction::GetArgumentTypeIndices_NoLock(FRigVMRegistryHandle& InRegistry, TArray<TRigVMTypeIndex>& OutTypeIndices) const
+{
+	OutTypeIndices.Reset();
+	
+	if(Struct)
+	{
+		for(const FRigVMFunctionArgument& Argument : Arguments)
 		{
-			for(const FRigVMFunctionArgument& Argument : Arguments)
+			if(const FProperty* Property = Struct->FindPropertyByName(Argument.Name))
 			{
-				if(const FProperty* Property = Struct->FindPropertyByName(Argument.Name))
-				{
-					FName CPPType = NAME_None;
-					UObject* CPPTypeObject = nullptr;
-					RigVMPropertyUtils::GetTypeFromProperty(Property, CPPType, CPPTypeObject);
+				FName CPPType = NAME_None;
+				UObject* CPPTypeObject = nullptr;
+				RigVMPropertyUtils::GetTypeFromProperty(Property, CPPType, CPPTypeObject);
 
-					const FRigVMTemplateArgumentType Type(CPPType, CPPTypeObject);
-					ArgumentTypeIndices.Add(InRegistry->FindOrAddType_NoLock(Type));
-				}
+				const FRigVMTemplateArgumentType Type(CPPType, CPPTypeObject);
+				OutTypeIndices.Add(InRegistry->FindOrAddType_NoLock(Type));
+			}
+			else
+			{
+				OutTypeIndices.Add(INDEX_NONE);
 			}
 		}
-		else if(InRegistry->GetTemplates_NoLock().IsValidIndex(TemplateIndex))
-		{
-			const FRigVMTemplate* Template = &InRegistry->GetTemplates_NoLock()[TemplateIndex];
+	}
+	else if(InRegistry->GetTemplates_NoLock().IsValidIndex(TemplateIndex))
+	{
+		const FRigVMTemplate* Template = &InRegistry->GetTemplates_NoLock()[TemplateIndex];
 			
-			check(PermutationIndex != INDEX_NONE);
-			check(PermutationIndex == Template->FindPermutation_NoLock(this, InRegistry));
+		check(PermutationIndex != INDEX_NONE);
+		check(PermutationIndex == Template->FindPermutation_NoLock(this, InRegistry));
 
-			for(const FRigVMFunctionArgument& FunctionArgument : Arguments)
+		for(const FRigVMFunctionArgument& FunctionArgument : Arguments)
+		{
+			TRigVMTypeIndex TypeIndex = InRegistry->GetTypeIndexFromCPPType_NoLock(FunctionArgument.Type);
+			
+			// fall back on the template just in case
+			if(TypeIndex == INDEX_NONE)
 			{
-				const FRigVMTemplateArgument* TemplateArgument = Template->FindArgument(FunctionArgument.Name);
-				check(TemplateArgument);
-				ArgumentTypeIndices.Add(TemplateArgument->GetTypeIndex_NoLock(PermutationIndex, InRegistry));
+				if (const FRigVMTemplateArgument* TemplateArgument = Template->FindArgument(FunctionArgument.Name))
+				{
+					TypeIndex = TemplateArgument->GetTypeIndex_NoLock(PermutationIndex, InRegistry);
+				}
 			}
-		}
-		else
-		{
-			// checkNoEntry();
+#if WITH_EDITOR
+			if(TypeIndex != INDEX_NONE)
+			{
+				const FRigVMTemplateArgumentType& Type = InRegistry->GetType_NoLock(TypeIndex); 
+				check(FunctionArgument.Type == Type.CPPType.ToString());
+			}
+#endif 
+			OutTypeIndices.Add(TypeIndex);
 		}
 	}
-	return ArgumentTypeIndices;
+	else
+	{
+		// checkNoEntry();
+	}
 }
 
 const FRigVMTemplate* FRigVMFunction::GetTemplate() const
diff --git a/Engine/Plugins/Runtime/RigVM/Source/RigVM/Private/RigVMCore/RigVMTemplate.cpp b/Engine/Plugins/Runtime/RigVM/Source/RigVM/Private/RigVMCore/RigVMTemplate.cpp
--- a/Engine/Plugins/Runtime/RigVM/Source/RigVM/Private/RigVMCore/RigVMTemplate.cpp
+++ b/Engine/Plugins/Runtime/RigVM/Source/RigVM/Private/RigVMCore/RigVMTemplate.cpp
@@ -2815,7 +2815,10 @@ TArray<FRigVMTemplateArgumentInfo> FRigVMTemplate::GetFlattenedArgumentInfosForF
 		check(Function->OwnerRegistry == this->OwnerRegistry);
 		check(Function->TemplateIndex == this->Index);
 
-		const TArray<TRigVMTypeIndex>& TypeIndicesForFunction = Function->GetArgumentTypeIndices_NoLock(InRegistry);
+		TArray<TRigVMTypeIndex> TypeIndicesForFunction;
+		Function->GetArgumentTypeIndices_NoLock(InRegistry, TypeIndicesForFunction);
+		check(TypeIndicesForFunction.Num() == Function->Arguments.Num());
+
 		for (int32 ArgumentIndex = 0; ArgumentIndex < Function->Arguments.Num(); ArgumentIndex++)
 		{
 			TypeIndices.FindOrAdd(Function->Arguments[ArgumentIndex].Name).Add(TypeIndicesForFunction[ArgumentIndex]);
diff --git a/Engine/Plugins/Runtime/RigVM/Source/RigVM/Public/RigVMCore/RigVM.h b/Engine/Plugins/Runtime/RigVM/Source/RigVM/Public/RigVMCore/RigVM.h
--- a/Engine/Plugins/Runtime/RigVM/Source/RigVM/Public/RigVMCore/RigVM.h
+++ b/Engine/Plugins/Runtime/RigVM/Source/RigVM/Public/RigVMCore/RigVM.h
@@ -1141,7 +1141,7 @@ protected:
 	RIGVM_API void StopProfiling(FRigVMExtendedExecuteContext& Context);
 
 	FTransactionallySafeCriticalSection ResolveFunctionsMutex;
-	FTransactionallySafeCriticalSection CreateLocalizedRegistryMutex;
+	mutable FTransactionallySafeCriticalSection CreateLocalizedRegistryMutex;
 
 	friend class URigVMCompiler;
 	friend struct FRigVMCompilerWorkData;
diff --git a/Engine/Plugins/Runtime/RigVM/Source/RigVM/Public/RigVMCore/RigVMFunction.h b/Engine/Plugins/Runtime/RigVM/Source/RigVM/Public/RigVMCore/RigVMFunction.h
--- a/Engine/Plugins/Runtime/RigVM/Source/RigVM/Public/RigVMCore/RigVMFunction.h
+++ b/Engine/Plugins/Runtime/RigVM/Source/RigVM/Public/RigVMCore/RigVMFunction.h
@@ -183,6 +183,7 @@ struct FRigVMFunction
 	const TArray<FRigVMFunctionArgument>& GetArguments() const { return Arguments; }
 	RIGVM_API const TArray<TRigVMTypeIndex>& GetArgumentTypeIndices() const;
 	RIGVM_API const TArray<TRigVMTypeIndex>& GetArgumentTypeIndices_NoLock(FRigVMRegistryHandle& InRegistry) const;
+	RIGVM_API void GetArgumentTypeIndices_NoLock(FRigVMRegistryHandle& InRegistry, TArray<TRigVMTypeIndex>& OutTypeIndices) const;
 	RIGVM_API const FRigVMTemplate* GetTemplate() const;
 	RIGVM_API const FRigVMTemplate* GetTemplate_NoLock(const FRigVMRegistryHandle& InRegistry) const;
 	RIGVM_API const UScriptStruct* GetExecuteContextStruct() const;
