RigVMTypeUtils::UserDefinedTypeFromCPPType
- Cook indeterminism fix - asset registry queries now search on-disk assets only

diff --git a/Engine/Plugins/Runtime/RigVM/Source/RigVM/Private/RigVMTypeUtils.cpp b/Engine/Plugins/Runtime/RigVM/Source/RigVM/Private/RigVMTypeUtils.cpp
--- a/Engine/Plugins/Runtime/RigVM/Source/RigVM/Private/RigVMTypeUtils.cpp
+++ b/Engine/Plugins/Runtime/RigVM/Source/RigVM/Private/RigVMTypeUtils.cpp
@@ -1,4 +1,4 @@
-﻿// Copyright Epic Games, Inc. All Rights Reserved.
+// Copyright Epic Games, Inc. All Rights Reserved.
 
 /*=============================================================================
 	RigVMTypeUtils.h: Module implementation.
@@ -67,6 +67,11 @@ UObject* RigVMTypeUtils::UserDefinedTypeFromCPPType(FString& InOutCPPType, const
 	{
 		const IAssetRegistry& AssetRegistry = FModuleManager::LoadModuleChecked<FAssetRegistryModule>("AssetRegistry").Get();
 
+		TArray<FAssetData> AssetDataList;
+		FARFilter AssetFilter;
+		AssetFilter.bRecursiveClasses = true;
+		AssetFilter.bIncludeOnlyOnDiskAssets = IsRunningCookCommandlet();	// When cooking, only enumerate on-disk assets to ensure deterministic results
+
 		InOutCPPType = OriginalTypeName;
 		while(IsArrayType(InOutCPPType))
 		{
@@ -77,8 +82,8 @@ UObject* RigVMTypeUtils::UserDefinedTypeFromCPPType(FString& InOutCPPType, const
 		{
 			static const FLazyName GuidTag(GET_MEMBER_NAME_CHECKED(UUserDefinedStruct, Guid));
 
-			TArray<FAssetData> AssetDataList;
-			AssetRegistry.GetAssetsByClass(UUserDefinedStruct::StaticClass()->GetClassPathName(), AssetDataList, true);
+			AssetFilter.ClassPaths = { UUserDefinedStruct::StaticClass()->GetClassPathName() };
+			AssetRegistry.GetAssets(AssetFilter, AssetDataList);
 
 			// Temporary workaround:
 			// User Defined Struct has been moved from Engine to CoreUObject at CL34495787 for UE-216472
@@ -89,7 +94,8 @@ UObject* RigVMTypeUtils::UserDefinedTypeFromCPPType(FString& InOutCPPType, const
 			TArray<FString> OldPathNames = FLinkerLoad::FindPreviousPathNamesForClass(UUserDefinedStruct::StaticClass()->GetClassPathName().ToString(), false);
 			for (const FString& OldPathName : OldPathNames)
 			{
-				AssetRegistry.GetAssetsByClass( FTopLevelAssetPath(OldPathName), AssetDataList, true);
+				AssetFilter.ClassPaths = { FTopLevelAssetPath(OldPathName) };
+				AssetRegistry.GetAssets(AssetFilter, AssetDataList);
 			}
 
 			// first pass - try to find it using the tag
@@ -138,8 +144,8 @@ UObject* RigVMTypeUtils::UserDefinedTypeFromCPPType(FString& InOutCPPType, const
 		}
 		else
 		{
-			TArray<FAssetData> AssetDataList;
-			AssetRegistry.GetAssetsByClass(UUserDefinedEnum::StaticClass()->GetClassPathName(), AssetDataList, true);
+			AssetFilter.ClassPaths = { UUserDefinedEnum::StaticClass()->GetClassPathName() };
+			AssetRegistry.GetAssets(AssetFilter, AssetDataList);
 
 			for(const FAssetData& AssetData : AssetDataList)
 			{
