Control Rig: The editor now highlights variables that are used in graphs, related context menu now offers an option to remove unused variables

diff --git a/Engine/Plugins/Runtime/RigVM/Source/RigVMDeveloper/Private/RigVMModel/RigVMClient.cpp b/Engine/Plugins/Runtime/RigVM/Source/RigVMDeveloper/Private/RigVMModel/RigVMClient.cpp
--- a/Engine/Plugins/Runtime/RigVM/Source/RigVMDeveloper/Private/RigVMModel/RigVMClient.cpp
+++ b/Engine/Plugins/Runtime/RigVM/Source/RigVMDeveloper/Private/RigVMModel/RigVMClient.cpp
@@ -1727,7 +1727,8 @@ void FRigVMClient::HandleGraphModifiedEvent(ERigVMGraphNotifType InNotifType, UR
 			break;
 		}
 	}
-	
+
+	PostGraphModifiedDelegate.Broadcast(InNotifType, InGraph, InSubject);
 #endif
 }
 
diff --git a/Engine/Plugins/Runtime/RigVM/Source/RigVMDeveloper/Public/RigVMModel/RigVMClient.h b/Engine/Plugins/Runtime/RigVM/Source/RigVMDeveloper/Public/RigVMModel/RigVMClient.h
--- a/Engine/Plugins/Runtime/RigVM/Source/RigVMDeveloper/Public/RigVMModel/RigVMClient.h
+++ b/Engine/Plugins/Runtime/RigVM/Source/RigVMDeveloper/Public/RigVMModel/RigVMClient.h
@@ -201,6 +201,9 @@ public:
 	virtual TObjectPtr<URigVMGraph> CreateContainedGraphModel(URigVMCollapseNode* CollapseNode, const FName& Name) = 0;
 };
 
+#if WITH_EDITOR
+DECLARE_MULTICAST_DELEGATE_ThreeParams(FRigVMClientPostGraphModified, ERigVMGraphNotifType /** InNotifType */, URigVMGraph* /** InGraph */, UObject* /** InSubject */);
+#endif // WITH_EDITOR
 
 // A management struct containing graphs and controllers.
 USTRUCT()
@@ -332,6 +335,9 @@ public:
 
 #if WITH_EDITOR
 	UE_API bool UpgradeAllOccurencesOfNodes(const TArray<URigVMNode*>& InNodesToUpgrade,  UObject* InOuter, bool bSetupUndoRedo = true);
+
+	/** Returns a delegate raised after a graph in this client was modified. Useful to track changes without keeping track of graphs. */
+	FRigVMClientPostGraphModified& GetPostGraphModified() { return PostGraphModifiedDelegate; }
 #endif
 
 	static constexpr TCHAR RigVMModelPrefix[] = TEXT("RigVMModel");
@@ -393,6 +399,11 @@ private:
 
 	TWeakInterfacePtr<IRigVMClientExternalModelHost> ExternalModelHost;
 
+#if WITH_EDITOR
+	/** Returns a delegate raised after any graph in this client was modified and after the client adopted these changes */
+	FRigVMClientPostGraphModified PostGraphModifiedDelegate;
+#endif // WITH_EDITOR
+
 	friend class UEngineTestClientHost;
 	friend class URigVMBlueprint;
 	friend class IRigVMAssetInterface;
diff --git a/Engine/Plugins/Runtime/RigVM/Source/RigVMEditor/Private/Editor/RigVMEdGraphNodeRegistry.cpp b/Engine/Plugins/Runtime/RigVM/Source/RigVMEditor/Private/Editor/RigVMEdGraphNodeRegistry.cpp
new file mode 100644
--- /dev/null
+++ b/Engine/Plugins/Runtime/RigVM/Source/RigVMEditor/Private/Editor/RigVMEdGraphNodeRegistry.cpp
@@ -0,0 +1,122 @@
+// Copyright Epic Games, Inc. All Rights Reserved.
+
+#include "Editor/RigVMEdGraphNodeRegistry.h"
+
+#include "EdGraph/RigVMEdGraph.h"
+#include "Editor.h"
+#include "RigVMAsset.h"
+#include "RigVMModel/RigVMNode.h"
+
+namespace UE::RigVMEditor
+{
+	TMap<FName, TWeakPtr<FRigVMEdGraphNodeRegistry>> FRigVMEdGraphNodeRegistry::TypeIDToRegistryMap;
+
+	TSharedRef<FRigVMEdGraphNodeRegistry> FRigVMEdGraphNodeRegistry::GetOrCreateRegistry(const TScriptInterface<IRigVMAssetInterface>& InRigVMAssetInterface, UClass* InRigVMNodeClass)
+	{		
+		if (!ensureMsgf(InRigVMNodeClass->IsChildOf(URigVMNode::StaticClass()), TEXT("Node class is not a RigVMNode type, creating empty Rig VM Node Registry")) ||
+			!ensureMsgf(InRigVMAssetInterface, TEXT("Rig VM Asset Interface is invalid, creating empty Rig VM Node Registry")))
+		{
+			return MakeShared<FRigVMEdGraphNodeRegistry>();
+		}
+
+		// Lazily cleanup existing registries
+		for (auto It = TypeIDToRegistryMap.CreateIterator(); It; ++It)
+		{
+			if (!(*It).Value.IsValid())
+			{
+				It.RemoveCurrent();
+			}
+		}
+
+		// Reuse registries where possible
+		const FName TypeID = *(InRigVMNodeClass->GetFName().ToString() + TEXT("_") + InRigVMAssetInterface->GetObject()->GetFName().ToString());
+
+		const TWeakPtr<FRigVMEdGraphNodeRegistry>* InstancePtr = TypeIDToRegistryMap.Find(TypeID);
+		if (InstancePtr && InstancePtr->IsValid())
+		{
+			return (*InstancePtr).Pin().ToSharedRef();
+		}
+		else
+		{
+			const TSharedRef<FRigVMEdGraphNodeRegistry> NewRegistry = MakeShared<FRigVMEdGraphNodeRegistry>(InRigVMAssetInterface, InRigVMNodeClass);
+			NewRegistry->Initialize();
+			NewRegistry->Update();
+
+			TypeIDToRegistryMap.Add(TypeID, NewRegistry);
+
+			return NewRegistry;
+		}
+	}
+
+	FRigVMEdGraphNodeRegistry::FRigVMEdGraphNodeRegistry(const TScriptInterface<IRigVMAssetInterface>& InRigVMAssetInterface, UClass* InRigVMNodeClass)
+		: WeakAssetInterface(InRigVMAssetInterface)
+		, RigVMNodeClass(InRigVMNodeClass)
+	{}
+
+	void FRigVMEdGraphNodeRegistry::Initialize()
+	{
+		FRigVMClient* ClientPtr = WeakAssetInterface.IsValid() ? WeakAssetInterface->GetRigVMClient() : nullptr;
+		if (!ensureMsgf(WeakAssetInterface.IsValid() && ClientPtr, TEXT("Cannot update Rig VM Node Registry, Rig VM Asset Interface or Client is invalid")))
+		{
+			return;
+		}
+
+		ClientPtr->GetPostGraphModified().AddSP(this, &FRigVMEdGraphNodeRegistry::PostGraphModified);
+	}
+
+	void FRigVMEdGraphNodeRegistry::Update()
+	{
+		TRACE_CPUPROFILER_EVENT_SCOPE(FRigVMEdGraphNodeRegistry::Update);
+
+		WeakEdGraphNodes.Reset();
+
+		TArray<UEdGraph*> EdGraphs;
+		WeakAssetInterface->GetAllEdGraphs(EdGraphs);
+
+		for (UEdGraph* Graph : EdGraphs)
+		{
+			if (URigVMEdGraph* RigVMEdGraph = Cast<URigVMEdGraph>(Graph))
+			{
+				TArray<URigVMEdGraphNode*> EdGraphNodes;
+				RigVMEdGraph->GetNodesOfClass<URigVMEdGraphNode>(EdGraphNodes);
+
+				for (URigVMEdGraphNode* EdGraphNode : EdGraphNodes)
+				{
+					if (EdGraphNode->GetModelNode() &&
+						EdGraphNode->GetModelNode()->GetClass() == RigVMNodeClass)
+					{
+						WeakEdGraphNodes.Add(EdGraphNode);
+					}
+				}
+			}
+		}
+
+		OnPostRegistryUpdated.Broadcast();
+	}
+
+	void FRigVMEdGraphNodeRegistry::PostGraphModified(ERigVMGraphNotifType InNotifType, URigVMGraph* InGraph, UObject* InSubject)
+	{
+		if (InNotifType == ERigVMGraphNotifType::NodeAdded ||
+			InNotifType == ERigVMGraphNotifType::NodeRemoved)
+		{
+			RequestUpdate();
+		}
+	}
+
+	void FRigVMEdGraphNodeRegistry::RequestUpdate()
+	{
+		if (!RequestUpdateTimerHandle.IsValid())
+		{
+			RequestUpdateTimerHandle = GEditor->GetTimerManager()->SetTimerForNextTick(
+				[WeakThis = AsWeak(), this]()
+				{
+					if (WeakThis.IsValid())
+					{
+						Update();
+
+						RequestUpdateTimerHandle.Invalidate();
+					}
+				});
+		}
+	}
+}
diff --git a/Engine/Plugins/Runtime/RigVM/Source/RigVMEditor/Private/Editor/RigVMEdGraphNodeRegistry.h b/Engine/Plugins/Runtime/RigVM/Source/RigVMEditor/Private/Editor/RigVMEdGraphNodeRegistry.h
new file mode 100644
--- /dev/null
+++ b/Engine/Plugins/Runtime/RigVM/Source/RigVMEditor/Private/Editor/RigVMEdGraphNodeRegistry.h
@@ -0,0 +1,77 @@
+// Copyright Epic Games, Inc. All Rights Reserved.
+
+#pragma once
+
+#include "Engine/TimerHandle.h"
+#include "Templates/SharedPointer.h"
+#include "UObject/ScriptInterface.h"
+#include "UObject/WeakInterfacePtr.h"
+#include "UObject/WeakObjectPtr.h"
+
+class IRigVMAssetInterface;
+class UClass;
+class URigVMEdGraphNode;
+class URigVMGraph;
+enum class ERigVMGraphNotifType : uint8;
+
+namespace UE::RigVMEditor
+{
+	/** Template to track existence of Rig VM Ed Graph Nodes for nodes of a specific Rig VM Node type in editor. */
+	class FRigVMEdGraphNodeRegistry
+		: public TSharedFromThis<FRigVMEdGraphNodeRegistry>
+	{
+		// Allow MakeShared with private constructor
+		friend class SharedPointerInternals::TIntrusiveReferenceController<FRigVMEdGraphNodeRegistry, ESPMode::ThreadSafe>;
+
+	public:
+		FRigVMEdGraphNodeRegistry() = default;
+		
+		/** 
+		 * Creates a registry instance. 
+		 * Note, The registry will be destroyed when the all references are released.
+		 *
+		 * @param InRigVMAssetInterface			The Rig VM Asset Interface for which this registry is constructed 
+		 * @param InRigVMNodeClass				The Rig VM Node Class for which Ed Graph nodes are held
+		 * 
+		 * @return								The reistry instance
+		 */
+		static TSharedRef<FRigVMEdGraphNodeRegistry> GetOrCreateRegistry(const TScriptInterface<IRigVMAssetInterface>& InRigVMAssetInterface, UClass* InRigVMNodeClass);
+
+		/** Returns the currently registered ed graph nodes which correspond to the node type */
+		const TArray<TWeakObjectPtr<URigVMEdGraphNode>>& GetWeakEdGrapNodes() const { return WeakEdGraphNodes; }
+
+		/** Delegate broadcasted when the registry was updated, raised one tick after changes occured */
+		FSimpleMulticastDelegate OnPostRegistryUpdated;
+
+	private:
+		/** Private constructor, hidden on purpose. Use GetOrCreateRegistry to create instances. */
+		FRigVMEdGraphNodeRegistry(const TScriptInterface<IRigVMAssetInterface>& InRigVMAssetInterface, UClass* InRigVMNodeClass);
+
+		/** Initializes the registry */
+		void Initialize();
+
+		/** Updates the registry from existing nodes */
+		void Update();
+
+		/** Called after a graph was modified */
+		void PostGraphModified(ERigVMGraphNotifType InNotifType, URigVMGraph* InGraph, UObject* InSubject);
+
+		/** Requests to Update the registry */
+		void RequestUpdate();
+
+		/** The currently registered nodes */
+		TArray<TWeakObjectPtr<URigVMEdGraphNode>> WeakEdGraphNodes;
+
+		/** The asset interface this registry refers to */
+		TWeakInterfacePtr<IRigVMAssetInterface> WeakAssetInterface;
+
+		/** The Rig VM Node type that is tracked in this registry */
+		UClass* RigVMNodeClass = nullptr;
+
+		/** Timer handle to notify when the registry was fully updated */
+		FTimerHandle RequestUpdateTimerHandle;
+
+		/** Map of registry IDs with their unique registry, useful to reuse registries of same type and asset */
+		static TMap<FName, TWeakPtr<FRigVMEdGraphNodeRegistry>> TypeIDToRegistryMap;
+	};
+}
diff --git a/Engine/Plugins/Runtime/RigVM/Source/RigVMEditor/Private/Widgets/SRigVMEditorGraphExplorer.cpp b/Engine/Plugins/Runtime/RigVM/Source/RigVMEditor/Private/Widgets/SRigVMEditorGraphExplorer.cpp
--- a/Engine/Plugins/Runtime/RigVM/Source/RigVMEditor/Private/Widgets/SRigVMEditorGraphExplorer.cpp
+++ b/Engine/Plugins/Runtime/RigVM/Source/RigVMEditor/Private/Widgets/SRigVMEditorGraphExplorer.cpp
@@ -2,8 +2,11 @@
 
 #include "Widgets/SRigVMEditorGraphExplorer.h"
 
+#include "Editor/RigVMEdGraphNodeRegistry.h"
 #include "Editor/RigVMNewEditor.h"
 #include "EdGraph/RigVMEdGraphSchema.h"
+#include "Framework/Notifications/NotificationManager.h"
+#include "Widgets/Notifications/SNotificationList.h"
 #include "Widgets/SRigVMEditorGraphExplorerTreeView.h"
 
 #include "Framework/Commands/GenericCommands.h"
@@ -201,6 +204,7 @@ void FRigVMEditorGraphExplorerCommands::RegisterCommands()
 	UI_COMMAND(PasteFunction, "Paste Function", "Pastes the function.", EUserInterfaceActionType::Button, FInputChord());
 	UI_COMMAND(PasteVariable, "Paste Variable", "Pastes the variable.", EUserInterfaceActionType::Button, FInputChord());
 	UI_COMMAND(PasteLocalVariable, "Paste Local Variable", "Pastes the local variable.", EUserInterfaceActionType::Button, FInputChord());
+	UI_COMMAND(RemoveUnusedVariables, "Remove Unused Variables", "Removes variables not used in any graph.", EUserInterfaceActionType::Button, FInputChord());
 }
 
 
@@ -209,27 +213,24 @@ void SRigVMEditorGraphExplorer::Construct(const FArguments& InArgs, TWeakPtr<IRi
 	bNeedsRefresh = false;
 	RigVMEditor = InRigVMEditor;
 
-	TSharedPtr<IRigVMEditor> Editor = RigVMEditor.Pin();
-	if (Editor)
-	{
-		Editor->OnRefresh().AddRaw(this, &SRigVMEditorGraphExplorer::Refresh);
-	}
-
 	RegisterCommands();
 
 	CreateWidgets();
 
 	LastPinType.ResetToDefaults();
 	LastPinType.PinCategory = TEXT("bool");
-}
 
-
-SRigVMEditorGraphExplorer::~SRigVMEditorGraphExplorer()
-{
 	TSharedPtr<IRigVMEditor> Editor = RigVMEditor.Pin();
 	if (Editor)
 	{
-		Editor->OnRefresh().RemoveAll(this);
+		Editor->OnRefresh().AddSP(this, &SRigVMEditorGraphExplorer::Refresh);
+
+		// Acquire a registry for variable nodes
+		if (const TScriptInterface<IRigVMAssetInterface> AssetInterface = Editor->GetRigVMAssetInterface())
+		{
+			EdGraphNodeVariableRegistry = FRigVMEdGraphNodeRegistry::GetOrCreateRegistry(AssetInterface, URigVMVariableNode::StaticClass());
+			EdGraphNodeVariableRegistry->OnPostRegistryUpdated.AddSP(this, &SRigVMEditorGraphExplorer::Refresh);
+		}
 	}
 }
 
@@ -369,6 +370,12 @@ void SRigVMEditorGraphExplorer::RegisterCommands()
 			FCanExecuteAction(),
 			FIsActionChecked(),
 			FIsActionButtonVisible::CreateSP(this, &SRigVMEditorGraphExplorer::CanFindReference) );
+
+		ToolKitCommandList->MapAction(FRigVMEditorGraphExplorerCommands::Get().RemoveUnusedVariables,
+			FExecuteAction::CreateSP(this, &SRigVMEditorGraphExplorer::OnRemoveUnusedVariables),
+			FCanExecuteAction(),
+			FIsActionChecked(),
+			FIsActionButtonVisible::CreateSP(this, &SRigVMEditorGraphExplorer::CanRemoveUnusedVariables));
 	}
 }
 
@@ -443,7 +450,7 @@ void SRigVMEditorGraphExplorer::CreateWidgets()
 	
 	Delegates.OnSelectionChanged = FRigVMGraphExplorer_OnSelectionChanged::CreateSP(this, &SRigVMEditorGraphExplorer::HandleSelectionChanged);
 	
-	SAssignNew(TreeView, SRigVMEditorGraphExplorerTreeView)
+	SAssignNew(TreeView, SRigVMEditorGraphExplorerTreeView, SharedThis(this))
 		.RigTreeDelegates(Delegates)
 	;
 
@@ -1263,6 +1270,7 @@ TSharedPtr<SWidget> SRigVMEditorGraphExplorer::OnContextMenuOpening()
 				MenuBuilder.AddMenuEntry(FGenericCommands::Get().Copy);
 				MenuBuilder.AddMenuEntry(FGenericCommands::Get().Duplicate);
 				MenuBuilder.AddMenuEntry(FGenericCommands::Get().Delete);
+				MenuBuilder.AddMenuEntry(FRigVMEditorGraphExplorerCommands::Get().RemoveUnusedVariables);
 				MenuBuilder.EndSection();
 				break;
 			}
@@ -1978,7 +1986,7 @@ bool SRigVMEditorGraphExplorer::OnRenameVariable(const FRigVMExplorerElementKey&
 		return false;
 	}
 
-	const FScopedTransaction Transaction( LOCTEXT( "RenameVariable", "Rename Variable" ) );
+	const FScopedTransaction Transaction( LOCTEXT( "RenameVariable", "Rename Variable" ));
 
 	FRigVMAssetInterfacePtr Blueprint = Editor->GetRigVMAssetInterface();
 
@@ -2076,6 +2084,72 @@ bool SRigVMEditorGraphExplorer::CanFindReference() const
 	return true;
 }
 
+void SRigVMEditorGraphExplorer::OnRemoveUnusedVariables()
+{
+	 const TScriptInterface<IRigVMAssetInterface> Asset = RigVMEditor.IsValid() ? RigVMEditor.Pin()->GetRigVMAssetInterface() : nullptr;
+	if (!TreeView.IsValid() ||
+		!Asset)
+	{
+		return;
+	}
+	
+	const TArray<TWeakObjectPtr<URigVMEdGraphNode>> WeakEdGraphNodes = EdGraphNodeVariableRegistry->GetWeakEdGrapNodes();
+
+	TArray<FRigVMExplorerElementKey> KeysToRemove;
+	Algo::TransformIf(TreeView->GetKeys(), KeysToRemove,
+		[&WeakEdGraphNodes](const FRigVMExplorerElementKey& Key)
+		{
+			if (Key.Type == ERigVMExplorerElementType::Variable ||
+				Key.Type == ERigVMExplorerElementType::LocalVariable)
+			{
+				const bool bContainedInGraph = WeakEdGraphNodes.ContainsByPredicate(
+					[&Key](const TWeakObjectPtr<URigVMEdGraphNode>& WeakEdGraphNode)
+					{
+						const URigVMVariableNode* VariableNode = WeakEdGraphNode.IsValid() ? Cast<URigVMVariableNode>(WeakEdGraphNode->GetModelNode()) : nullptr;
+						if (VariableNode)
+						{
+							return VariableNode->GetVariableName() == Key.Name;
+						}
+
+						return false;
+					});
+
+				return !bContainedInGraph;
+			}
+
+			return false;
+		},
+		[](const FRigVMExplorerElementKey& Key)
+		{
+			return Key;
+		});
+
+	if (KeysToRemove.IsEmpty())
+	{
+		FNotificationInfo Info(LOCTEXT("AllVariablesReferencedNotification", "All variables are referenced"));
+		Info.bFireAndForget = true;
+		Info.FadeOutDuration = 5.0f;
+		Info.ExpireDuration = 5.0f;
+
+		FSlateNotificationManager::Get().AddNotification(Info);
+	}
+	else if (Asset->GetObject())
+	{
+		const FScopedTransaction RemoveUnusedVariablesTransaction(LOCTEXT("RemoveUnusedVariables", "Remove Unused Variables"), !GIsTransacting);
+
+		Asset->GetObject()->Modify();
+		for (const FRigVMExplorerElementKey& KeyToRemove : KeysToRemove)
+		{
+			Asset->RemoveMemberVariable(*KeyToRemove.Name);
+		}
+	}
+}
+
+bool SRigVMEditorGraphExplorer::CanRemoveUnusedVariables() const
+{
+	return true;
+}
+
 void SRigVMEditorGraphExplorer::HandleSelectionChanged(TSharedPtr<FRigVMEditorGraphExplorerTreeElement> Selection, ESelectInfo::Type SelectInfo)
 {
 	TreeView->RefreshTreeView();
diff --git a/Engine/Plugins/Runtime/RigVM/Source/RigVMEditor/Private/Widgets/SRigVMEditorGraphExplorerTreeView.cpp b/Engine/Plugins/Runtime/RigVM/Source/RigVMEditor/Private/Widgets/SRigVMEditorGraphExplorerTreeView.cpp
--- a/Engine/Plugins/Runtime/RigVM/Source/RigVMEditor/Private/Widgets/SRigVMEditorGraphExplorerTreeView.cpp
+++ b/Engine/Plugins/Runtime/RigVM/Source/RigVMEditor/Private/Widgets/SRigVMEditorGraphExplorerTreeView.cpp
@@ -5,6 +5,7 @@
 #include "SPinTypeSelector.h"
 #include "EdGraph/RigVMEdGraphSchema.h"
 #include "Editor/RigVMNewEditor.h"
+#include "Editor/RigVMEdGraphNodeRegistry.h"
 #include "RigVMModel/Nodes/RigVMLibraryNode.h"
 #include "Widgets/SRigVMEditorGraphExplorer.h"
 
@@ -109,6 +110,11 @@ void SRigVMEditorGraphExplorerItem::Construct(const FArguments& InArgs,
 	{
 		TArray<TSharedPtr<IPinTypeSelectorFilter>> CustomPinTypeFilters = Delegates.GetCustomPinFilters();
 		
+		// Lower text color for variables that are not used in the graph
+		const FLinearColor TextColor = InTreeView->IsVariableElementUsedInGraph(InElement) ?
+			FLinearColor::White :
+			FLinearColor::White.CopyWithNewOpacity(0.7f);
+
 		STableRow<TSharedPtr<FRigVMEditorGraphExplorerTreeElement>>::Construct(
 		   STableRow<TSharedPtr<FRigVMEditorGraphExplorerTreeElement>>::FArguments()
 		   .ShowWires(false)
@@ -121,7 +127,7 @@ void SRigVMEditorGraphExplorerItem::Construct(const FArguments& InArgs,
 				   .VAlign(VAlign_Center)
 				   .Padding(3.0f, 0.0f)
 				   [
-					   CreateTextSlotWidget(InElement->Key, InTreeView->FilterText)
+					   CreateTextSlotWidget(InElement->Key, InTreeView->FilterText, TextColor)
 				   ]
 
 				+SHorizontalBox::Slot()
@@ -252,14 +258,18 @@ TSharedRef<SWidget> SRigVMEditorGraphExplorerItem::CreateIconWidget(const FRigVM
 	}
 }
 
-TSharedRef<SWidget> SRigVMEditorGraphExplorerItem::CreateTextSlotWidget(const FRigVMExplorerElementKey& Key, const FText& InHighlightText)
+TSharedRef<SWidget> SRigVMEditorGraphExplorerItem::CreateTextSlotWidget(
+	const FRigVMExplorerElementKey& Key, 
+	const FText& InHighlightText, 
+	const FSlateColor& TextColor)
 {
 	TSharedPtr<SInlineEditableTextBlock> EditableTextElement = SNew(SInlineEditableTextBlock)
 		.Text(this, &SRigVMEditorGraphExplorerItem::GetDisplayText)
 		.HighlightText(InHighlightText)
 		.ToolTipText(this, &SRigVMEditorGraphExplorerItem::GetItemTooltip)
 		.OnVerifyTextChanged(this, &SRigVMEditorGraphExplorerItem::OnVerifyNameChanged)
-		.OnTextCommitted(this, &SRigVMEditorGraphExplorerItem::OnNameCommitted);
+		.OnTextCommitted(this, &SRigVMEditorGraphExplorerItem::OnNameCommitted)
+		.ColorAndOpacity(TextColor);
 
 	InlineRenameWidget = EditableTextElement.ToSharedRef();
 
@@ -464,8 +474,10 @@ void FRigVMEditorGraphExplorerTreeElement::RequestRename()
 	OnRenameRequested.ExecuteIfBound();
 }
 
-void SRigVMEditorGraphExplorerTreeView::Construct(const FArguments& InArgs)
+void SRigVMEditorGraphExplorerTreeView::Construct(const FArguments& InArgs, const TSharedRef<SRigVMEditorGraphExplorer>& InGraphExplorer)
 {
+	WeakGraphExplorer = InGraphExplorer;
+
 	Delegates = InArgs._RigTreeDelegates;
 	
 	STreeView<TSharedPtr<FRigVMEditorGraphExplorerTreeElement>>::FArguments SuperArgs;
@@ -480,7 +492,7 @@ void SRigVMEditorGraphExplorerTreeView::Construct(const FArguments& InArgs)
 	SuperArgs.OnMouseButtonDoubleClick(this, &SRigVMEditorGraphExplorerTreeView::OnItemDoubleClicked);
 	SuperArgs.OnSelectionChanged(FRigVMGraphExplorer_OnSelectionChanged::CreateRaw(&Delegates, &FRigVMEditorGraphExplorerTreeDelegates::HandleSelectionChanged));
 	SuperArgs.OnContextMenuOpening(FRigVMGraphExplorer_OnRequestContextMenu::CreateRaw(&Delegates, &FRigVMEditorGraphExplorerTreeDelegates::RequestContextMenu));
-	
+
 
 	STreeView<TSharedPtr<FRigVMEditorGraphExplorerTreeElement>>::Construct(SuperArgs);
 }
@@ -834,6 +846,14 @@ void SRigVMEditorGraphExplorerTreeView::HandleGetChildrenForTree(TSharedPtr<FRig
 	OutChildren = InItem.Get()->Children;
 }
 
+TArray<FRigVMExplorerElementKey> SRigVMEditorGraphExplorerTreeView::GetKeys() const
+{
+	TArray<FRigVMExplorerElementKey> Keys;
+	ElementMap.GenerateKeyArray(Keys);
+	
+	return Keys;
+}
+
 TArray<FRigVMExplorerElementKey> SRigVMEditorGraphExplorerTreeView::GetSelectedKeys() const
 {
 	TArray<FRigVMExplorerElementKey> Keys;
@@ -957,4 +977,32 @@ FReply SRigVMEditorGraphExplorerTreeView::OnAddButtonClickedOnSection(const FRig
 	return FReply::Unhandled();
 }
 
+bool SRigVMEditorGraphExplorerTreeView::IsVariableElementUsedInGraph(const TSharedRef<FRigVMEditorGraphExplorerTreeElement>& Element) const
+{
+	TRACE_CPUPROFILER_EVENT_SCOPE(SRigVMEditorGraphExplorerTreeView::IsVariableElementUsedInGraph);
+
+	const TSharedPtr<FRigVMEdGraphNodeRegistry> EdGraphNodeVariableRegistry = WeakGraphExplorer.IsValid() ? 
+		WeakGraphExplorer.Pin()->GetEdGraphNodeVariableRegistry() : 
+		nullptr;
+
+	if (!EdGraphNodeVariableRegistry.IsValid())
+	{
+		return false;
+	}
+
+	return
+		EdGraphNodeVariableRegistry->GetWeakEdGrapNodes()
+		.ContainsByPredicate(
+			[&Element](const TWeakObjectPtr<URigVMEdGraphNode>& WeakEdGraphNode)
+			{
+				const URigVMVariableNode* VariableNode = WeakEdGraphNode.IsValid() ?
+					Cast<URigVMVariableNode>(WeakEdGraphNode->GetModelNode()) :
+					nullptr;
+
+				return VariableNode ?
+					VariableNode->GetVariableName() == Element->Key.Name :
+					false;
+			});
+}
+
 #undef LOCTEXT_NAMESPACE
\ No newline at end of file
diff --git a/Engine/Plugins/Runtime/RigVM/Source/RigVMEditor/Public/Widgets/SRigVMEditorGraphExplorer.h b/Engine/Plugins/Runtime/RigVM/Source/RigVMEditor/Public/Widgets/SRigVMEditorGraphExplorer.h
--- a/Engine/Plugins/Runtime/RigVM/Source/RigVMEditor/Public/Widgets/SRigVMEditorGraphExplorer.h
+++ b/Engine/Plugins/Runtime/RigVM/Source/RigVMEditor/Public/Widgets/SRigVMEditorGraphExplorer.h
@@ -15,6 +15,8 @@
 
 class IRigVMEditor;
 class SRigVMEditorGraphExplorerTreeView;
+namespace UE::RigVMEditor { class FRigVMEdGraphNodeRegistry; }
+
 //
 class URigVMGraph;
 class URigVMEdGraph;
@@ -78,6 +80,7 @@ public:
 	TSharedPtr<FUICommandInfo> PasteFunction;
 	TSharedPtr<FUICommandInfo> PasteVariable;
 	TSharedPtr<FUICommandInfo> PasteLocalVariable;
+	TSharedPtr<FUICommandInfo> RemoveUnusedVariables;
 
 	void RegisterCommands() override;
 };
@@ -85,12 +88,13 @@ public:
 class SRigVMEditorGraphExplorer :
 	public SCompoundWidget
 {
+	using FRigVMEdGraphNodeRegistry = UE::RigVMEditor::FRigVMEdGraphNodeRegistry;
+
 public:
 	SLATE_BEGIN_ARGS(SRigVMEditorGraphExplorer) {}
 	SLATE_END_ARGS()
 
 	void Construct(const FArguments& InArgs, TWeakPtr<IRigVMEditor> InRigVMEditor);
-	~SRigVMEditorGraphExplorer();
 
 	// Refresh the graph action menu.
 	void Refresh();
@@ -106,6 +110,10 @@ public:
 
 	void ClearSelection();
 
+	/** Returns a registry for Ed Graph Nodes that hold RigVMVariableNodes */
+	const TSharedPtr<FRigVMEdGraphNodeRegistry>& GetEdGraphNodeVariableRegistry() const { return EdGraphNodeVariableRegistry; }
+
+
 private:
 	bool bNeedsRefresh;
 
@@ -205,6 +213,8 @@ private:
 
 	void OnFindReference(bool bSearchAllBlueprints);
 	bool CanFindReference() const;
+	void OnRemoveUnusedVariables();
+	bool CanRemoveUnusedVariables() const;
 
 	void HandleSelectionChanged(TSharedPtr<FRigVMEditorGraphExplorerTreeElement> Selection, ESelectInfo::Type SelectInfo);
 
@@ -214,6 +224,9 @@ private:
 
 	TSharedPtr<SSearchBox> FilterBox;
 
+	/** A registry for Ed Graph Nodes that hold RigVMVariableNodes */
+	TSharedPtr<FRigVMEdGraphNodeRegistry> EdGraphNodeVariableRegistry;
+
 	bool bShowEmptySections = true;
 };
 
diff --git a/Engine/Plugins/Runtime/RigVM/Source/RigVMEditor/Public/Widgets/SRigVMEditorGraphExplorerTreeView.h b/Engine/Plugins/Runtime/RigVM/Source/RigVMEditor/Public/Widgets/SRigVMEditorGraphExplorerTreeView.h
--- a/Engine/Plugins/Runtime/RigVM/Source/RigVMEditor/Public/Widgets/SRigVMEditorGraphExplorerTreeView.h
+++ b/Engine/Plugins/Runtime/RigVM/Source/RigVMEditor/Public/Widgets/SRigVMEditorGraphExplorerTreeView.h
@@ -9,11 +9,13 @@
 
 class FRigVMEditorBase;
 class SSearchBox;
+class SRigVMEditorGraphExplorer;
 class SRigVMEditorGraphExplorerTreeView;
 class SRigVMEditorGraphExplorerItem;
 class FRigVMEditorGraphExplorerTreeElement;
 class URigVMGraph;
 class IPinTypeSelectorFilter;
+namespace UE::RigVMEditor { class FRigVMEdGraphNodeRegistry; }
 
 namespace ERigVMExplorerElementType
 {
@@ -524,7 +526,7 @@ public:
 		TSharedPtr<SRigVMEditorGraphExplorerTreeView> InTreeView);
 
 	TSharedRef<SWidget> CreateIconWidget(const FRigVMExplorerElementKey& Key);
-	TSharedRef<SWidget> CreateTextSlotWidget(const FRigVMExplorerElementKey& Key, const FText& InHighlightText);
+	TSharedRef<SWidget> CreateTextSlotWidget(const FRigVMExplorerElementKey& Key, const FText& InHighlightText, const FSlateColor& TextColor = FLinearColor::White);
 	FText GetDisplayText() const;
 	FText GetItemTooltip() const;
 	bool OnVerifyNameChanged(const FText& InText, FText& OutErrorMessage);
@@ -563,6 +565,8 @@ public:
 
 class SRigVMEditorGraphExplorerTreeView : public STreeView<TSharedPtr<FRigVMEditorGraphExplorerTreeElement>>
 {
+	using FRigVMEdGraphNodeRegistry = UE::RigVMEditor::FRigVMEdGraphNodeRegistry;
+
 public:
 
 	SLATE_BEGIN_ARGS(SRigVMEditorGraphExplorerTreeView)
@@ -570,7 +574,7 @@ public:
 		SLATE_ARGUMENT(FRigVMEditorGraphExplorerTreeDelegates, RigTreeDelegates)
 	SLATE_END_ARGS()
 
-	void Construct(const FArguments& InArgs);
+	void Construct(const FArguments& InArgs, const TSharedRef<SRigVMEditorGraphExplorer>& InGraphExplorer);
 	virtual ~SRigVMEditorGraphExplorerTreeView() {}
 
 	/** Save a snapshot of the internal map that tracks item expansion before tree reconstruction */
@@ -603,6 +607,7 @@ public:
 	TSharedPtr< SWidget > CreateContextMenu();
 	void SetExpansionRecursive(TSharedPtr<FRigVMEditorGraphExplorerTreeElement> InElement, bool bTowardsParent, bool bShouldBeExpanded);
 	void HandleGetChildrenForTree(TSharedPtr<FRigVMEditorGraphExplorerTreeElement> InItem, TArray<TSharedPtr<FRigVMEditorGraphExplorerTreeElement>>& OutChildren);
+	TArray<FRigVMExplorerElementKey> GetKeys() const;
 	TArray<FRigVMExplorerElementKey> GetSelectedKeys() const;
 	void SetSelection(TArray<FRigVMExplorerElementKey>& InSelectedKeys);
 	TSharedPtr<FRigVMEditorGraphExplorerTreeElement> FindElement(const FRigVMExplorerElementKey& Key);
@@ -615,6 +620,9 @@ public:
 
 private:
 
+	/** Returns true if the eleemnt refers to a variable used in the graph. The element's key type has to correspond to a variable (ensured) */
+	bool IsVariableElementUsedInGraph(const TSharedRef<FRigVMEditorGraphExplorerTreeElement>& Element) const;
+
 	/** A temporary snapshot of the SparseItemInfos in STreeView, used during RefreshTreeView() */
 	TSparseItemMap OldSparseItemInfos;
 
@@ -631,6 +639,9 @@ private:
 
 	FText FilterText;
 
+	/** The outer graph explorer widget */
+	TWeakPtr<SRigVMEditorGraphExplorer> WeakGraphExplorer;
+
 	friend class SRigVMEditorGraphExplorer;
 	friend class SRigVMEditorGraphExplorerItem;
 };
\ No newline at end of file
