Control Rig: Improve node registry related comments

#rnx

diff --git a/Engine/Plugins/Runtime/RigVM/Source/RigVMDeveloper/Public/RigVMModel/RigVMClient.h b/Engine/Plugins/Runtime/RigVM/Source/RigVMDeveloper/Public/RigVMModel/RigVMClient.h
--- a/Engine/Plugins/Runtime/RigVM/Source/RigVMDeveloper/Public/RigVMModel/RigVMClient.h
+++ b/Engine/Plugins/Runtime/RigVM/Source/RigVMDeveloper/Public/RigVMModel/RigVMClient.h
@@ -400,7 +400,7 @@ private:
 	TWeakInterfacePtr<IRigVMClientExternalModelHost> ExternalModelHost;
 
 #if WITH_EDITOR
-	/** Returns a delegate raised after any graph in this client was modified and after the client adopted these changes */
+	/** Delegate raised after any graph in this client was modified and after the client adopted these changes */
 	FRigVMClientPostGraphModified PostGraphModifiedDelegate;
 #endif // WITH_EDITOR
 
diff --git a/Engine/Plugins/Runtime/RigVM/Source/RigVMEditor/Private/Editor/RigVMEdGraphNodeRegistry.h b/Engine/Plugins/Runtime/RigVM/Source/RigVMEditor/Private/Editor/RigVMEdGraphNodeRegistry.h
--- a/Engine/Plugins/Runtime/RigVM/Source/RigVMEditor/Private/Editor/RigVMEdGraphNodeRegistry.h
+++ b/Engine/Plugins/Runtime/RigVM/Source/RigVMEditor/Private/Editor/RigVMEdGraphNodeRegistry.h
@@ -33,7 +33,7 @@ namespace UE::RigVMEditor
 		 * @param InRigVMAssetInterface			The Rig VM Asset Interface for which this registry is constructed 
 		 * @param InRigVMNodeClass				The Rig VM Node Class for which Ed Graph nodes are held
 		 * 
-		 * @return								The reistry instance
+		 * @return								The registry instance
 		 */
 		static TSharedRef<FRigVMEdGraphNodeRegistry> GetOrCreateRegistry(const TScriptInterface<IRigVMAssetInterface>& InRigVMAssetInterface, UClass* InRigVMNodeClass);
 
