RigVM: Stabilize type removal in the registry by avoiding removal of array types

#rnx

diff --git a/Engine/Plugins/Runtime/RigVM/Source/RigVM/Private/RigVMCore/RigVMRegistry.cpp b/Engine/Plugins/Runtime/RigVM/Source/RigVM/Private/RigVMCore/RigVMRegistry.cpp
--- a/Engine/Plugins/Runtime/RigVM/Source/RigVM/Private/RigVMCore/RigVMRegistry.cpp
+++ b/Engine/Plugins/Runtime/RigVM/Source/RigVM/Private/RigVMCore/RigVMRegistry.cpp
@@ -1033,7 +1033,15 @@ bool FRigVMRegistry_NoLock::OnModulesUnloadedInternal(const TArray<FString>& Mod
 			{
 				return;
 			}
-			
+
+			// avoid array types, the base type removal will clean up related array types.
+			// For example removing FMyStruct will also clean up TArray<FMyStruct>.
+			// Note: There's always a base type registered for an array type being present.
+			if (Types[Index].bIsArray)
+			{
+				return;
+			}
+
 			if (!IsWithinPlugin(ModulePrefixes, Types[Index].Type.CPPTypeObject))
 			{
 				return;
