RigVM: FString and FName to integer units

#rnx

diff --git a/Engine/Plugins/Runtime/RigVM/Source/RigVM/Private/RigVMFunctions/RigVMFunction_Name.cpp b/Engine/Plugins/Runtime/RigVM/Source/RigVM/Private/RigVMFunctions/RigVMFunction_Name.cpp
--- a/Engine/Plugins/Runtime/RigVM/Source/RigVM/Private/RigVMFunctions/RigVMFunction_Name.cpp
+++ b/Engine/Plugins/Runtime/RigVM/Source/RigVM/Private/RigVMFunctions/RigVMFunction_Name.cpp
@@ -1,6 +1,7 @@
 // Copyright Epic Games, Inc. All Rights Reserved.
 
 #include "RigVMFunctions/RigVMFunction_Name.h"
+#include "RigVMFunctions/RigVMFunction_String.h"
 #if WITH_DEV_AUTOMATION_TESTS
 #include "RigVMCore/RigVMStructTest.h"
 #endif
@@ -261,4 +262,30 @@ FRigVMFunction_IsNone_Execute()
 FRigVMFunction_IsNameValid_Execute()
 {
 	Result = !Value.IsNone();
-}
\ No newline at end of file
+}
+
+FRigVMFunction_GetNameNumericSuffix_Execute()
+{
+	Suffix = INDEX_NONE;
+	Success = false;
+
+	if (Name.IsNone())
+	{
+		return;
+	}
+
+	if (Name.GetNumber() != NAME_NO_NUMBER_INTERNAL)
+	{
+		Suffix = NAME_INTERNAL_TO_EXTERNAL(Name.GetNumber());
+		Success = true;
+		return;
+	}
+
+	FRigVMFunction_StringToInteger::StaticExecute(
+		ExecuteContext, 
+		Name.ToString(),
+		true, // chop left
+		false, // chop right
+		Suffix,
+		Success);
+}
diff --git a/Engine/Plugins/Runtime/RigVM/Source/RigVM/Private/RigVMFunctions/RigVMFunction_String.cpp b/Engine/Plugins/Runtime/RigVM/Source/RigVM/Private/RigVMFunctions/RigVMFunction_String.cpp
--- a/Engine/Plugins/Runtime/RigVM/Source/RigVM/Private/RigVMFunctions/RigVMFunction_String.cpp
+++ b/Engine/Plugins/Runtime/RigVM/Source/RigVM/Private/RigVMFunctions/RigVMFunction_String.cpp
@@ -162,6 +162,86 @@ FRigVMFunction_StringPadInteger_Execute()
 	}
 }
 
+FRigVMFunction_StringToInteger_Execute()
+{
+	Success = false;
+	Result = INDEX_NONE;
+
+	FStringView View(Value);
+
+	auto IsValidChar = [](int32 InIndex, const TCHAR& InChar)
+	{
+		return FChar::IsDigit(InChar) || (InIndex == 0 && InChar == '-');
+	};
+	
+	if (ChopLeft)
+	{
+		for (int32 Index = View.Len() - 1; Index >= 0; --Index)
+		{
+			const TCHAR& C = View[Index];
+			if (!IsValidChar(Index, C))
+			{
+				View.MidInline(Index+1);
+				break;
+			}
+		}
+	}
+	
+	if (ChopRight)
+	{
+		for (int32 Index = 0; Index < View.Len(); ++Index)
+		{
+			const TCHAR& C = View[Index];
+			if (!IsValidChar(Index, C))
+			{
+				View.LeftInline(Index);
+				break;
+			}
+		}
+
+		while (!View.IsEmpty())
+		{
+			const int32 Index = View.Len()-1;
+			const TCHAR& C = View[Index];
+			if (IsValidChar(Index, C))
+			{
+				break;
+			}				
+			View.LeftChopInline(1);
+		}
+	}
+
+	if (View.IsEmpty())
+	{
+		return;
+	}
+
+	for (int32 i = 0; i < View.Len(); ++i)
+	{
+		const TCHAR& C = View[i];
+		if (!IsValidChar(i, C))
+		{
+#if WITH_EDITOR
+			if(ExecuteContext.GetLog() != nullptr)
+			{
+				ExecuteContext.Report(EMessageSeverity::Error, ExecuteContext.GetFunctionName(), ExecuteContext.GetInstructionIndex(), TEXT("String contains non-digit characters."));
+			}
+#endif
+			return;
+		}
+	}
+
+	if (View.Len() == Value.Len())
+	{
+		Result = FCString::Atoi(*Value);
+	}
+	else
+	{
+		Result = FCString::Atoi(*FString(View));
+	}
+	Success = true;
+}
+
 const TArray<FRigVMTemplateArgumentInfo>& FRigDispatch_ToString::GetArgumentInfos(FRigVMRegistryHandle& InRegistry) const
 {
 	if(CachedArgumentInfos.IsEmpty())
diff --git a/Engine/Plugins/Runtime/RigVM/Source/RigVM/Public/RigVMFunctions/RigVMFunction_Name.h b/Engine/Plugins/Runtime/RigVM/Source/RigVM/Public/RigVMFunctions/RigVMFunction_Name.h
--- a/Engine/Plugins/Runtime/RigVM/Source/RigVM/Public/RigVMFunctions/RigVMFunction_Name.h
+++ b/Engine/Plugins/Runtime/RigVM/Source/RigVM/Public/RigVMFunctions/RigVMFunction_Name.h
@@ -233,4 +233,32 @@ struct FRigVMFunction_IsNameValid : public FRigVMFunction_NameBase
 
 	UPROPERTY(meta = (Output))
 	bool Result;
+};
+
+/**
+ * Tests whether this name ends with a numeric suffix
+ */
+USTRUCT(meta = (DisplayName = "Get Numeric Suffix", Keywords = "Suffix,Number,Integer"))
+struct FRigVMFunction_GetNameNumericSuffix : public FRigVMFunction_NameBase
+{
+	GENERATED_BODY()
+
+	FRigVMFunction_GetNameNumericSuffix()
+	{
+		Name = NAME_None;
+		Suffix = INDEX_NONE;
+		Success = false;
+	}
+
+	RIGVM_METHOD()
+	RIGVM_API virtual void Execute() override;
+
+	UPROPERTY(meta = (Input))
+	FName Name;
+
+	UPROPERTY(meta = (Output))
+	int32 Suffix;
+
+	UPROPERTY(meta = (Output))
+	bool Success;
 };
\ No newline at end of file
diff --git a/Engine/Plugins/Runtime/RigVM/Source/RigVM/Public/RigVMFunctions/RigVMFunction_String.h b/Engine/Plugins/Runtime/RigVM/Source/RigVM/Public/RigVMFunctions/RigVMFunction_String.h
--- a/Engine/Plugins/Runtime/RigVM/Source/RigVM/Public/RigVMFunctions/RigVMFunction_String.h
+++ b/Engine/Plugins/Runtime/RigVM/Source/RigVM/Public/RigVMFunctions/RigVMFunction_String.h
@@ -500,6 +500,42 @@ struct FRigVMFunction_StringPadInteger : public FRigVMFunction_StringBase
 	FString Result;
 };
 
+/**
+ * Converts a string to an integer
+ */
+USTRUCT(meta = (DisplayName = "To Integer", Keywords = "ToInt,Number"))
+struct FRigVMFunction_StringToInteger : public FRigVMFunction_StringBase
+{
+	GENERATED_BODY()
+
+	FRigVMFunction_StringToInteger()
+	{
+		Value = FString();
+		Result = INDEX_NONE;
+		ChopLeft  = ChopRight = Success = false;
+	}
+
+	RIGVM_METHOD()
+	RIGVM_API virtual void Execute();
+
+	UPROPERTY(meta = (Input))
+	FString Value;
+
+	// chops non-digit characters from the left of the string
+	UPROPERTY(meta = (Input))
+	bool ChopLeft;
+
+	// chops non-digit characters from the right of the string
+	UPROPERTY(meta = (Input))
+	bool ChopRight;
+
+	UPROPERTY(meta = (Output))
+	int32 Result;
+
+	UPROPERTY(meta = (Output))
+	bool Success;
+};
+
 /*
  * Converts any value to string
  */
