RigVM: Enable localized registry as per default

#rnx

Full cook is here: https://horde.devtools.epicgames.com/job/68b7eebfd73dd1b156e065b9

diff --git a/Engine/Plugins/Runtime/RigVM/Source/RigVM/Private/RigVMCore/RigVMRegistry.cpp b/Engine/Plugins/Runtime/RigVM/Source/RigVM/Private/RigVMCore/RigVMRegistry.cpp
--- a/Engine/Plugins/Runtime/RigVM/Source/RigVM/Private/RigVMCore/RigVMRegistry.cpp
+++ b/Engine/Plugins/Runtime/RigVM/Source/RigVM/Private/RigVMCore/RigVMRegistry.cpp
@@ -22,7 +22,7 @@
 static FDelayedAutoRegisterHelper GRigVMRegistrySingletonHelper(EDelayedRegisterRunPhase::EndOfEngineInit, &FRigVMRegistry_NoLock::OnEngineInit);
 
 static TAutoConsoleVariable<bool> CVarRigVMUpdateDispatchFactoriesGreedily(TEXT("RigVM.UpdateDispatchFactoriesGreedily"), true, TEXT("Set this to false to avoid loading dispatch factories during engine init / plugin mount."));
-TAutoConsoleVariable<bool> CVarRigVMEnableLocalizedRegistry(TEXT("RigVM.EnableLocalizedRegistry"), false, TEXT("Set this to true to allow VMs to serialize their own localized registries."));
+TAutoConsoleVariable<bool> CVarRigVMEnableLocalizedRegistry(TEXT("RigVM.EnableLocalizedRegistry"), true, TEXT("Set this to true to allow VMs to serialize their own localized registries."));
 
 FRigVMRegistry_NoLock::FRigVMRegistry_NoLock(bool InIsGlobalRegistry)
 	: ThisHandle(this)
