RigVM: Don't log warning about serialized localized registry being skipped

#rnx

diff --git a/Engine/Plugins/Runtime/RigVM/Source/RigVM/Private/RigVMCore/RigVM.cpp b/Engine/Plugins/Runtime/RigVM/Source/RigVM/Private/RigVMCore/RigVM.cpp
--- a/Engine/Plugins/Runtime/RigVM/Source/RigVM/Private/RigVMCore/RigVM.cpp
+++ b/Engine/Plugins/Runtime/RigVM/Source/RigVM/Private/RigVMCore/RigVM.cpp
@@ -327,7 +327,8 @@ void URigVM::Load(FArchive& Ar)
 
 			if(CVarRigVMEnableLocalizedRegistry.GetValueOnAnyThread() == false)
 			{
-				UE_LOG(LogRigVM, Warning, TEXT("VM '%s' was serialized with a localized registry, but CVarRigVMEnableLocalizedRegistry is disabled."), *GetPathName());
+				// disabling warning since it is causing cooking errors.
+				// UE_LOG(LogRigVM, Warning, TEXT("VM '%s' was serialized with a localized registry, but CVarRigVMEnableLocalizedRegistry is disabled."), *GetPathName());
 				Ar.Seek(ArchivePosAfterSerializedRegistry);
 			}
 			else
