Control Rig: Finishing UX polish on Preview Here functionality

#rnx

diff --git a/Engine/Plugins/Animation/ControlRig/Source/ControlRigDeveloper/Public/Graph/ControlRigGraphSchema.h b/Engine/Plugins/Animation/ControlRig/Source/ControlRigDeveloper/Public/Graph/ControlRigGraphSchema.h
--- a/Engine/Plugins/Animation/ControlRig/Source/ControlRigDeveloper/Public/Graph/ControlRigGraphSchema.h
+++ b/Engine/Plugins/Animation/ControlRig/Source/ControlRigDeveloper/Public/Graph/ControlRigGraphSchema.h
@@ -36,6 +36,7 @@ public:
 	// URigVMEdGraphSchema interface
 	virtual const FLazyName& GetRootGraphName() const override { return GraphName_ControlRig; }
 	UE_API virtual bool IsRigVMDefaultEvent(const FName& InEventName) const override;
+	virtual bool SupportsPreviewFromHereOnNode(const UEdGraphNode* InNode) const override { return true; }
 };
 
 #undef UE_API
diff --git a/Engine/Plugins/Animation/ControlRig/Source/ControlRigEditor/Private/ControlRigEditorStyle.h b/Engine/Plugins/Animation/ControlRig/Source/ControlRigEditor/Private/ControlRigEditorStyle.h
--- a/Engine/Plugins/Animation/ControlRig/Source/ControlRigEditor/Private/ControlRigEditorStyle.h
+++ b/Engine/Plugins/Animation/ControlRig/Source/ControlRigEditor/Private/ControlRigEditorStyle.h
@@ -130,6 +130,8 @@ public:
 				Set( "ControlRig.Viewport.Notification.DirectManipulation", new BOX_BRUSH("Common/RoundedSelection_16x", 4.0f/16.0f, FLinearColor(FColor(0, 112, 224))));
 				// similar style to "AnimViewport.Notification.Warning"
 				Set( "ControlRig.Viewport.Notification.ReplayValidation", new BOX_BRUSH("Common/RoundedSelection_16x", 4.0f/16.0f, ErrorFillColor));
+				// similar style to "AnimViewport.Notification.Warning"
+				Set( "ControlRig.Viewport.Notification.PreviewingNode", new BOX_BRUSH("Common/RoundedSelection_16x", 4.0f/16.0f, FLinearColor(FColor(0, 0, 255))));
 			}
 		}
 
diff --git a/Engine/Plugins/Animation/ControlRig/Source/ControlRigEditor/Private/Editor/ControlRigEditor.cpp b/Engine/Plugins/Animation/ControlRig/Source/ControlRigEditor/Private/Editor/ControlRigEditor.cpp
--- a/Engine/Plugins/Animation/ControlRig/Source/ControlRigEditor/Private/Editor/ControlRigEditor.cpp
+++ b/Engine/Plugins/Animation/ControlRig/Source/ControlRigEditor/Private/Editor/ControlRigEditor.cpp
@@ -2657,6 +2657,73 @@ void FControlRigBaseEditor::HandleViewportCreated(const TSharedRef<class IPerson
 		);
 	}
 
+	{
+		FPersonaViewportNotificationOptions PreviewingNodeNotificationOptions(TAttribute<EVisibility>::CreateRaw(this, &FControlRigBaseEditor::GetPreviewingNodeVisibility));
+		PreviewingNodeNotificationOptions.OnGetBrushOverride = TAttribute<const FSlateBrush*>(FControlRigEditorStyle::Get().GetBrush("ControlRig.Viewport.Notification.PreviewingNode"));
+
+		InViewport->AddNotification(
+			EMessageSeverity::Info,
+			false,
+			SNew(SHorizontalBox)
+			.Visibility(SharedRef(), &FControlRigBaseEditor::GetPreviewingNodeVisibility)
+			.ToolTipText(LOCTEXT("PreviewingNodeTooltip", "Previewing Node in the graph,\nthe full rig is not running.\nClick here to jump to the node."))
+			+SHorizontalBox::Slot()
+			.FillWidth(1.0f)
+			.Padding(4.0f, 4.0f)
+			[
+				SNew(SHorizontalBox)
+				+SHorizontalBox::Slot()
+				.AutoWidth()
+				.VAlign(VAlign_Center)
+				.Padding(0.0f, 0.0f, 4.0f, 0.0f)
+				[
+					SNew(STextBlock)
+					.TextStyle(FAppStyle::Get(), "AnimViewport.MessageText")
+					.Font(FAppStyle::Get().GetFontStyle("FontAwesome.9"))
+					.Text(FEditorFontGlyphs::Eye)
+				]
+				
+				+SHorizontalBox::Slot()
+				.VAlign(VAlign_Center)
+				.AutoWidth()
+				[
+					SNew(SButton)
+					.ButtonStyle( FAppStyle::Get(), "SimpleButton" )
+					.OnClicked_Lambda([this]()
+					{
+						OnPreviewingNodeJumpTo();
+						return FReply::Handled();
+					})
+					[
+						SNew(STextBlock)
+						.TextStyle(FAppStyle::Get(), "AnimViewport.MessageText")
+						.Text(LOCTEXT("PreviewingNode", "Previewing Node"))
+					]
+				]
+
+				+SHorizontalBox::Slot()
+				.VAlign(VAlign_Center)
+				.AutoWidth()
+				[
+					SNew(SButton)
+					.ToolTipText(LOCTEXT("CancelPreviewingNodeTooltip", "Stop previewing and run the full rig."))
+					.ButtonStyle( FAppStyle::Get(), "SimpleButton" )
+					.OnClicked_Lambda([this]()
+					{
+						OnPreviewingNodeCancel();
+						return FReply::Handled();
+					})
+					[
+						SNew(SImage)
+						.ColorAndOpacity(FSlateColor::UseForeground())
+						.Image(FAppStyle::GetBrush("Icons.X"))
+					]
+				]
+			],
+			PreviewingNodeNotificationOptions
+		);
+	}
+
 	if(CVarControlRigShowTestingToolbar.GetValueOnAnyThread())
 	{
 		{
@@ -4107,7 +4174,7 @@ bool FControlRigBaseEditor::IsDirectManipulationEnabled() const
 
 EVisibility FControlRigBaseEditor::GetDirectManipulationVisibility() const
 {
-	return IsDirectManipulationEnabled() ? EVisibility::Visible : EVisibility::Hidden;
+	return IsDirectManipulationEnabled() ? EVisibility::Visible : EVisibility::Collapsed;
 }
 
 FText FControlRigBaseEditor::GetDirectionManipulationText() const
@@ -4230,11 +4297,71 @@ void FControlRigBaseEditor::RefreshDirectManipulationTextList()
 	}
 }
 
+EVisibility FControlRigBaseEditor::GetPreviewingNodeVisibility() const
+{
+	if (UControlRig* DebuggedControlRig = Cast<UControlRig>(GetRigVMAssetInterface()->GetObjectBeingDebugged()))
+	{
+		if (DebuggedControlRig->GetDebugInfo().GetInstructionForEarlyExit() != INDEX_NONE)
+		{
+			return EVisibility::Visible;
+		}
+	}
+	return EVisibility::Collapsed;
+}
+
+void FControlRigBaseEditor::OnPreviewingNodeJumpTo()
+{
+	check(GetRigVMAssetInterface());
+
+	UControlRig* DebuggedControlRig = Cast<UControlRig>(GetRigVMAssetInterface()->GetObjectBeingDebugged());
+	if (!DebuggedControlRig)
+	{
+		return;
+	}
+	const int32 InstructionToJumpTo = DebuggedControlRig->GetDebugInfo().GetInstructionForEarlyExit();
+	if (InstructionToJumpTo == INDEX_NONE)
+	{
+		return;
+	}
+
+	URigVM* VM = DebuggedControlRig->GetVM();
+	if (!VM)
+	{
+		return;
+	}
+
+	const URigVMNode* Subject = Cast<URigVMNode>(VM->GetByteCode().GetSubjectForInstruction(InstructionToJumpTo));
+	if (!Subject)
+	{
+		return;
+	}
+	
+	FRigVMClient* Client = GetRigVMAssetInterface()->GetRigVMClient();
+	if (!Client)
+	{
+		return;
+	}
+
+	URigVMController* Controller = Client->GetOrCreateController(Subject->GetGraph());
+	if (!Controller)
+	{
+		return;
+	}
+
+	Controller->RequestJumpToHyperLink(Subject);
+}
+
+void FControlRigBaseEditor::OnPreviewingNodeCancel()
+{
+	check(GetRigVMAssetInterface());
+	GetRigVMAssetInterface()->ResetEarlyExitInstruction();
+}
+
 EVisibility FControlRigBaseEditor::GetConnectorWarningVisibility() const
 {
 	if(GetConnectorWarningText().IsEmpty())
 	{
-		return EVisibility::Hidden;
+		return EVisibility::Collapsed;
 	}
 	return EVisibility::Visible;
 }
@@ -6765,7 +6892,7 @@ EVisibility FControlRigBaseEditor::GetReplayValidationErrorVisibility() const
 			return EVisibility::Visible;
 		}
 	}
-	return EVisibility::Hidden;
+	return EVisibility::Collapsed;
 }
 
 FText FControlRigBaseEditor::GetReplayValidationErrorTooltip() const
diff --git a/Engine/Plugins/Animation/ControlRig/Source/ControlRigEditor/Private/Editor/ControlRigEditor.h b/Engine/Plugins/Animation/ControlRig/Source/ControlRigEditor/Private/Editor/ControlRigEditor.h
--- a/Engine/Plugins/Animation/ControlRig/Source/ControlRigEditor/Private/Editor/ControlRigEditor.h
+++ b/Engine/Plugins/Animation/ControlRig/Source/ControlRigEditor/Private/Editor/ControlRigEditor.h
@@ -318,6 +318,9 @@ protected:
 	const TArray<TSharedPtr<FString>>& GetDirectManipulationTargetTextList() const;
 	bool ClearDirectManipulationSubject() { return SetDirectionManipulationSubject(nullptr); }
 	void RefreshDirectManipulationTextList();
+	EVisibility GetPreviewingNodeVisibility() const;
+	void OnPreviewingNodeJumpTo();
+	void OnPreviewingNodeCancel();
 
 	// Rig connector functionality
 	EVisibility GetConnectorWarningVisibility() const;
diff --git a/Engine/Plugins/Runtime/RigVM/Source/RigVM/Private/RigVMSettings.cpp b/Engine/Plugins/Runtime/RigVM/Source/RigVM/Private/RigVMSettings.cpp
--- a/Engine/Plugins/Runtime/RigVM/Source/RigVM/Private/RigVMSettings.cpp
+++ b/Engine/Plugins/Runtime/RigVM/Source/RigVM/Private/RigVMSettings.cpp
@@ -15,7 +15,7 @@ URigVMEditorSettings::URigVMEditorSettings(const FObjectInitializer& ObjectIniti
 #if WITH_EDITORONLY_DATA
 	bHighlightSimilarNodes = false;
 	bFadeOutUnrelatedNodes = false;
-	bUseFlashLight = true;
+	bUseFlashLight = false;
 	bAutoLinkMutableNodes = false;
 #endif
 }
diff --git a/Engine/Plugins/Runtime/RigVM/Source/RigVMDeveloper/Private/EdGraph/RigVMEdGraph.cpp b/Engine/Plugins/Runtime/RigVM/Source/RigVMDeveloper/Private/EdGraph/RigVMEdGraph.cpp
--- a/Engine/Plugins/Runtime/RigVM/Source/RigVMDeveloper/Private/EdGraph/RigVMEdGraph.cpp
+++ b/Engine/Plugins/Runtime/RigVM/Source/RigVMDeveloper/Private/EdGraph/RigVMEdGraph.cpp
@@ -431,12 +431,9 @@ bool URigVMEdGraph::HandleModifiedEvent_Internal(ERigVMGraphNotifType InNotifTyp
 		{
 			if (URigVMNode* ModelNode = Cast<URigVMNode>(InSubject))
 			{
-				if (ModelNode->IsA<URigVMLibraryNode>() || ModelNode->IsA<URigVMTemplateNode>())
+				if (URigVMEdGraphNode* RigNode = Cast<URigVMEdGraphNode>(FindNodeForModelNodeName(ModelNode->GetFName())))
 				{
-					if (URigVMEdGraphNode* RigNode = Cast<URigVMEdGraphNode>(FindNodeForModelNodeName(ModelNode->GetFName())))
-					{
-						RigNode->SetColorFromModel(ModelNode->GetNodeColor());
-					}
+					RigNode->InvalidateNodeColor();
 				}
 				else if(UEdGraphNode_Comment * EdComment = Cast<UEdGraphNode_Comment>(FindNodeForModelNodeName(ModelNode->GetFName())))
 				{
@@ -753,6 +750,18 @@ bool URigVMEdGraph::HandleModifiedEvent_Internal(ERigVMGraphNotifType InNotifTyp
 			}
 			break;
 		}
+		case ERigVMGraphNotifType::NodeEarlyExitChanged:
+		{
+			if (URigVMNode* ModelNode = Cast<URigVMNode>(InSubject))
+			{
+				if (URigVMEdGraphNode* RigNode = Cast<URigVMEdGraphNode>(FindNodeForModelNodeName(ModelNode->GetFName())))
+				{
+					RigNode->InvalidateNodeTitle();
+					RigNode->InvalidateNodeColor();
+				}
+			}
+			break;
+		}
 		default:
 		{
 			break;
diff --git a/Engine/Plugins/Runtime/RigVM/Source/RigVMDeveloper/Private/EdGraph/RigVMEdGraphNode.cpp b/Engine/Plugins/Runtime/RigVM/Source/RigVMDeveloper/Private/EdGraph/RigVMEdGraphNode.cpp
--- a/Engine/Plugins/Runtime/RigVM/Source/RigVMDeveloper/Private/EdGraph/RigVMEdGraphNode.cpp
+++ b/Engine/Plugins/Runtime/RigVM/Source/RigVMDeveloper/Private/EdGraph/RigVMEdGraphNode.cpp
@@ -410,15 +410,7 @@ void URigVMEdGraphNode::PostReconstructNode()
 		}
 	}
 
-	if(IsOutDated())
-	{
-		static const FLinearColor WarningColor = FAppStyle::GetColor("ErrorReporting.WarningBackgroundColor");
-		SetColorFromModel(WarningColor);
-	}
-	else if(const URigVMNode* ModelNode = GetModelNode())
-	{
-		SetColorFromModel(ModelNode->GetNodeColor());
-	}
+	InvalidateNodeColor();
 }
 
 void URigVMEdGraphNode::SetColorFromModel(const FLinearColor& InColor)
@@ -2088,6 +2080,28 @@ void URigVMEdGraphNode::InvalidateNodeTitle() const
 	NodeTitleDirtied.Broadcast();
 }
 
+void URigVMEdGraphNode::InvalidateNodeColor()
+{
+	if(const URigVMNode* ModelNode = GetModelNode())
+	{
+		if (ModelNode->IsOutDated())
+		{
+			static const FLinearColor WarningColor = FAppStyle::GetColor("ErrorReporting.WarningBackgroundColor");
+			SetColorFromModel(WarningColor);
+		}
+		else if (ModelNode->HasEarlyExitMarker())
+		{
+			static constexpr FLinearColor PreviewColor = FLinearColor(0.f,0.f,1.f,1.f);
+			SetColorFromModel(PreviewColor);
+		}
+		else
+		{
+			SetColorFromModel(ModelNode->GetNodeColor());
+		}
+	}
+
+}
+
 bool URigVMEdGraphNode::CanCreateUnderSpecifiedSchema(const UEdGraphSchema* InSchema) const
 {
 	return InSchema->IsA<URigVMEdGraphSchema>();
diff --git a/Engine/Plugins/Runtime/RigVM/Source/RigVMDeveloper/Private/EdGraph/RigVMEdGraphSchema.cpp b/Engine/Plugins/Runtime/RigVM/Source/RigVMDeveloper/Private/EdGraph/RigVMEdGraphSchema.cpp
--- a/Engine/Plugins/Runtime/RigVM/Source/RigVMDeveloper/Private/EdGraph/RigVMEdGraphSchema.cpp
+++ b/Engine/Plugins/Runtime/RigVM/Source/RigVMDeveloper/Private/EdGraph/RigVMEdGraphSchema.cpp
@@ -1038,6 +1038,11 @@ bool URigVMEdGraphSchema::CanPromotePinToVariable(const UEdGraphPin& Pin, bool b
 	return true;
 }
 
+bool URigVMEdGraphSchema::SupportsPreviewFromHereOnNode(const UEdGraphNode* InNode) const
+{
+	return false;
+}
+
 bool URigVMEdGraphSchema::SupportsPinType(const UScriptStruct* ScriptStruct) const
 {
 	if(!ScriptStruct)
diff --git a/Engine/Plugins/Runtime/RigVM/Source/RigVMDeveloper/Private/RigVMAsset.cpp b/Engine/Plugins/Runtime/RigVM/Source/RigVMDeveloper/Private/RigVMAsset.cpp
--- a/Engine/Plugins/Runtime/RigVM/Source/RigVMDeveloper/Private/RigVMAsset.cpp
+++ b/Engine/Plugins/Runtime/RigVM/Source/RigVMDeveloper/Private/RigVMAsset.cpp
@@ -1007,6 +1007,8 @@ bool IRigVMAssetInterface::SetEarlyExitInstruction(int32 InInstructionIndex)
 			}
 		}
 	}
+
+	GetRigGraphDisplaySettings().bShowNodeRunCounts = true;
 	
 	return true;
 }
@@ -1031,13 +1033,13 @@ bool IRigVMAssetInterface::SetEarlyExitInstruction(URigVMNode* InNodeToExitEarly
 	}
 
 	// find the last consecutive instruction for a node in the first block
-	int32 Instruction = InstructionIndices[0];
-	while (Instruction < InstructionIndices.Num() - 1 && Instruction == InstructionIndices[Instruction + 1] - 1)
+	int32 IndexInArray = 0;
+	while (IndexInArray < InstructionIndices.Num() - 1 && InstructionIndices[IndexInArray] == InstructionIndices[IndexInArray + 1] - 1)
 	{
-		Instruction++;
+		IndexInArray++;
 	}
-	
-	return SetEarlyExitInstruction(Instruction);
+
+	return SetEarlyExitInstruction(InstructionIndices[IndexInArray]);
 }
 
 bool IRigVMAssetInterface::ResetEarlyExitInstruction()
@@ -1059,9 +1061,257 @@ bool IRigVMAssetInterface::ResetEarlyExitInstruction()
 		}
 	}
 
+	GetRigGraphDisplaySettings().bShowNodeRunCounts = false;;
 	return true;
 }
 
+void IRigVMAssetInterface::TogglePreviewHere(const URigVMGraph* InGraph)
+{
+	if (!InGraph)
+	{
+		return;
+	}
+
+	URigVMHost* DebuggedHost = Cast<URigVMHost>(GetObjectBeingDebugged());
+	if (!DebuggedHost)
+	{
+		return;
+	}
+
+	const int32 CurrentEarlyExistInstruction = DebuggedHost->GetDebugInfo().GetInstructionForEarlyExit();
+	if (CurrentEarlyExistInstruction != INDEX_NONE)
+	{
+		if (CurrentEarlyExistInstruction == GetPreviewNodeInstructionIndexFromSelection(InGraph))
+		{
+			ResetEarlyExitInstruction();
+			return;
+		}
+	}
+
+	const TArray<FName> SelectedNodeNames = InGraph->GetSelectNodes();
+	if (SelectedNodeNames.Num() == 0)
+	{
+		ResetEarlyExitInstruction();
+		return;
+	}
+
+	URigVMNode* Node = InGraph->FindNodeByName(SelectedNodeNames[0]);
+	if (!Node)
+	{
+		return;
+	}
+
+	if (!Node->IsMutable())
+	{
+		return;
+	}
+
+	SetEarlyExitInstruction(Node);
+}
+
+void IRigVMAssetInterface::PreviewHereStepForward()
+{
+	URigVMHost* DebuggedHost = Cast<URigVMHost>(GetObjectBeingDebugged());
+	if (!DebuggedHost)
+	{
+		return;
+	}
+
+	const int32 CurrentInstructionIndex =  DebuggedHost->GetDebugInfo().GetInstructionForEarlyExit();
+	if (CurrentInstructionIndex == INDEX_NONE)
+	{
+		return;
+	}
+
+	const URigVM* VM = DebuggedHost->GetVM();
+	if (!VM)
+	{
+		return;
+	}
+
+	URigVMNode* Node = Cast<URigVMNode>(VM->GetByteCode().GetSubjectForInstruction(CurrentInstructionIndex));
+	if (!Node)
+	{
+		return;
+	}
+
+	struct Local
+	{
+		static URigVMNode* FindNextNodeOnCompletedPin(URigVMNode* Node)
+		{
+			for (const URigVMPin* Pin : Node->GetPins())
+			{
+				if (Pin->GetDirection() == ERigVMPinDirection::Input || Pin->GetDirection() == ERigVMPinDirection::IO)
+				{
+					if (!Pin->IsExecuteContext())
+					{
+						continue;
+					}
+
+					const TArray<URigVMPin*> SourcePins = Pin->GetLinkedSourcePins();
+					if (SourcePins.IsEmpty())
+					{
+						continue;
+					}
+
+					URigVMNode* SourceNode = SourcePins[0]->GetNode();
+					if (SourcePins[0]->GetFName() == FRigVMStruct::ControlFlowCompletedName)
+					{
+						return FindNextNodeOnCompletedPin(SourceNode);
+					}
+
+					const URigVMPin* NextExecutePin = SourceNode->FindRootPinByName(FRigVMStruct::ControlFlowCompletedName);
+					if (!NextExecutePin)
+					{
+						// since we don't have a completed pin, try to the find the next execute pin after this one
+						const int32 SourcePinIndex = SourcePins[0]->GetPinIndex();
+						for (int32 Index = SourcePinIndex + 1; Index < SourceNode->GetPins().Num(); Index++)
+						{
+							NextExecutePin = SourceNode->GetPins()[Index];
+							if (NextExecutePin->GetDirection() != ERigVMPinDirection::Output || !NextExecutePin->IsExecuteContext())
+							{
+								NextExecutePin = nullptr;
+								continue;
+							}
+							break;
+						}
+					}
+
+					if (!NextExecutePin || !NextExecutePin->IsExecuteContext())
+					{
+						continue;
+					}
+					
+					const TArray<URigVMPin*> TargetPins = NextExecutePin->GetLinkedTargetPins();
+					if (TargetPins.IsEmpty())
+					{
+						continue;
+					}
+
+					return TargetPins[0]->GetNode();
+				}
+			}
+
+			return nullptr;
+		}
+
+		static URigVMNode* FindNextNode(URigVMNode* Node)
+		{
+			if (!Node->IsMutable())
+			{
+				return nullptr;
+			}
+
+			for (const URigVMPin* Pin : Node->GetPins())
+			{
+				if (Pin->GetDirection() == ERigVMPinDirection::Output || Pin->GetDirection() == ERigVMPinDirection::IO)
+				{
+					if (!Pin->IsExecuteContext())
+					{
+						continue;
+					}
+
+					const TArray<URigVMPin*> TargetPins = Pin->GetLinkedTargetPins();
+					if (TargetPins.IsEmpty())
+					{
+						continue;
+					}
+
+					if (!TargetPins[0]->IsExecuteContext())
+					{
+						continue;
+					}
+
+					return TargetPins[0]->GetNode();
+				}
+			}
+
+			return FindNextNodeOnCompletedPin(Node);
+		}
+	};
+
+	URigVMNode* NextNode = Local::FindNextNode(Node);
+	if (!NextNode)
+	{
+		return;
+	}
+
+	SetEarlyExitInstruction(NextNode);
+
+	FRigVMClient* Client = GetRigVMClient();
+	if (!Client)
+	{
+		return;
+	}
+
+	URigVMController* Controller = Client->GetOrCreateController(Node->GetGraph());
+	if (!Controller)
+	{
+		return;
+	}
+
+	Controller->RequestJumpToHyperLink(NextNode);
+}
+
+bool IRigVMAssetInterface::CanPreviewHereStepForward() const
+{
+	URigVMHost* DebuggedHost = Cast<URigVMHost>(GetObjectBeingDebugged());
+	if (!DebuggedHost)
+	{
+		return false;
+	}
+
+	return DebuggedHost->GetDebugInfo().GetInstructionForEarlyExit() != INDEX_NONE;
+
+}
+
+int32 IRigVMAssetInterface::GetPreviewNodeInstructionIndexFromSelection(const URigVMGraph* InGraph) const
+{
+	if (!InGraph)
+	{
+		return INDEX_NONE;
+	}
+	
+	URigVMHost* DebuggedHost = Cast<URigVMHost>(GetObjectBeingDebugged());
+	if (!DebuggedHost)
+	{
+		return INDEX_NONE;
+	}
+
+	const TArray<FName> SelectedNodeNames = InGraph->GetSelectNodes();
+	if (SelectedNodeNames.Num() == 0)
+	{
+		return INDEX_NONE;
+	}
+
+	URigVMNode* Node = InGraph->FindNodeByName(SelectedNodeNames[0]);
+	if (!Node || !Node->IsMutable())
+	{
+		return INDEX_NONE;
+	}
+
+	const URigVM* VM = DebuggedHost->GetVM();
+	if (!VM)
+	{
+		return INDEX_NONE;
+	}
+
+	const TArray<int32>& InstructionIndices = VM->GetByteCode().GetAllInstructionIndicesForSubject(Node);
+	if (InstructionIndices.IsEmpty())
+	{
+		return INDEX_NONE;
+	}
+
+	// find the last consecutive instruction for a node in the first block
+	int32 IndexInArray = 0;
+	while (IndexInArray < InstructionIndices.Num() - 1 && InstructionIndices[IndexInArray] == InstructionIndices[IndexInArray + 1] - 1)
+	{
+		IndexInArray++;
+	}
+
+	return InstructionIndices[IndexInArray];
+}
+
 IRigVMEditorModule* IRigVMAssetInterface::GetEditorModule() const
 {
 	return &IRigVMEditorModule::Get();
diff --git a/Engine/Plugins/Runtime/RigVM/Source/RigVMDeveloper/Private/RigVMModel/RigVMNode.cpp b/Engine/Plugins/Runtime/RigVM/Source/RigVMDeveloper/Private/RigVMModel/RigVMNode.cpp
--- a/Engine/Plugins/Runtime/RigVM/Source/RigVMDeveloper/Private/RigVMModel/RigVMNode.cpp
+++ b/Engine/Plugins/Runtime/RigVM/Source/RigVMDeveloper/Private/RigVMModel/RigVMNode.cpp
@@ -1443,6 +1443,23 @@ bool URigVMNode::HasEarlyExitMarker() const
 	return bHasEarlyExitMarker;
 }
 
+void URigVMNode::SetHasEarlyExitMarker(const bool bValue)
+{
+	if (bHasEarlyExitMarker == bValue)
+	{
+		return;
+	}
+	bHasEarlyExitMarker = bValue;
+
+	URigVMGraph* Graph = GetGraph();
+	if (Graph == nullptr)
+	{
+		return;
+	}
+
+	Graph->Notify(ERigVMGraphNotifType::NodeEarlyExitChanged, this);
+}
+
 bool URigVMNode::IsExcludedByEarlyExit() const
 {
 	return bIsExcludedByEarlyExit;
diff --git a/Engine/Plugins/Runtime/RigVM/Source/RigVMDeveloper/Public/EdGraph/RigVMEdGraphNode.h b/Engine/Plugins/Runtime/RigVM/Source/RigVMDeveloper/Public/EdGraph/RigVMEdGraphNode.h
--- a/Engine/Plugins/Runtime/RigVM/Source/RigVMDeveloper/Public/EdGraph/RigVMEdGraphNode.h
+++ b/Engine/Plugins/Runtime/RigVM/Source/RigVMDeveloper/Public/EdGraph/RigVMEdGraphNode.h
@@ -280,6 +280,9 @@ protected:
 	/** Something that could change our title has changed */
 	UE_API void InvalidateNodeTitle() const;
 
+	/** Something that could change our color has changed */
+	UE_API void InvalidateNodeColor();
+
 	/** Destroy all pins in an array */
 	UE_API void DestroyPinList(TArray<UEdGraphPin*>& InPins);
 
diff --git a/Engine/Plugins/Runtime/RigVM/Source/RigVMDeveloper/Public/EdGraph/RigVMEdGraphSchema.h b/Engine/Plugins/Runtime/RigVM/Source/RigVMDeveloper/Public/EdGraph/RigVMEdGraphSchema.h
--- a/Engine/Plugins/Runtime/RigVM/Source/RigVMDeveloper/Public/EdGraph/RigVMEdGraphSchema.h
+++ b/Engine/Plugins/Runtime/RigVM/Source/RigVMDeveloper/Public/EdGraph/RigVMEdGraphSchema.h
@@ -363,6 +363,7 @@ public:
 	UE_API virtual TSharedPtr<INameValidatorInterface> GetNameValidator(const UBlueprint* BlueprintObj, const FName& OriginalName, const UStruct* ValidationScope, const FName& ActionTypeId) const override;
 	UE_API virtual TSharedPtr<INameValidatorInterface> GetNameValidator(const FRigVMAssetInterfacePtr Asset, const FName& OriginalName, const UStruct* ValidationScope) const;
 	UE_API bool CanPromotePinToVariable(const UEdGraphPin& Pin, bool bCond) const;
+	UE_API virtual bool SupportsPreviewFromHereOnNode(const UEdGraphNode* InNode) const;
 
 	/** Returns true if the schema supports the script type */
 	UE_API bool SupportsPinType(const UScriptStruct* ScriptStruct) const;
diff --git a/Engine/Plugins/Runtime/RigVM/Source/RigVMDeveloper/Public/RigVMAsset.h b/Engine/Plugins/Runtime/RigVM/Source/RigVMDeveloper/Public/RigVMAsset.h
--- a/Engine/Plugins/Runtime/RigVM/Source/RigVMDeveloper/Public/RigVMAsset.h
+++ b/Engine/Plugins/Runtime/RigVM/Source/RigVMDeveloper/Public/RigVMAsset.h
@@ -284,6 +284,26 @@ public:
 	 */
 	bool ResetEarlyExitInstruction();
 
+	/**
+	 * Toggles the preview here functionality based on the selection
+	 */
+	void TogglePreviewHere(const URigVMGraph* InGraph);
+
+	/**
+	 * Steps forward in a preview here session
+	 */
+	void PreviewHereStepForward();
+
+	/**
+	 * Returns true if we can step forward in a preview here session
+	 */
+	bool CanPreviewHereStepForward() const;
+
+	/**
+	 * Returns the instruction index of the currently selected node
+	 */
+	int32 GetPreviewNodeInstructionIndexFromSelection(const URigVMGraph* InGraph) const;
+
 	virtual FRigVMController_RequestJumpToHyperlinkDelegate& OnRequestJumpToHyperlink()  { return RequestJumpToHyperlink; };
 
 	/** Returns the settings defaults for this blueprint */
diff --git a/Engine/Plugins/Runtime/RigVM/Source/RigVMDeveloper/Public/RigVMModel/RigVMGraph.h b/Engine/Plugins/Runtime/RigVM/Source/RigVMDeveloper/Public/RigVMModel/RigVMGraph.h
--- a/Engine/Plugins/Runtime/RigVM/Source/RigVMDeveloper/Public/RigVMModel/RigVMGraph.h
+++ b/Engine/Plugins/Runtime/RigVM/Source/RigVMDeveloper/Public/RigVMModel/RigVMGraph.h
@@ -250,6 +250,7 @@ private:
 	friend class FRigVMControllerCompileBracketScope;
 	friend class URigVMCompiler;
 	friend class URigVMSchema;
+	friend class URigVMNode;
 	friend struct FRigVMClient;
 	friend struct FRigVMControllerObjectFactory;
 };
diff --git a/Engine/Plugins/Runtime/RigVM/Source/RigVMDeveloper/Public/RigVMModel/RigVMNode.h b/Engine/Plugins/Runtime/RigVM/Source/RigVMDeveloper/Public/RigVMModel/RigVMNode.h
--- a/Engine/Plugins/Runtime/RigVM/Source/RigVMDeveloper/Public/RigVMModel/RigVMNode.h
+++ b/Engine/Plugins/Runtime/RigVM/Source/RigVMDeveloper/Public/RigVMModel/RigVMNode.h
@@ -355,7 +355,7 @@ public:
 	UE_API bool HasEarlyExitMarker() const;
 
 	UFUNCTION(BlueprintCallable, Category = RigVMNode)
-	UE_API void SetHasEarlyExitMarker(const bool bValue) { bHasEarlyExitMarker = bValue; }
+	UE_API void SetHasEarlyExitMarker(const bool bValue);
 
 	UFUNCTION(BlueprintCallable, Category = RigVMNode)
 	UE_API bool IsExcludedByEarlyExit() const;
diff --git a/Engine/Plugins/Runtime/RigVM/Source/RigVMDeveloper/Public/RigVMModel/RigVMNotifications.h b/Engine/Plugins/Runtime/RigVM/Source/RigVMDeveloper/Public/RigVMModel/RigVMNotifications.h
--- a/Engine/Plugins/Runtime/RigVM/Source/RigVMDeveloper/Public/RigVMModel/RigVMNotifications.h
+++ b/Engine/Plugins/Runtime/RigVM/Source/RigVMDeveloper/Public/RigVMModel/RigVMNotifications.h
@@ -60,6 +60,7 @@ enum class ERigVMGraphNotifType : uint8
 	PinCategoriesChanged, // The category list of a node has changed - requiring a rebuild of the node user interface (Subject == URigVMNode)
 	PinCategoryExpansionChanged, // The category of a pin expanded / collapsed (Subject == URigVMNode)
 	FunctionVariantGuidChanged, // The guid for a function has changed (Subject == URigVMLibraryNode)
+	NodeEarlyExitChanged, // The node's early exit state has changed
 	Invalid // The max for this enum (used for guarding)
 };
 
diff --git a/Engine/Plugins/Runtime/RigVM/Source/RigVMEditor/Private/Editor/RigVMEditor.cpp b/Engine/Plugins/Runtime/RigVM/Source/RigVMEditor/Private/Editor/RigVMEditor.cpp
--- a/Engine/Plugins/Runtime/RigVM/Source/RigVMEditor/Private/Editor/RigVMEditor.cpp
+++ b/Engine/Plugins/Runtime/RigVM/Source/RigVMEditor/Private/Editor/RigVMEditor.cpp
@@ -1440,6 +1440,41 @@ void FRigVMEditorBase::BindCommands()
 		EditorCommands.ToggleProfiling,
 		FExecuteAction::CreateSP(ThisRef, &FRigVMEditorBase::ToggleProfiling),
 		FCanExecuteAction());
+
+	ToolkitCommands->MapAction(
+		EditorCommands.TogglePreviewHere,
+		FExecuteAction::CreateLambda([ThisRef]()
+		{
+			if (const FRigVMAssetInterfacePtr RigVMAsset = ThisRef->GetRigVMAssetInterface())
+			{
+				if (const FRigVMClient* Client = RigVMAsset->GetRigVMClient())
+				{
+					if (const URigVMGraph* Graph = Client->GetModel(ThisRef->GetFocusedGraph()))
+					{
+						RigVMAsset->TogglePreviewHere(Graph);
+					}
+				}
+			}
+		}),
+		FCanExecuteAction());
+
+	ToolkitCommands->MapAction(
+		EditorCommands.PreviewHereStepForward,
+		FExecuteAction::CreateLambda([ThisRef]()
+		{
+			if (const FRigVMAssetInterfacePtr RigVMAsset = ThisRef->GetRigVMAssetInterface())
+			{
+				RigVMAsset->PreviewHereStepForward();
+			}
+		}),
+		FCanExecuteAction::CreateLambda([ThisRef]() -> bool
+		{
+			if (const FRigVMAssetInterfacePtr RigVMAsset = ThisRef->GetRigVMAssetInterface())
+			{
+				return RigVMAsset->CanPreviewHereStepForward();
+			}
+			return false;
+		}));
 }
 
 void FRigVMEditorBase::ToggleAutoCompileGraph()
diff --git a/Engine/Plugins/Runtime/RigVM/Source/RigVMEditor/Private/Editor/RigVMEditorCommands.cpp b/Engine/Plugins/Runtime/RigVM/Source/RigVMEditor/Private/Editor/RigVMEditorCommands.cpp
--- a/Engine/Plugins/Runtime/RigVM/Source/RigVMEditor/Private/Editor/RigVMEditorCommands.cpp
+++ b/Engine/Plugins/Runtime/RigVM/Source/RigVMEditor/Private/Editor/RigVMEditorCommands.cpp
@@ -4,6 +4,11 @@
 
 #define LOCTEXT_NAMESPACE "RigVMEditorCommands"
 
+TSharedPtr<const FRigVMEditorCommands> FRigVMEditorCommands::GetSharedPtr()
+{
+	return GetInstance().Pin();
+}
+
 void FRigVMEditorCommands::RegisterCommands()
 {
 	UI_COMMAND(Compile, "Compile", "Compile the blueprint", EUserInterfaceActionType::Button, FInputChord());
@@ -26,6 +31,8 @@ void FRigVMEditorCommands::RegisterCommands()
 	UI_COMMAND(SwapAssetReferences, "Swap Asset References", "Swaps an asset reference for another asset.", EUserInterfaceActionType::Button, FInputChord());
 
 	UI_COMMAND(ToggleProfiling, "Toggle Profiling", "Enables or disables heat map profiling for the graph.", EUserInterfaceActionType::ToggleButton, FInputChord(EKeys::P));
+	UI_COMMAND(TogglePreviewHere, "Toggle Preview Here", "Run the graph up to this node", EUserInterfaceActionType::Button, FInputChord(EKeys::F9));
+	UI_COMMAND(PreviewHereStepForward, "Step Forwards", "Run the graph up to the next node.", EUserInterfaceActionType::Button, FInputChord(EKeys::F10));
 }
 
 #undef LOCTEXT_NAMESPACE
diff --git a/Engine/Plugins/Runtime/RigVM/Source/RigVMEditor/Private/Editor/RigVMEditorCommands.h b/Engine/Plugins/Runtime/RigVM/Source/RigVMEditor/Private/Editor/RigVMEditorCommands.h
--- a/Engine/Plugins/Runtime/RigVM/Source/RigVMEditor/Private/Editor/RigVMEditorCommands.h
+++ b/Engine/Plugins/Runtime/RigVM/Source/RigVMEditor/Private/Editor/RigVMEditorCommands.h
@@ -18,6 +18,8 @@ public:
 	)
 	{}
 
+	static TSharedPtr<const FRigVMEditorCommands> GetSharedPtr();
+
 	/** Compile the blueprint */
 	TSharedPtr<FUICommandInfo> Compile;
 	TSharedPtr<FUICommandInfo> SaveOnCompile_Never;
@@ -57,6 +59,12 @@ public:
 	/** Toggle heatmap profiling for the graph */
 	TSharedPtr<FUICommandInfo> ToggleProfiling;
 
+	/** Toggle the preview here state of the selected node */
+	TSharedPtr< FUICommandInfo > TogglePreviewHere;
+
+	/** Step to next node during preview */
+	TSharedPtr< FUICommandInfo > PreviewHereStepForward;
+
 	/**
 	 * Initialize commands
 	 */
diff --git a/Engine/Plugins/Runtime/RigVM/Source/RigVMEditor/Private/RigVMEditorModule.cpp b/Engine/Plugins/Runtime/RigVM/Source/RigVMEditor/Private/RigVMEditorModule.cpp
--- a/Engine/Plugins/Runtime/RigVM/Source/RigVMEditor/Private/RigVMEditorModule.cpp
+++ b/Engine/Plugins/Runtime/RigVM/Source/RigVMEditor/Private/RigVMEditorModule.cpp
@@ -43,6 +43,7 @@
 #include "Editor/RigVMEditorTools.h"
 #include "Editor/RigVMVariantDetailCustomization.h"
 #include "UObject/UObjectIterator.h"
+#include "Toolkits/ToolkitManager.h"
 #include "ToolMenus.h"
 #include "Widgets/SRigVMSwapFunctionsWidget.h"
 #include "Widgets/SRigVMBulkEditDialog.h"
@@ -907,56 +908,56 @@ void FRigVMEditorModule::GetNodeConversionContextMenuActions(TScriptInterface<IR
 
 void FRigVMEditorModule::GetNodeDebugContextMenuActions(TScriptInterface<IRigVMClientHost> RigVMClientHost, const URigVMEdGraphNode* EdGraphNode, URigVMNode* ModelNode, UToolMenu* Menu) const
 {
-	if(FRigVMAssetInterfacePtr RigVMBlueprint = RigVMClientHost.GetObject())
+	if(FRigVMAssetInterfacePtr RigVMAsset = RigVMClientHost.GetObject())
 	{
 		FToolMenuSection& DebugSection = Menu->AddSection("RigVMEditorContextMenuDebug", LOCTEXT("DebugHeader", "Debug"));
 		bool bAnyNodeIsEarlyExit = false;
 
 		const URigVMGraph* Model = ModelNode->GetGraph();
 
-		TArray<URigVMNode*> SelectedNodes;
-		TArray<FName> SelectedNodeNames = Model->GetSelectNodes();
-		SelectedNodeNames.AddUnique(ModelNode->GetFName());
-
-		for (FName SelectedNodeName : SelectedNodeNames)
-		{
-			if (URigVMNode* FoundNode = Model->FindNodeByName(SelectedNodeName))
+		TSharedPtr<FUICommandList> PreviewHereCommandList = MakeShared<FUICommandList>();
+		
+		PreviewHereCommandList->MapAction(
+			FRigVMEditorCommands::Get().TogglePreviewHere,
+			FExecuteAction::CreateLambda([RigVMAsset, Model]()
 			{
-				SelectedNodes.Add(FoundNode);
-				if (FoundNode->HasEarlyExitMarker())
+				if (const FRigVMClient* Client = RigVMAsset->GetRigVMClient())
 				{
-					bAnyNodeIsEarlyExit = true;
+					RigVMAsset->TogglePreviewHere(Model);
 				}
-			}
-		}
-	
-		if (SelectedNodeNames.Num() == 1)
+			}),
+			FCanExecuteAction()
+		);
+
+		PreviewHereCommandList->MapAction(
+			FRigVMEditorCommands::Get().PreviewHereStepForward,
+			FExecuteAction::CreateLambda([RigVMAsset]()
+			{
+				RigVMAsset->PreviewHereStepForward();
+			}),
+			FCanExecuteAction::CreateLambda([RigVMAsset]() -> bool
+			{
+				return RigVMAsset->CanPreviewHereStepForward();
+			})
+		);
+		
+		if (ModelNode->HasEarlyExitMarker())
 		{
 			DebugSection.AddMenuEntry(
-			"Preview Here",
-			LOCTEXT("PreviewHere", "Preview Here"),
-			LOCTEXT("PreviewHere_Tooltip", "Preview the result of the node here"),
-			FSlateIcon(),
-			FUIAction(FExecuteAction::CreateLambda([SelectedNodes, RigVMBlueprint]()
-			{
-				for (URigVMNode* SelectedNode : SelectedNodes)
-				{
-					RigVMBlueprint->SetEarlyExitInstruction(SelectedNode);
-					break;
-				}
-			})));
+				FRigVMEditorCommands::Get().TogglePreviewHere,
+				LOCTEXT("StopPreviewHere", "Stop Preview"),
+				LOCTEXT("StopPreviewHere_Tooltip", "Stops the preview of this node and run the full graph")
+			).SetCommandList(PreviewHereCommandList);
+			
+			DebugSection.AddMenuEntry(FRigVMEditorCommands::Get().PreviewHereStepForward).SetCommandList(PreviewHereCommandList);
 		}
-		else if (bAnyNodeIsEarlyExit)
-		{						
+		else
+		{
 			DebugSection.AddMenuEntry(
-			"Stop Preview",
-			LOCTEXT("StopPreview", "Stop Preview Here"),
-			LOCTEXT("StopPreview_Tooltip", "Stops the preview from here functionality and runs the whole graph"),
-			FSlateIcon(),
-			FUIAction(FExecuteAction::CreateLambda([RigVMBlueprint]()
-			{
-					RigVMBlueprint->ResetEarlyExitInstruction();
-			})));
+				FRigVMEditorCommands::Get().TogglePreviewHere,
+				LOCTEXT("PreviewHere", "Preview Here"),
+				LOCTEXT("PreviewHere_Tooltip", "Run the graph up to this node")
+			).SetCommandList(PreviewHereCommandList);
 		}
 	}
 }
diff --git a/Engine/Plugins/Runtime/RigVM/Source/RigVMEditor/Private/Widgets/SRigVMExecutionStackView.cpp b/Engine/Plugins/Runtime/RigVM/Source/RigVMEditor/Private/Widgets/SRigVMExecutionStackView.cpp
--- a/Engine/Plugins/Runtime/RigVM/Source/RigVMEditor/Private/Widgets/SRigVMExecutionStackView.cpp
+++ b/Engine/Plugins/Runtime/RigVM/Source/RigVMEditor/Private/Widgets/SRigVMExecutionStackView.cpp
@@ -17,6 +17,7 @@
 #include "RigVMModel/Nodes/RigVMAggregateNode.h"
 #include "Widgets/Input/SSearchBox.h"
 #include "Dialog/SCustomDialog.h"
+#include "Async/TaskGraphInterfaces.h"
 
 #define LOCTEXT_NAMESPACE "SRigVMExecutionStackView"
 
@@ -56,8 +57,6 @@ void SRigStackItem::Construct(const FArguments& InArgs, const TSharedRef<STableV
 	WeakStackEntry = InStackEntry;
 	WeakBlueprint = InBlueprint;
 	WeakCommandList = InCommandList;
-	MicroSeconds = 0;
-	MicroSecondsFrames.Reset();
 
 	TSharedPtr< STextBlock > NumberWidget;
 	TSharedPtr< STextBlock > TextWidget;
@@ -110,6 +109,7 @@ void SRigStackItem::Construct(const FArguments& InArgs, const TSharedRef<STableV
 				[
 					SAssignNew(NumberWidget, STextBlock)
 					.Text(this, &SRigStackItem::GetIndexText)
+					.ColorAndOpacity(this, &SRigStackItem::GetLabelColorAndOpacity)
 					.Font(IDetailLayoutBuilder::GetDetailFont())
 				]
 				+ SHorizontalBox::Slot()
@@ -129,7 +129,8 @@ void SRigStackItem::Construct(const FArguments& InArgs, const TSharedRef<STableV
 				[
 					SAssignNew(TextWidget, STextBlock)
 					.Text(this, &SRigStackItem::GetLabelText)
-					.Font(this, &SRigStackItem::GetLabelFont)
+					.Font(IDetailLayoutBuilder::GetDetailFont())
+					.ColorAndOpacity(this, &SRigStackItem::GetLabelColorAndOpacity)
 					.ToolTipText(this, &SRigStackItem::GetTooltip)
 					.Justification(ETextJustify::Left)
 				]
@@ -154,6 +155,7 @@ void SRigStackItem::Construct(const FArguments& InArgs, const TSharedRef<STableV
 	            [
 	                SNew(STextBlock)
 	                .Text(this, &SRigStackItem::GetVisitedCountText)
+					.ColorAndOpacity(this, &SRigStackItem::GetLabelColorAndOpacity)
 	                .Font(IDetailLayoutBuilder::GetDetailFont())
 	            ]
 	            
@@ -165,6 +167,7 @@ void SRigStackItem::Construct(const FArguments& InArgs, const TSharedRef<STableV
 	            [
 	                SNew(STextBlock)
 	                .Text(this, &SRigStackItem::GetDurationText)
+					.ColorAndOpacity(this, &SRigStackItem::GetLabelColorAndOpacity)
 	                .Font(IDetailLayoutBuilder::GetDetailFont())
 	            ]
 	        ]
@@ -182,13 +185,16 @@ FText SRigStackItem::GetLabelText() const
 	return (FText::FromString(WeakStackEntry.Pin()->Label));
 }
 
-FSlateFontInfo SRigStackItem::GetLabelFont() const
+FSlateColor SRigStackItem::GetLabelColorAndOpacity() const
 {
-	if(GetVisitedCountText().IsEmpty())
+	if (TSharedPtr<FRigStackEntry> StackEntry = WeakStackEntry.Pin())
 	{
-		return IDetailLayoutBuilder::GetDetailFont();
+		if(StackEntry->ShowAsFadedOut.Get(false))
+		{
+			return FSlateColor::UseSubduedForeground();
+		}
 	}
-	return IDetailLayoutBuilder::GetDetailFontBold();
+	return FSlateColor::UseForeground();
 }
 
 FText SRigStackItem::GetTooltip() const
@@ -200,29 +206,8 @@ FText SRigStackItem::GetVisitedCountText() const
 {
 	if (TSharedPtr<FRigStackEntry> StackEntry = WeakStackEntry.Pin())
 	{
-		if (WeakBlueprint.IsValid())
-		{
-			IRigVMAssetInterface* Blueprint = WeakBlueprint.Get();
-			if(Blueprint->GetRigGraphDisplaySettings().bShowNodeRunCounts)
-			{
-				if(WeakStackEntry.Pin()->EntryType == ERigStackEntry::Operator)
-				{
-					if(URigVMHost* RigVMHost = Blueprint->GetDebuggedRigVMHost())
-					{
-						if(URigVM* VM = RigVMHost->GetVM())
-						{
-							const int32 Count = VM->GetInstructionVisitedCount(RigVMHost->GetRigVMExtendedExecuteContext(), WeakStackEntry.Pin()->InstructionIndex);
-							if(Count > 0)
-							{
-								return FText::FromString(FString::FromInt(Count));
-							}
-						}
-					}
-				}
-			}
-		}
+		return StackEntry->VisitedCountText.Get(FText());
 	}
-
 	return FText();
 }
 
@@ -230,28 +215,7 @@ FText SRigStackItem::GetDurationText() const
 {
 	if (TSharedPtr<FRigStackEntry> StackEntry = WeakStackEntry.Pin())
 	{
-		if (TStrongObjectPtr<UObject> StrongPtr = WeakBlueprint.GetWeakObjectPtr().Pin())
-		{
-			IRigVMAssetInterface* Blueprint = WeakBlueprint.Get();
-			if(Blueprint->GetVMRuntimeSettings().bEnableProfiling)
-			{
-				if(WeakStackEntry.Pin()->EntryType == ERigStackEntry::Operator)
-				{
-					if(URigVMHost* RigVMHost = Blueprint->GetDebuggedRigVMHost())
-					{
-						if(URigVM* VM = RigVMHost->GetVM())
-						{
-							const double CurrentMicroSeconds = VM->GetInstructionMicroSeconds(RigVMHost->GetRigVMExtendedExecuteContext(), WeakStackEntry.Pin()->InstructionIndex);
-							MicroSeconds = Blueprint->GetRigGraphDisplaySettings().AggregateAverage(MicroSecondsFrames, MicroSeconds, CurrentMicroSeconds);
-							if(MicroSeconds > 0.0)
-							{
-								return FText::FromString(FString::Printf(TEXT("%.02f µs"), (float)MicroSeconds));
-							}
-						}
-					}
-				}
-			}
-		}
+		return StackEntry->DurationText.Get(FText());
 	}
 	return FText();
 }
@@ -262,31 +226,8 @@ FText SRigStackItem::GetDurationText() const
 
 SRigVMExecutionStackView::~SRigVMExecutionStackView()
 {
-	if (TSharedPtr<IRigVMEditor> Editor = WeakEditor.Pin())
-	{
-		if (OnModelModified.IsValid() && Editor->GetRigVMAssetInterface())
-		{
-			Editor->GetRigVMAssetInterface()->OnModified().Remove(OnModelModified);
-		}
-		if (OnHostInitializedHandle.IsValid())
-		{
-			if(Editor->GetRigVMHost())
-			{
-				Editor->GetRigVMHost()->OnInitialized_AnyThread().Remove(OnHostInitializedHandle);
-			}
-		}
-		if (OnPreviewHostUpdatedHandle.IsValid())
-		{
-			Editor->OnPreviewHostUpdated().Remove(OnPreviewHostUpdatedHandle);
-		}
-	}
-	if (TStrongObjectPtr<UObject> Blueprint = WeakRigVMBlueprint.GetWeakObjectPtr().Pin())
-	{
-		if (OnVMCompiledHandle.IsValid())
-		{
-			WeakRigVMBlueprint->OnVMCompiled().Remove(OnVMCompiledHandle);
-		}
-	}
+	UnbindHostDelegates();
+	UnbindAssetDelegates();
 }
 
 void SRigVMExecutionStackView::Construct( const FArguments& InArgs, TSharedRef<IRigVMEditor> InRigVMEditor)
@@ -364,15 +305,14 @@ void SRigVMExecutionStackView::Construct( const FArguments& InArgs, TSharedRef<I
 			RigVMBlueprint->OnModified().Remove(OnModelModified);
 		}
 		OnVMCompiledHandle = RigVMBlueprint->OnVMCompiled().AddSP(this, &SRigVMExecutionStackView::OnVMCompiled);
+		OnObjectBeingDebuggedChangedHandle = WeakRigVMBlueprint->OnSetObjectBeingDebugged().AddSP(this, &SRigVMExecutionStackView::OnObjectBeingDebuggedChanged);
+
 		OnModelModified = RigVMBlueprint->OnModified().AddSP(this, &SRigVMExecutionStackView::HandleModifiedEvent);
 
-		if (TSharedPtr<IRigVMEditor> Editor = WeakEditor.Pin())
-		{
-			if (URigVMHost* RigVMHost = Editor->GetRigVMHost())
-			{
-				OnVMCompiled(RigVMBlueprint->GetObject(),  RigVMHost->GetVM(),  RigVMHost->GetRigVMExtendedExecuteContext());
-			}
-		}
+		BindHostDelegates();
+		
+		URigVMHost* RigVMHost = PreviouslyDebuggedHost.Get();
+		OnObjectBeingDebuggedChanged(RigVMHost);
 	}
 
 	if (TSharedPtr<IRigVMEditor> Editor = WeakEditor.Pin())
@@ -395,11 +335,7 @@ void SRigVMExecutionStackView::OnSelectionChanged(TSharedPtr<FRigStackEntry> Sel
 		if (TStrongObjectPtr<UObject> StrongPtr = WeakRigVMBlueprint.GetWeakObjectPtr().Pin())
 		{
 			IRigVMAssetInterface* RigVMBlueprint = WeakRigVMBlueprint.Get();
-			URigVMHost* Host = nullptr;
-			if (TSharedPtr<IRigVMEditor> Editor = WeakEditor.Pin())
-			{
-				Host = Editor->GetRigVMHost();
-			}
+			URigVMHost* Host = PreviouslyDebuggedHost.Get();
 			if (Host == nullptr || Host->GetVM() == nullptr)
 			{
 				return;
@@ -741,6 +677,68 @@ void SRigVMExecutionStackView::PopulateStackView(URigVM* InVM, FRigVMExtendedExe
 	}
 }
 
+void SRigVMExecutionStackView::BindHostDelegates()
+{
+	UnbindHostDelegates();
+	
+	if (TStrongObjectPtr<UObject> Blueprint = WeakRigVMBlueprint.GetWeakObjectPtr().Pin())
+	{
+		BindHostDelegates(WeakRigVMBlueprint->GetObjectBeingDebugged());
+	}
+}
+
+void SRigVMExecutionStackView::BindHostDelegates(UObject* InObjectBeingDebugged)
+{
+	UnbindHostDelegates();
+
+	if (URigVMHost* DebuggedHost = Cast<URigVMHost>(InObjectBeingDebugged))
+	{
+		PreviouslyDebuggedHost = DebuggedHost;
+		OnHostInitializedHandle = DebuggedHost->OnInitialized_AnyThread().AddSP(this, &SRigVMExecutionStackView::HandleHostInitializedEvent);
+		OnHostExecutedHandle = DebuggedHost->OnExecuted_AnyThread().AddSP(this, &SRigVMExecutionStackView::HandleHostExecutedEvent);
+	}
+}
+
+void SRigVMExecutionStackView::UnbindHostDelegates()
+{
+	if (URigVMHost* Host = PreviouslyDebuggedHost.Get())
+	{
+		if (OnHostInitializedHandle.IsValid())
+		{
+			Host->OnInitialized_AnyThread().Remove(OnHostInitializedHandle);
+		}
+		if (OnHostExecutedHandle.IsValid())
+		{
+			Host->OnExecuted_AnyThread().Remove(OnHostExecutedHandle);
+		}
+	}
+	OnHostInitializedHandle.Reset();
+	OnHostExecutedHandle.Reset();
+	PreviouslyDebuggedHost.Reset();
+}
+
+void SRigVMExecutionStackView::UnbindAssetDelegates()
+{
+	if (TSharedPtr<IRigVMEditor> Editor = WeakEditor.Pin())
+	{
+		if (OnModelModified.IsValid() && Editor->GetRigVMAssetInterface())
+		{
+			Editor->GetRigVMAssetInterface()->OnModified().Remove(OnModelModified);
+		}
+	}
+	if (TStrongObjectPtr<UObject> Blueprint = WeakRigVMBlueprint.GetWeakObjectPtr().Pin())
+	{
+		if (OnVMCompiledHandle.IsValid())
+		{
+			WeakRigVMBlueprint->OnVMCompiled().Remove(OnVMCompiledHandle);
+		}
+		if (OnObjectBeingDebuggedChangedHandle.IsValid())
+		{
+			WeakRigVMBlueprint->OnSetObjectBeingDebugged().Remove(OnObjectBeingDebuggedChangedHandle);
+		}
+	}
+}
+
 void SRigVMExecutionStackView::RefreshTreeView(URigVM* InVM, FRigVMExtendedExecuteContext* InVMContext)
 {
 	Operators.Reset();
@@ -753,7 +751,7 @@ void SRigVMExecutionStackView::RefreshTreeView(URigVM* InVM, FRigVMExtendedExecu
 		// fill the children from the log
 		if (TSharedPtr<IRigVMEditor> Editor = WeakEditor.Pin())
 		{
-			URigVMHost* Host = Editor->GetRigVMHost();
+			URigVMHost* Host = PreviouslyDebuggedHost.Get();
 			if(Host && Host->GetLog())
 			{
 				const TArray<FRigVMLog::FLogEntry>& LogEntries = Host->GetLog()->Entries;
@@ -844,7 +842,7 @@ void SRigVMExecutionStackView::HandleFocusOnSelectedGraphNode()
 			if (TStrongObjectPtr<UObject> StrongPtr = WeakRigVMBlueprint.GetWeakObjectPtr().Pin())
 			{
 				IRigVMAssetInterface* RigVMBlueprint = WeakRigVMBlueprint.Get();
-				URigVMHost* Host = Editor->GetRigVMHost();
+				URigVMHost* Host = PreviouslyDebuggedHost.Get();
 				if (Host == nullptr || Host->GetVM() == nullptr)
 				{
 					return;
@@ -947,11 +945,7 @@ TArray<TSharedPtr<FRigStackEntry>> SRigVMExecutionStackView::GetTargetItems(cons
 {
 	TArray<TSharedPtr<FRigStackEntry>> TargetItems;
 
-	URigVMHost* Host = nullptr;
-	if (TSharedPtr<IRigVMEditor> Editor = WeakEditor.Pin())
-	{
-		Host = Editor->GetRigVMHost();
-	}
+	URigVMHost* Host = PreviouslyDebuggedHost.Get();
 	if (Host == nullptr || Host->GetVM() == nullptr)
 	{
 		return TargetItems;
@@ -1057,30 +1051,25 @@ void SRigVMExecutionStackView::UpdateTargetItemHighlighting()
 
 void SRigVMExecutionStackView::OnVMCompiled(UObject* InCompiledObject, URigVM* InCompiledVM, FRigVMExtendedExecuteContext& InVMContext)
 {
+	BindHostDelegates();
 	RefreshTreeView(InCompiledVM, &InVMContext);
+}
 
-	if (TSharedPtr<IRigVMEditor> Editor = WeakEditor.Pin())
+void SRigVMExecutionStackView::OnObjectBeingDebuggedChanged(UObject* InObjectBeingDebugged)
+{
+	BindHostDelegates(InObjectBeingDebugged);
+	if (URigVMHost* DebuggedHost = Cast<URigVMHost>(InObjectBeingDebugged))
 	{
-		if (!OnHostInitializedHandle.IsValid())
-		{
-			if(URigVMHost* Host = Editor->GetRigVMHost())
-			{
-				OnHostInitializedHandle = Host->OnInitialized_AnyThread().AddSP(this, &SRigVMExecutionStackView::HandleHostInitializedEvent);
-			}
-		}
+		RefreshTreeView(DebuggedHost->GetVM(), &DebuggedHost->GetRigVMExtendedExecuteContext());
 	}
 }
 
 void SRigVMExecutionStackView::OnFilterTextChanged(const FText& SearchText)
 {
 	FilterText = SearchText;
-	if (TSharedPtr<IRigVMEditor> Editor = WeakEditor.Pin())
+	if (URigVMHost* RigVMHost = PreviouslyDebuggedHost.Get())
 	{
-		URigVMHost* RigVMHost = Editor->GetRigVMHost();
-		if (RigVMHost != nullptr)
-		{
-			RefreshTreeView(RigVMHost->GetVM(), &RigVMHost->GetRigVMExtendedExecuteContext());
-		}
+		RefreshTreeView(RigVMHost->GetVM(), &RigVMHost->GetRigVMExtendedExecuteContext());
 	}
 }
 
@@ -1091,11 +1080,7 @@ void SRigVMExecutionStackView::HandleModifiedEvent(ERigVMGraphNotifType InNotifT
 		return;
 	}
 
-	URigVMHost* Host = nullptr;
-	if (TSharedPtr<IRigVMEditor> Editor = WeakEditor.Pin())
-	{
-		Host = Editor->GetRigVMHost();
-	}
+	URigVMHost* Host = PreviouslyDebuggedHost.Get();
 	if (Host == nullptr || Host->GetVM() == nullptr)
 	{
 		return;
@@ -1150,48 +1135,134 @@ void SRigVMExecutionStackView::HandleModifiedEvent(ERigVMGraphNotifType InNotifT
 
 void SRigVMExecutionStackView::HandleHostInitializedEvent(URigVMHost* InHost, const FName& InEventName)
 {
-	TGuardValue<bool> SuspendControllerSelection(bSuspendControllerSelection, true);
+	auto HandleHostInitializedEventLambda = [this, InHost, InEventName]
+	{
+		TGuardValue<bool> SuspendControllerSelection(bSuspendControllerSelection, true);
 
-	RefreshTreeView(InHost->GetVM(), &InHost->GetRigVMExtendedExecuteContext());
-	OnSelectionChanged(TSharedPtr<FRigStackEntry>(), ESelectInfo::Direct);
+		RefreshTreeView(InHost->GetVM(), &InHost->GetRigVMExtendedExecuteContext());
+		OnSelectionChanged(TSharedPtr<FRigStackEntry>(), ESelectInfo::Direct);
 
-	for (TSharedPtr<FRigStackEntry>& Operator : Operators)
-	{
-		for (TSharedPtr<FRigStackEntry>& Child : Operator->Children)
+		for (TSharedPtr<FRigStackEntry>& Operator : Operators)
 		{
-			if (Child->EntryType == ERigStackEntry::Warning || Child->EntryType == ERigStackEntry::Error)
+			for (TSharedPtr<FRigStackEntry>& Child : Operator->Children)
 			{
-				TreeView->SetItemExpansion(Operator, true);
-				break;
+				if (Child->EntryType == ERigStackEntry::Warning || Child->EntryType == ERigStackEntry::Error)
+				{
+					TreeView->SetItemExpansion(Operator, true);
+					break;
+				}
 			}
 		}
-	}
+
+		if (TSharedPtr<IRigVMEditor> Editor = WeakEditor.Pin())
+		{
+			InHost->OnInitialized_AnyThread().Remove(OnHostInitializedHandle);
+			OnHostInitializedHandle.Reset();
+
+			if (FRigVMAssetInterfacePtr RigBlueprint = Editor->GetRigVMAssetInterface())
+			{
+				TArray<URigVMGraph*> Models = RigBlueprint->GetAllModels();
+				for (URigVMGraph* Model : Models)
+				{
+					for (URigVMNode* ModelNode : Model->GetNodes())
+					{
+						if (ModelNode->IsSelected())
+						{
+							HandleModifiedEvent(ERigVMGraphNotifType::NodeSelected, Model, ModelNode);
+						}
+					}
+				}
+			}
+		}
+	};
 	
-	if (TSharedPtr<IRigVMEditor> Editor = WeakEditor.Pin())
+	if(IsInGameThread())
+	{
+		HandleHostInitializedEventLambda();
+	}
+	else
+	{
+		FFunctionGraphTask::CreateAndDispatchWhenReady(HandleHostInitializedEventLambda, TStatId(), NULL, ENamedThreads::GameThread);
+	}
+}
+
+void SRigVMExecutionStackView::HandleHostExecutedEvent(URigVMHost* InHost, const FName& InEventName)
+{
+	auto HandleHostExecutedEventLambda = [this, InHost, InEventName]()
 	{
-		InHost->OnInitialized_AnyThread().Remove(OnHostInitializedHandle);
-		OnHostInitializedHandle.Reset();
+		for(const TSharedPtr<FRigStackEntry>& Entry : Operators)
+		{
+			Entry->ShowAsFadedOut.Reset();
+			Entry->VisitedCountText.Reset();
+			Entry->DurationText.Reset();
+		}
 
-		if (FRigVMAssetInterfacePtr RigBlueprint = Editor->GetRigVMAssetInterface())
+		if (WeakRigVMBlueprint.IsValid())
 		{
-			TArray<URigVMGraph*> Models = RigBlueprint->GetAllModels();
-			for (URigVMGraph* Model : Models)
+			IRigVMAssetInterface* Blueprint = WeakRigVMBlueprint.Get();
+			check(Blueprint);
+			
+			if(Blueprint->GetRigGraphDisplaySettings().bShowNodeRunCounts)
+			{
+				if(URigVM* VM = InHost->GetVM())
+				{
+					for(TSharedPtr<FRigStackEntry>& Entry : Operators)
+					{
+						if(Entry->EntryType == ERigStackEntry::Operator)
+						{
+							const int32 Count = VM->GetInstructionVisitedCount(InHost->GetRigVMExtendedExecuteContext(), Entry->InstructionIndex);
+							if(Count > 0)
+							{
+								Entry->VisitedCountText = FText::FromString(FString::FromInt(Count));
+							}
+							else
+							{
+								Entry->ShowAsFadedOut = true;
+							}
+						}
+					}
+				}
+			}
+
+			if(Blueprint->GetVMRuntimeSettings().bEnableProfiling)
 			{
-				for (URigVMNode* ModelNode : Model->GetNodes())
+				if(URigVM* VM = InHost->GetVM())
 				{
-					if (ModelNode->IsSelected())
+					for(TSharedPtr<FRigStackEntry>& Entry : Operators)
 					{
-						HandleModifiedEvent(ERigVMGraphNotifType::NodeSelected, Model, ModelNode);
+						if(Entry->EntryType == ERigStackEntry::Operator)
+						{
+							const double CurrentMicroSeconds = VM->GetInstructionMicroSeconds(InHost->GetRigVMExtendedExecuteContext(), Entry->InstructionIndex);
+							Entry->MicroSeconds = Blueprint->GetRigGraphDisplaySettings().AggregateAverage(Entry->MicroSecondsFrames, Entry->MicroSeconds, CurrentMicroSeconds);
+							if(Entry->MicroSeconds > 0.0)
+							{
+								Entry->DurationText = FText::FromString(FString::Printf(TEXT("%.02f µs"), (float)Entry->MicroSeconds));
+							}
+							else
+							{
+								Entry->ShowAsFadedOut = true;
+							}
+						}
 					}
 				}
 			}
 		}
+	};
+
+	if(IsInGameThread())
+	{
+		HandleHostExecutedEventLambda();
 	}
+	else
+	{
+		FFunctionGraphTask::CreateAndDispatchWhenReady(HandleHostExecutedEventLambda, TStatId(), NULL, ENamedThreads::GameThread);
+	}
+
 }
 
 void SRigVMExecutionStackView::HandlePreviewHostUpdated(IRigVMEditor* InEditor)
 {
-	if(URigVMHost* RigVMHost = InEditor->GetRigVMHost())
+	if (URigVMHost* RigVMHost = PreviouslyDebuggedHost.Get())
 	{
 		RefreshTreeView(RigVMHost->GetVM(), &RigVMHost->GetRigVMExtendedExecuteContext());
 	}
@@ -1204,7 +1275,7 @@ void SRigVMExecutionStackView::HandleItemMouseDoubleClick(TSharedPtr<FRigStackEn
 		if (TStrongObjectPtr<UObject> StrongPtr = WeakRigVMBlueprint.GetWeakObjectPtr().Pin())
 		{
 			IRigVMAssetInterface* RigVMBlueprint = WeakRigVMBlueprint.Get();
-			URigVMHost* Host = RigVMBlueprint->GetDebuggedRigVMHost();
+			URigVMHost* Host = PreviouslyDebuggedHost.Get();
 			if (!Host || !Host->GetVM())
 			{
 				return;
diff --git a/Engine/Plugins/Runtime/RigVM/Source/RigVMEditor/Private/Widgets/SRigVMGraphNode.cpp b/Engine/Plugins/Runtime/RigVM/Source/RigVMEditor/Private/Widgets/SRigVMGraphNode.cpp
--- a/Engine/Plugins/Runtime/RigVM/Source/RigVMEditor/Private/Widgets/SRigVMGraphNode.cpp
+++ b/Engine/Plugins/Runtime/RigVM/Source/RigVMEditor/Private/Widgets/SRigVMGraphNode.cpp
@@ -314,6 +314,8 @@ void SRigVMGraphNode::Construct( const FArguments& InArgs )
 	NodeErrorType = int32(EMessageSeverity::Info) + 1;
 	this->UpdateGraphNode();
 
+	bPreviewHereIndicatorHovered = false;
+	
 	SetIsEditable(false);
 
 	Controller->OnModified().AddSP(this, &SRigVMGraphNode::HandleModifiedEvent);
@@ -322,46 +324,51 @@ void SRigVMGraphNode::Construct( const FArguments& InArgs )
 
 	const FSlateBrush* ImageBrush = FSlateIcon(TEXT("RigVMEditorStyle"), TEXT("RigVM.Bug.Dot")).GetIcon();
 
-	PreviewHereIndicatorWidget =
-		SNew(SButton)
-		.Visibility(EVisibility::Visible)
-		.ButtonStyle( FAppStyle::Get(), "SimpleButton" )
-		.ToolTipText(LOCTEXT("PreviewHereTooltip", "Toggles the Preview Here functionality.\nBy pressing this you can skip the execution of consecutive nodes to preview the results of the graph up to this point."))
-		.ContentPadding(0.0f)
-		.OnHovered_Lambda([this]
-		{
-			bPreviewHereIndicatorHovered = true;
-		})
-		.OnUnhovered_Lambda([this]
-		{
-			bPreviewHereIndicatorHovered = false;
-		})
-		.OnClicked_Lambda([this]() -> FReply
+	if (const URigVMEdGraphSchema* Schema = Cast<URigVMEdGraphSchema>(EdGraphNode->GetSchema()))
+	{
+		if (Schema->SupportsPreviewFromHereOnNode(EdGraphNode))
 		{
-			if (ModelNode.IsValid() && Blueprint.IsValid())
-			{
-				if (ModelNode->HasEarlyExitMarker())
+			PreviewHereIndicatorWidget =
+				SNew(SButton)
+				.Visibility(EVisibility::Visible)
+				.ButtonStyle( FAppStyle::Get(), "SimpleButton" )
+				.ToolTipText(LOCTEXT("PreviewHereTooltip", "Run the graph up to this node [F9]"))
+				.ContentPadding(0.0f)
+				.OnHovered_Lambda([this]
 				{
-					Blueprint->ResetEarlyExitInstruction();
-				}
-				else
+					bPreviewHereIndicatorHovered = true;
+				})
+				.OnUnhovered_Lambda([this]
 				{
-					Blueprint->SetEarlyExitInstruction(ModelNode.Get());
-				}
-				return FReply::Handled();
-			}
-			return FReply::Unhandled();
-		})
-		.Content()
-		[
-			SNew(SImage)
-			.Image(FAppStyle::GetBrush(TEXT("Icons.Preview")))
-			.DesiredSizeOverride(FVector2D(16, 16))
-			.ColorAndOpacity(this, &SRigVMGraphNode::GetPreviewHereColor)
-			.Visibility(EVisibility::Visible)
-		];
-	bPreviewHereIndicatorHovered = false;
-
+					bPreviewHereIndicatorHovered = false;
+				})
+				.OnClicked_Lambda([this, EdGraphNode]() -> FReply
+				{
+					if (ModelNode.IsValid() && Blueprint.IsValid())
+					{
+						if (ModelNode->HasEarlyExitMarker())
+						{
+							Blueprint->ResetEarlyExitInstruction();
+						}
+						else
+						{
+							Blueprint->SetEarlyExitInstruction(ModelNode.Get());
+						}
+						EdGraphNode->InvalidateNodeTitle();
+						return FReply::Handled();
+					}
+					return FReply::Unhandled();
+				})
+				.Content()
+				[
+					SNew(SImage)
+					.Image(FAppStyle::GetBrush(TEXT("Icons.Preview")))
+					.DesiredSizeOverride(FVector2D(10, 10))
+					.ColorAndOpacity(this, &SRigVMGraphNode::GetPreviewHereColor)
+					.Visibility(EVisibility::Visible)
+				];
+		}
+	}
 	VisualDebugIndicatorWidget =
 		SNew(SImage)
 		.Image(ImageBrush)
@@ -1279,12 +1286,12 @@ TArray<FOverlayWidgetInfo> SRigVMGraphNode::GetOverlayWidgets(bool bSelected, co
 	{
 		float HorizontalOffset = 0;
 
-		if (ModelNode->IsMutable())
+		if (ModelNode->IsMutable() && PreviewHereIndicatorWidget.IsValid())
 		{
 			const FVector2f ImageSize = PreviewHereIndicatorWidget->GetDesiredSize();
 
 			FOverlayWidgetInfo Info;
-			Info.OverlayOffset = FVector2f(WidgetSize.X - ImageSize.X - 6.f - HorizontalOffset, 4.f);
+			Info.OverlayOffset = FVector2f(WidgetSize.X - ImageSize.X - 6.f - HorizontalOffset, 8.f);
 			HorizontalOffset += ImageSize.X + 6.f;
 			Info.Widget = PreviewHereIndicatorWidget;
 
@@ -2984,11 +2991,11 @@ FSlateColor SRigVMGraphNode::GetPreviewHereColor() const
 	{
 		if (CurrentModelNode->HasEarlyExitMarker())
 		{
-			return FSlateColor(EStyleColor::AccentBlue);
+			return FSlateColor(EStyleColor::AccentOrange);
 		}
 		else if (CurrentModelNode->IsExcludedByEarlyExit())
 		{
-			return FSlateColor(EStyleColor::AccentBlack);
+			return FSlateColor(EStyleColor::SelectInactive);
 		}
 	}
 	if (bPreviewHereIndicatorHovered)
diff --git a/Engine/Plugins/Runtime/RigVM/Source/RigVMEditor/Public/Widgets/SRigVMExecutionStackView.h b/Engine/Plugins/Runtime/RigVM/Source/RigVMEditor/Public/Widgets/SRigVMExecutionStackView.h
--- a/Engine/Plugins/Runtime/RigVM/Source/RigVMEditor/Public/Widgets/SRigVMExecutionStackView.h
+++ b/Engine/Plugins/Runtime/RigVM/Source/RigVMEditor/Public/Widgets/SRigVMExecutionStackView.h
@@ -47,6 +47,11 @@ public:
 	ERigVMOpCode OpCode;
 	FString Label;
 	TArray<TSharedPtr<FRigStackEntry>> Children;
+	mutable TOptional<bool> ShowAsFadedOut;
+	mutable TOptional<FText> VisitedCountText;
+	mutable TOptional<FText> DurationText;
+	mutable double MicroSeconds;
+	mutable TArray<double> MicroSecondsFrames;
 };
 
 class SRigStackItem : public STableRow<TSharedPtr<FRigStackEntry>>
@@ -61,12 +66,9 @@ private:
 	TWeakInterfacePtr<IRigVMAssetInterface> WeakBlueprint;
 	TWeakPtr<FUICommandList> WeakCommandList;
 
-	mutable double MicroSeconds;
-	mutable TArray<double> MicroSecondsFrames;
-
 	FText GetIndexText() const;
 	FText GetLabelText() const;
-	FSlateFontInfo GetLabelFont() const;
+	FSlateColor GetLabelColorAndOpacity() const;
 	FText GetTooltip() const;
 	FText GetVisitedCountText() const;
 	FText GetDurationText() const;
@@ -117,6 +119,7 @@ private:
 	UE_API void UpdateTargetItemHighlighting();
 
 	UE_API void OnVMCompiled(UObject* InCompiledObject, URigVM* InCompiledVM, FRigVMExtendedExecuteContext& InVMContext);
+	UE_API void OnObjectBeingDebuggedChanged(UObject* InObjectBeingDebugged);
 
 	/** Search box widget */
 	TSharedPtr<SSearchBox> FilterBox;
@@ -128,12 +131,18 @@ private:
 	bool bSuspendControllerSelection;
 	UE_API void HandleModifiedEvent(ERigVMGraphNotifType InNotifType, URigVMGraph* InGraph, UObject* InSubject);
 	UE_API void HandleHostInitializedEvent(URigVMHost* InHost, const FName& InEventName);
+	UE_API void HandleHostExecutedEvent(URigVMHost* InHost, const FName& InEventName);
 	UE_API void HandlePreviewHostUpdated(IRigVMEditor* InEditor);
 	UE_API void HandleItemMouseDoubleClick(TSharedPtr<FRigStackEntry> InItem);
 
 	/** Populate the execution stack with descriptive names for each instruction */
 	UE_API void PopulateStackView(URigVM* InVM, FRigVMExtendedExecuteContext* InVMContext);
 
+	void BindHostDelegates();
+	void BindHostDelegates(UObject* InObjectBeingDebugged);
+	void UnbindHostDelegates();
+	void UnbindAssetDelegates();
+
 	TSharedPtr<STreeView<TSharedPtr<FRigStackEntry>>> TreeView;
 
 	/** Command list we bind to */
@@ -146,8 +155,11 @@ private:
 
 	FDelegateHandle OnModelModified;
 	FDelegateHandle OnHostInitializedHandle;
+	FDelegateHandle OnHostExecutedHandle;
 	FDelegateHandle OnVMCompiledHandle;
 	FDelegateHandle OnPreviewHostUpdatedHandle;
+	FDelegateHandle OnObjectBeingDebuggedChangedHandle;
+	TWeakObjectPtr<URigVMHost> PreviouslyDebuggedHost;
 };
 
 #undef UE_API
diff --git a/Engine/Source/Editor/MaterialEditor/Private/MaterialNodes/SGraphNodeMaterialBase.cpp b/Engine/Source/Editor/MaterialEditor/Private/MaterialNodes/SGraphNodeMaterialBase.cpp
--- a/Engine/Source/Editor/MaterialEditor/Private/MaterialNodes/SGraphNodeMaterialBase.cpp
+++ b/Engine/Source/Editor/MaterialEditor/Private/MaterialNodes/SGraphNodeMaterialBase.cpp
@@ -570,6 +570,7 @@ TSharedRef<SWidget> SGraphNodeMaterialBase::CreateTitleRightWidget()
 			.HAlign(HAlign_Center)
 			[
 				SNew(SImage)
+				.ColorAndOpacity(FSlateColor::UseForeground())
 				.Image(FAppStyle::GetBrush(TEXT("Icons.Preview")))
 			]
 		]
