RigVM: Fix uninitialized member

#rnx

diff --git a/Engine/Plugins/Runtime/RigVM/Source/RigVMEditor/Private/Widgets/SRigVMExecutionStackView.cpp b/Engine/Plugins/Runtime/RigVM/Source/RigVMEditor/Private/Widgets/SRigVMExecutionStackView.cpp
--- a/Engine/Plugins/Runtime/RigVM/Source/RigVMEditor/Private/Widgets/SRigVMExecutionStackView.cpp
+++ b/Engine/Plugins/Runtime/RigVM/Source/RigVMEditor/Private/Widgets/SRigVMExecutionStackView.cpp
@@ -34,6 +34,7 @@ FRigStackEntry::FRigStackEntry(int32 InEntryIndex, ERigStackEntry::Type InEntryT
 	, Callstack(InProxy.GetCallstack())
 	, OpCode(InOpCode)
 	, Label(InLabel)
+	, MicroSeconds(0.0)
 {
 
 }
