RigVM: Fix undo redo for new wire slicing tool. Solidfy lookup of controller

#rnx

diff --git a/Engine/Plugins/Runtime/RigVM/Source/RigVMDeveloper/Private/RigVMModel/RigVMControllerActions.cpp b/Engine/Plugins/Runtime/RigVM/Source/RigVMDeveloper/Private/RigVMModel/RigVMControllerActions.cpp
--- a/Engine/Plugins/Runtime/RigVM/Source/RigVMDeveloper/Private/RigVMModel/RigVMControllerActions.cpp
+++ b/Engine/Plugins/Runtime/RigVM/Source/RigVMDeveloper/Private/RigVMModel/RigVMControllerActions.cpp
@@ -268,16 +268,54 @@ void URigVMActionStack::LogAction(const UScriptStruct* InActionStruct, const FRi
 #endif
 
 
+FRigVMBaseAction::FRigVMBaseAction(URigVMController* InController)
+: WeakController(InController)
+, Title(TEXT("Action"))
+{
+	if(InController)
+	{
+		ControllerPath = InController;
+
+		if (IRigVMClientHost* ClientHost = InController->GetClientHost())
+		{
+			if (FRigVMClient* Client = ClientHost->GetRigVMClient())
+			{
+				if (UObject* ClientOuter = Client->GetOuter())
+				{
+					ClientHostPath = ClientOuter;
+					ControllerModelPath = InController->GetGraph()->GetNodePath(); 
+				}
+			}
+		}
+	}
+}
+
 URigVMController* FRigVMBaseAction::GetController() const
 {
 	EnsureControllerValidity();
-	if(!WeakController.IsValid())
+	if (!WeakController.IsValid())
 	{
 		if(URigVMController* Controller = Cast<URigVMController>(ControllerPath.ResolveObject()))
 		{
 			WeakController = Controller;
 		}
 	}
+	if (!WeakController.IsValid())
+	{
+		if(UObject* ClientOuter = ClientHostPath.ResolveObject())
+		{
+			if (IRigVMClientHost* ClientHost = Cast<IRigVMClientHost>(ClientOuter))
+			{
+				if (FRigVMClient* Client = ClientHost->GetRigVMClient())
+				{
+					if (const URigVMGraph* Graph = Client->GetModel(ControllerModelPath))
+					{
+						WeakController = Client->GetOrCreateController(Graph);
+					}
+				}
+			}
+		}
+	}
 	return WeakController.Get();
 }
 
diff --git a/Engine/Plugins/Runtime/RigVM/Source/RigVMDeveloper/Public/RigVMModel/RigVMControllerActions.h b/Engine/Plugins/Runtime/RigVM/Source/RigVMDeveloper/Public/RigVMModel/RigVMControllerActions.h
--- a/Engine/Plugins/Runtime/RigVM/Source/RigVMDeveloper/Public/RigVMModel/RigVMControllerActions.h
+++ b/Engine/Plugins/Runtime/RigVM/Source/RigVMDeveloper/Public/RigVMModel/RigVMControllerActions.h
@@ -7,6 +7,7 @@
 #if WITH_EDITOR
 #include "ScopedTransaction.h"
 #endif
+#include "RigVMClient.h"
 #include "UObject/StructOnScope.h"
 #include "RigVMControllerActions.generated.h"
 
@@ -126,15 +127,7 @@ public:
 	static const inline TCHAR* CancelActionPrefix = TEXT("Cancel Action");
 
 	// Default constructor
-	FRigVMBaseAction(URigVMController* InController)
-		: WeakController(InController)
-		, Title(TEXT("Action"))
-	{
-		if(InController)
-		{
-			ControllerPath = FSoftObjectPath(InController);
-		}
-	}
+	RIGVMDEVELOPER_API FRigVMBaseAction(URigVMController* InController);
 
 	// Default destructor
 	virtual ~FRigVMBaseAction() {};
@@ -206,6 +199,12 @@ protected:
 	UPROPERTY()
 	FSoftObjectPath ControllerPath;
 
+	UPROPERTY()
+	FSoftObjectPath ClientHostPath;
+
+	UPROPERTY()
+	FString ControllerModelPath;
+
 	UPROPERTY()
 	FString Title;
 
