Control Rig: Fix crash when schema class in control rig is not assigned
#rnx

diff --git a/Engine/Plugins/Runtime/RigVM/Source/RigVMDeveloper/Private/RigVMModel/RigVMController.cpp b/Engine/Plugins/Runtime/RigVM/Source/RigVMDeveloper/Private/RigVMModel/RigVMController.cpp
--- a/Engine/Plugins/Runtime/RigVM/Source/RigVMDeveloper/Private/RigVMModel/RigVMController.cpp
+++ b/Engine/Plugins/Runtime/RigVM/Source/RigVMDeveloper/Private/RigVMModel/RigVMController.cpp
@@ -198,7 +198,11 @@ void URigVMController::SetGraph(URigVMGraph* InGraph)
 
 URigVMSchema* URigVMController::GetSchema() const
 {
-	return SchemaClass->GetDefaultObject<URigVMSchema>(); 
+	if (SchemaClass)
+	{
+		return SchemaClass->GetDefaultObject<URigVMSchema>();
+	}
+	return nullptr;
 }
 
 bool URigVMController::PushGraph(URigVMGraph* InGraph, bool bSetupUndoRedo)
