RigVM: Polish preview here functionality

#rnx

diff --git a/Engine/Plugins/Runtime/RigVM/Source/RigVMDeveloper/Private/EdGraph/RigVMEdGraphNode.cpp b/Engine/Plugins/Runtime/RigVM/Source/RigVMDeveloper/Private/EdGraph/RigVMEdGraphNode.cpp
--- a/Engine/Plugins/Runtime/RigVM/Source/RigVMDeveloper/Private/EdGraph/RigVMEdGraphNode.cpp
+++ b/Engine/Plugins/Runtime/RigVM/Source/RigVMDeveloper/Private/EdGraph/RigVMEdGraphNode.cpp
@@ -1513,7 +1513,23 @@ float URigVMEdGraphNode::GetFadedOutState() const
 	}
 	if (IsExcludedFromPreview())
 	{
-		return 0.4f;
+		bool bWasRunDuringLastEval = false;
+		if (CachedModelNode.IsValid())
+		{
+			if(IRigVMAssetInterface* Blueprint = GetImplementingOuter<IRigVMAssetInterface>())
+			{
+				if(URigVMHost* DebuggedHost = Cast<URigVMHost>(Blueprint->GetObjectBeingDebugged()))
+				{
+					const int32 Count = CachedModelNode->GetInstructionVisitedCount(DebuggedHost->GetRigVMExtendedExecuteContext(), DebuggedHost->GetVM(), FRigVMASTProxy());
+					bWasRunDuringLastEval = Count > 0;
+				}
+			}
+		}
+
+		if (!bWasRunDuringLastEval)
+		{
+			return 0.4f;
+		}
 	}
 #endif
 	return 1.f;
@@ -2084,16 +2100,16 @@ void URigVMEdGraphNode::InvalidateNodeColor()
 {
 	if(const URigVMNode* ModelNode = GetModelNode())
 	{
-		if (ModelNode->IsOutDated())
-		{
-			static const FLinearColor WarningColor = FAppStyle::GetColor("ErrorReporting.WarningBackgroundColor");
-			SetColorFromModel(WarningColor);
-		}
-		else if (ModelNode->HasEarlyExitMarker())
+		if (ModelNode->HasEarlyExitMarker())
 		{
 			static constexpr FLinearColor PreviewColor = FLinearColor(0.f,0.f,1.f,1.f);
 			SetColorFromModel(PreviewColor);
 		}
+		else if (ModelNode->IsOutDated())
+		{
+			static const FLinearColor WarningColor = FAppStyle::GetColor("ErrorReporting.WarningBackgroundColor");
+			SetColorFromModel(WarningColor);
+		}
 		else
 		{
 			SetColorFromModel(ModelNode->GetNodeColor());
diff --git a/Engine/Plugins/Runtime/RigVM/Source/RigVMDeveloper/Private/RigVMAsset.cpp b/Engine/Plugins/Runtime/RigVM/Source/RigVMDeveloper/Private/RigVMAsset.cpp
--- a/Engine/Plugins/Runtime/RigVM/Source/RigVMDeveloper/Private/RigVMAsset.cpp
+++ b/Engine/Plugins/Runtime/RigVM/Source/RigVMDeveloper/Private/RigVMAsset.cpp
@@ -920,12 +920,11 @@ const FLazyName& IRigVMAssetInterface::GetPanelPinFactoryName() const
 	return RigVMPanelPinFactoryName;
 }
 
-bool IRigVMAssetInterface::SetEarlyExitInstruction(int32 InInstructionIndex)
+bool IRigVMAssetInterface::SetEarlyExitInstruction(URigVMNode* InNodeToExitEarlyAfter, int32 InInstruction, bool bRequestHyperLink)
 {
-	const bool bResetResult = ResetEarlyExitInstruction();
-	if (InInstructionIndex == INDEX_NONE)
+	if (InNodeToExitEarlyAfter == nullptr)
 	{
-		return bResetResult;
+		return ResetEarlyExitInstruction(true);
 	}
 
 	const FRigVMByteCode* ByteCode = GetController()->GetCurrentByteCode();
@@ -934,12 +933,41 @@ bool IRigVMAssetInterface::SetEarlyExitInstruction(int32 InInstructionIndex)
 		return false;
 	}
 
-	if (InInstructionIndex < 0 || InInstructionIndex >= ByteCode->GetNumInstructions())
+	int32 InstructionIndex = InInstruction;
+	
+	const TArray<int32>& InstructionIndices = ByteCode->GetAllInstructionIndicesForSubject(InNodeToExitEarlyAfter);
+	if (!InstructionIndices.IsEmpty())
+	{
+		// find the last consecutive instruction for a node in the first block
+		int32 IndexInArray = 0;
+		if (InInstruction != INDEX_NONE)
+		{
+			IndexInArray = InstructionIndices.Find(InInstruction);
+			if (IndexInArray == INDEX_NONE)
+			{
+				IndexInArray = 0;
+			}
+		}
+		while (IndexInArray < InstructionIndices.Num() - 1 && InstructionIndices[IndexInArray] == InstructionIndices[IndexInArray + 1] - 1)
+		{
+			IndexInArray++;
+		}
+
+		InstructionIndex = InstructionIndices[IndexInArray];
+	}
+
+	const bool bResetResult = ResetEarlyExitInstruction(InstructionIndex == INDEX_NONE);
+	if (InstructionIndex == INDEX_NONE)
+	{
+		return bResetResult;
+	}
+
+	if (InstructionIndex < 0 || InstructionIndex >= ByteCode->GetNumInstructions())
 	{
 		return false;
 	}
 
-	GetDebugInfo().SetInstructionForEarlyExit(InInstructionIndex);
+	GetDebugInfo().SetInstructionForEarlyExit(InstructionIndex);
 
 	if (URigVMHost* DebuggedHost = GetDebuggedRigVMHost())
 	{
@@ -947,102 +975,54 @@ bool IRigVMAssetInterface::SetEarlyExitInstruction(int32 InInstructionIndex)
 		DebuggedHost->SetIsInDebugMode(true);
 	}
 
-	URigVMNode* Node = Cast<URigVMNode>(ByteCode->GetSubjectForInstruction(InInstructionIndex));
-	if (Node)
-	{
-		Node->SetHasEarlyExitMarker(true);
-
-		TArray<URigVMNode*> ExcludedNodes = {Node};
-		for (int32 ExcludedNodeIndex = 0; ExcludedNodeIndex < ExcludedNodes.Num(); ExcludedNodeIndex++)
-		{
-			URigVMNode* ExcludedNode = ExcludedNodes[ExcludedNodeIndex];
-			if (!ExcludedNode->IsMutable())
-			{
-				continue;
-			}
-
-			if (ExcludedNode != Node)
-			{
-				ExcludedNode->SetIsExcludedByEarlyExit(true);
-			}
+	TArray<URigVMNode*> IncludedNodes;
+	TSet<URigVMNode*> VisitedNodes;
 
-			const TArray<URigVMNode*> LinkedNodes = ExcludedNode->GetLinkedTargetNodes();
-			for (URigVMNode* LinkedNode : LinkedNodes)
-			{
-				if (!LinkedNode->IsMutable())
-				{
-					continue;
-				}
-				ExcludedNodes.AddUnique(LinkedNode);
-			}
-		}
+	IncludedNodes.Add(InNodeToExitEarlyAfter);
 
-		for (int32 ExcludedNodeIndex = 0; ExcludedNodeIndex < ExcludedNodes.Num(); ExcludedNodeIndex++)
+	for (int32 Index = 0; Index < IncludedNodes.Num(); Index++)
+	{
+		VisitedNodes.Add(IncludedNodes[Index]);
+		
+		const TArray<URigVMNode*> SourceNodes = IncludedNodes[Index]->GetLinkedSourceNodes();
+		for (URigVMNode* SourceNode : SourceNodes)
 		{
-			URigVMNode* ExcludedNode = ExcludedNodes[ExcludedNodeIndex];
-
-			if (!ExcludedNode->IsMutable())
+			if (VisitedNodes.Contains(SourceNode))
 			{
-				const TArray<URigVMNode*> LinkedTargetNodes = ExcludedNode->GetLinkedTargetNodes();
-				const bool bHasIncludedTargetNode = LinkedTargetNodes.ContainsByPredicate([](const URigVMNode* LinkedTargetNode) -> bool {
-					return !LinkedTargetNode->IsExcludedByEarlyExit();
-				});
-
-				if (bHasIncludedTargetNode)
-				{
-					continue;
-				}
-
-				ExcludedNode->SetIsExcludedByEarlyExit(true);
-			}
-
-			const TArray<URigVMNode*> LinkedSourceNodes = ExcludedNode->GetLinkedSourceNodes();
-			for (URigVMNode* LinkedNode : LinkedSourceNodes)
-			{
-				if (LinkedNode->IsMutable())
-				{
-					continue;
-				}
-				ExcludedNodes.AddUnique(LinkedNode);
+				continue;
 			}
+			IncludedNodes.AddUnique(SourceNode);
 		}
 	}
 
-	GetRigGraphDisplaySettings().bShowNodeRunCounts = true;
-	
-	return true;
-}
+	InNodeToExitEarlyAfter->SetHasEarlyExitMarker(true);
 
-bool IRigVMAssetInterface::SetEarlyExitInstruction(URigVMNode* InNodeToExitEarlyAfter)
-{
-	if (InNodeToExitEarlyAfter == nullptr)
+	TArray<URigVMNode*> AllNodes = InNodeToExitEarlyAfter->GetGraph()->GetNodes();
+	for (URigVMNode* Node : AllNodes)
 	{
-		return ResetEarlyExitInstruction();
+		Node->SetHasEarlyExitMarker(Node == InNodeToExitEarlyAfter);
+		Node->SetIsExcludedByEarlyExit(!VisitedNodes.Contains(Node));
 	}
 
-	const FRigVMByteCode* ByteCode = GetController()->GetCurrentByteCode();
-	if (ByteCode == nullptr)
-	{
-		return false;
-	}
+	GetRigGraphDisplaySettings().bShowNodeRunCounts = true;
 
-	const TArray<int32>& InstructionIndices = ByteCode->GetAllInstructionIndicesForSubject(InNodeToExitEarlyAfter);
-	if (InstructionIndices.IsEmpty())
-	{
-		return false;
-	}
+	LastPreviewHereNodes.Emplace(InNodeToExitEarlyAfter, InstructionIndex);
 
-	// find the last consecutive instruction for a node in the first block
-	int32 IndexInArray = 0;
-	while (IndexInArray < InstructionIndices.Num() - 1 && InstructionIndices[IndexInArray] == InstructionIndices[IndexInArray + 1] - 1)
+	if (bRequestHyperLink)
 	{
-		IndexInArray++;
+		if (FRigVMClient* Client = GetRigVMClient())
+		{
+			if (URigVMController* Controller = Client->GetOrCreateController(InNodeToExitEarlyAfter->GetGraph()))
+			{
+				Controller->RequestJumpToHyperLink(InNodeToExitEarlyAfter);
+			}
+		}
 	}
 
-	return SetEarlyExitInstruction(InstructionIndices[IndexInArray]);
+	return true;
 }
 
-bool IRigVMAssetInterface::ResetEarlyExitInstruction()
+bool IRigVMAssetInterface::ResetEarlyExitInstruction(bool bResetCallstack)
 {
 	GetDebugInfo().Reset();
 
@@ -1062,6 +1042,10 @@ bool IRigVMAssetInterface::ResetEarlyExitInstruction()
 	}
 
 	GetRigGraphDisplaySettings().bShowNodeRunCounts = false;;
+	if (bResetCallstack)
+	{
+		LastPreviewHereNodes.Reset();
+	}
 	return true;
 }
 
@@ -1083,7 +1067,7 @@ void IRigVMAssetInterface::TogglePreviewHere(const URigVMGraph* InGraph)
 	{
 		if (CurrentEarlyExistInstruction == GetPreviewNodeInstructionIndexFromSelection(InGraph))
 		{
-			ResetEarlyExitInstruction();
+			ResetEarlyExitInstruction(true);
 			return;
 		}
 	}
@@ -1091,7 +1075,7 @@ void IRigVMAssetInterface::TogglePreviewHere(const URigVMGraph* InGraph)
 	const TArray<FName> SelectedNodeNames = InGraph->GetSelectNodes();
 	if (SelectedNodeNames.Num() == 0)
 	{
-		ResetEarlyExitInstruction();
+		ResetEarlyExitInstruction(true);
 		return;
 	}
 
@@ -1117,7 +1101,7 @@ void IRigVMAssetInterface::PreviewHereStepForward()
 		return;
 	}
 
-	const int32 CurrentInstructionIndex =  DebuggedHost->GetDebugInfo().GetInstructionForEarlyExit();
+ 	int32 CurrentInstructionIndex = DebuggedHost->GetDebugInfo().GetInstructionForEarlyExit();
 	if (CurrentInstructionIndex == INDEX_NONE)
 	{
 		return;
@@ -1129,7 +1113,22 @@ void IRigVMAssetInterface::PreviewHereStepForward()
 		return;
 	}
 
-	URigVMNode* Node = Cast<URigVMNode>(VM->GetByteCode().GetSubjectForInstruction(CurrentInstructionIndex));
+	URigVMNode* Node = nullptr;
+	if (!LastPreviewHereNodes.IsEmpty())
+	{
+		const TPair<TWeakObjectPtr<URigVMNode>,int32> PreviousEntry = LastPreviewHereNodes.Pop();
+		Node = PreviousEntry.Key.Get();
+		if (Node)
+		{
+			CurrentInstructionIndex = PreviousEntry.Value;
+		}
+	}
+
+	if (!Node)
+	{
+		Node = Cast<URigVMNode>(VM->GetByteCode().GetSubjectForInstruction(CurrentInstructionIndex));
+	}
+	
 	if (!Node)
 	{
 		return;
@@ -1137,8 +1136,14 @@ void IRigVMAssetInterface::PreviewHereStepForward()
 
 	struct Local
 	{
-		static URigVMNode* FindNextNodeOnCompletedPin(URigVMNode* Node)
+		static URigVMNode* FindNextNodeOnCompletedPin(URigVMNode* Node, TSet<URigVMNode*>& VisitedNodes)
 		{
+			if (VisitedNodes.Contains(Node))
+			{
+				return nullptr;
+			}
+			VisitedNodes.Add(Node);
+			
 			for (const URigVMPin* Pin : Node->GetPins())
 			{
 				if (Pin->GetDirection() == ERigVMPinDirection::Input || Pin->GetDirection() == ERigVMPinDirection::IO)
@@ -1157,45 +1162,47 @@ void IRigVMAssetInterface::PreviewHereStepForward()
 					URigVMNode* SourceNode = SourcePins[0]->GetNode();
 					if (SourcePins[0]->GetFName() == FRigVMStruct::ControlFlowCompletedName)
 					{
-						return FindNextNodeOnCompletedPin(SourceNode);
+						return FindNextNodeOnCompletedPin(SourceNode,VisitedNodes);
 					}
 
-					const URigVMPin* NextExecutePin = SourceNode->FindRootPinByName(FRigVMStruct::ControlFlowCompletedName);
-					if (!NextExecutePin)
+					TArray<const URigVMPin*> PinsToFollow;
+
+					if (const URigVMPin* CompletedPin = SourceNode->FindRootPinByName(FRigVMStruct::ControlFlowCompletedName))
+					{
+						PinsToFollow.Add(CompletedPin);
+					}
+					else
 					{
 						// since we don't have a completed pin, try to the find the next execute pin after this one
 						const int32 SourcePinIndex = SourcePins[0]->GetPinIndex();
 						for (int32 Index = SourcePinIndex + 1; Index < SourceNode->GetPins().Num(); Index++)
 						{
-							NextExecutePin = SourceNode->GetPins()[Index];
-							if (NextExecutePin->GetDirection() != ERigVMPinDirection::Output || !NextExecutePin->IsExecuteContext())
+							const URigVMPin* NextExecutePin = SourceNode->GetPins()[Index];
+							if (NextExecutePin->GetDirection() == ERigVMPinDirection::Output && NextExecutePin->IsExecuteContext())
 							{
-								NextExecutePin = nullptr;
-								continue;
+								PinsToFollow.Add(NextExecutePin);
 							}
-							break;
 						}
 					}
 
-					if (!NextExecutePin || !NextExecutePin->IsExecuteContext())
-					{
-						continue;
-					}
-					
-					const TArray<URigVMPin*> TargetPins = NextExecutePin->GetLinkedTargetPins();
-					if (TargetPins.IsEmpty())
+					for (const URigVMPin* NextExecutePin : PinsToFollow)
 					{
-						continue;
+						const TArray<URigVMPin*> TargetPins = NextExecutePin->GetLinkedTargetPins();
+						if (TargetPins.IsEmpty())
+						{
+							continue;
+						}
+						return TargetPins[0]->GetNode();
 					}
 
-					return TargetPins[0]->GetNode();
+					return FindNextNodeOnCompletedPin(SourceNode,VisitedNodes);
 				}
 			}
 
 			return nullptr;
 		}
 
-		static URigVMNode* FindNextNode(URigVMNode* Node)
+		static URigVMNode* FindNextNode(URigVMNode* Node, TSet<URigVMNode*>& VisitedNodes)
 		{
 			if (!Node->IsMutable())
 			{
@@ -1226,31 +1233,69 @@ void IRigVMAssetInterface::PreviewHereStepForward()
 				}
 			}
 
-			return FindNextNodeOnCompletedPin(Node);
+			return FindNextNodeOnCompletedPin(Node, VisitedNodes);
 		}
 	};
 
-	URigVMNode* NextNode = Local::FindNextNode(Node);
+	TSet<URigVMNode*> VisitedNodes;
+	URigVMNode* NextNode = Local::FindNextNode(Node, VisitedNodes);
+
+	while (NextNode && NextNode->IsA<URigVMRerouteNode>())
+	{
+		const TArray<URigVMNode*> TargetNodes = NextNode->GetLinkedTargetNodes();
+		if (TargetNodes.IsEmpty())
+		{
+			break;
+		}
+		NextNode = TargetNodes[0];
+	}
+	
 	if (!NextNode)
 	{
+		LastPreviewHereNodes.Emplace(Node, CurrentInstructionIndex);
 		return;
 	}
 
-	SetEarlyExitInstruction(NextNode);
-
-	FRigVMClient* Client = GetRigVMClient();
-	if (!Client)
+	if (NextNode->IsA<URigVMFunctionInterfaceNode>())
 	{
-		return;
+		int32 InstructionIndex = CurrentInstructionIndex;
+		NextNode = nullptr;
+		while (NextNode == nullptr && (++InstructionIndex) < VM->GetByteCode().GetNumInstructions())
+		{
+			NextNode = Cast<URigVMNode>(VM->GetByteCode().GetSubjectForInstruction(InstructionIndex));
+			if (Cast<URigVMFunctionInterfaceNode>(NextNode))
+			{
+				NextNode = nullptr;
+			}
+		}
+
+		if (NextNode == nullptr)
+		{
+			return;
+		}
 	}
 
-	URigVMController* Controller = Client->GetOrCreateController(Node->GetGraph());
-	if (!Controller)
+	bool bFoundNextNode = false;
+	while (!bFoundNextNode && (++CurrentInstructionIndex) < VM->GetByteCode().GetNumInstructions())
 	{
-		return;
+		const TArray<TWeakObjectPtr<UObject>>* Callstack = VM->GetByteCode().GetCallstackForInstruction(CurrentInstructionIndex);
+		if (Callstack)
+		{
+			for (const TWeakObjectPtr<UObject>& WeakObjectPtr : *Callstack)
+			{
+				if (UObject* WeakObject = WeakObjectPtr.Get())
+				{
+					if (WeakObject == NextNode)
+					{
+						bFoundNextNode = true;
+						break;
+					}
+				}
+			}
+		}
 	}
 
-	Controller->RequestJumpToHyperLink(NextNode);
+	SetEarlyExitInstruction(NextNode, CurrentInstructionIndex,  true);
 }
 
 bool IRigVMAssetInterface::CanPreviewHereStepForward() const
@@ -1833,7 +1878,7 @@ void IRigVMAssetInterface::RecompileVM()
 		TGuardValue<bool> ReentrantGuardSelf(bSuspendModelNotificationsForSelf, true);
 		TGuardValue<bool> ReentrantGuardOthers(RigVMClient->bSuspendModelNotificationsForOthers, true);
 
-		ResetEarlyExitInstruction();
+		ResetEarlyExitInstruction(true);
 		
 		PreCompile();
 
@@ -2772,7 +2817,7 @@ void IRigVMAssetInterface::SetObjectBeingDebugged(UObject* NewObject)
 		NewRigBeingDebugged->bIsBeingDebugged = true;
 	}
 	
-	ResetEarlyExitInstruction();
+	ResetEarlyExitInstruction(true);
 #endif
 }
 
@@ -3307,7 +3352,7 @@ void IRigVMAssetInterface::HandleModifiedEvent(ERigVMGraphNotifType InNotifType,
 					}
 
 					RequestAutoVMRecompilation();
-					ResetEarlyExitInstruction();
+					ResetEarlyExitInstruction(true);
 
 					(void)MarkPackageDirty();
 					MarkBlueprintAsStructurallyModified();
@@ -3340,7 +3385,7 @@ void IRigVMAssetInterface::HandleModifiedEvent(ERigVMGraphNotifType InNotifType,
 			case ERigVMGraphNotifType::PinDirectionChanged:
 			{
 				RequestAutoVMRecompilation();
-				ResetEarlyExitInstruction();
+				ResetEarlyExitInstruction(true);
 				(void)MarkPackageDirty();
 
 				// we don't need to mark the blueprint as modified since we only
diff --git a/Engine/Plugins/Runtime/RigVM/Source/RigVMDeveloper/Public/RigVMAsset.h b/Engine/Plugins/Runtime/RigVM/Source/RigVMDeveloper/Public/RigVMAsset.h
--- a/Engine/Plugins/Runtime/RigVM/Source/RigVMDeveloper/Public/RigVMAsset.h
+++ b/Engine/Plugins/Runtime/RigVM/Source/RigVMDeveloper/Public/RigVMAsset.h
@@ -276,13 +276,12 @@ public:
 	/**
 	 * Sets the instruction index to exit early on 
 	 */
-	bool SetEarlyExitInstruction(int32 InInstructionIndex);
-	bool SetEarlyExitInstruction(URigVMNode* InNodeToExitEarlyAfter);
+	bool SetEarlyExitInstruction(URigVMNode* InNodeToExitEarlyAfter, int32 InInstruction = INDEX_NONE, bool bRequestHyperLink = false);
 
 	/**
 	 * Resets / removes the instruction index to exit early on 
 	 */
-	bool ResetEarlyExitInstruction();
+	bool ResetEarlyExitInstruction(bool bResetCallstack = true);
 
 	/**
 	 * Toggles the preview here functionality based on the selection
@@ -557,7 +556,7 @@ protected:
 	TObjectPtr<URigVMGraph> TemplateModel;
 	TObjectPtr<URigVMController> TemplateController;
 	mutable TArray<UObject::FAssetRegistryTag> CachedAssetTags;
-
+	TArray<TPair<TWeakObjectPtr<URigVMNode>,int32>> LastPreviewHereNodes;
 	
 
 	bool bSuspendModelNotificationsForSelf;
