RigVM: more defensive coding against dead controllers

#rnx

diff --git a/Engine/Plugins/Runtime/RigVM/Source/RigVMDeveloper/Private/EdGraph/RigVMEdGraphSchema.cpp b/Engine/Plugins/Runtime/RigVM/Source/RigVMDeveloper/Private/EdGraph/RigVMEdGraphSchema.cpp
--- a/Engine/Plugins/Runtime/RigVM/Source/RigVMDeveloper/Private/EdGraph/RigVMEdGraphSchema.cpp
+++ b/Engine/Plugins/Runtime/RigVM/Source/RigVMDeveloper/Private/EdGraph/RigVMEdGraphSchema.cpp
@@ -1419,7 +1419,10 @@ void URigVMEdGraphSchema::SetNodePosition(UEdGraphNode* Node, const FVector2D& P
 	{
 		if (URigVMController* Controller = RigNode->GetController())
 		{
-			Controller->SetNodePosition(RigNode->GetModelNode(), Position, bSetupUndo, false, false);
+			if  (!URigVMHost::IsGarbageOrDestroyed(Controller))
+			{
+				Controller->SetNodePosition(RigNode->GetModelNode(), Position, bSetupUndo, false, false);
+			}
 		}
 	}
 	
@@ -1429,7 +1432,10 @@ void URigVMEdGraphSchema::SetNodePosition(UEdGraphNode* Node, const FVector2D& P
 		{
 			if (URigVMController* Controller = Graph->GetController())
 			{
-				Controller->SetNodePositionByName(CommentNode->GetFName(), Position, bSetupUndo, false, false);
+				if  (!URigVMHost::IsGarbageOrDestroyed(Controller))
+				{
+					Controller->SetNodePositionByName(CommentNode->GetFName(), Position, bSetupUndo, false, false);
+				}
 			}
 		}
 	}
diff --git a/Engine/Plugins/Runtime/RigVM/Source/RigVMDeveloper/Private/RigVMModel/RigVMClient.cpp b/Engine/Plugins/Runtime/RigVM/Source/RigVMDeveloper/Private/RigVMModel/RigVMClient.cpp
--- a/Engine/Plugins/Runtime/RigVM/Source/RigVMDeveloper/Private/RigVMModel/RigVMClient.cpp
+++ b/Engine/Plugins/Runtime/RigVM/Source/RigVMDeveloper/Private/RigVMModel/RigVMClient.cpp
@@ -435,8 +435,11 @@ URigVMController* FRigVMClient::GetController(const URigVMGraph* InModel) const
 		const FSoftObjectPath Key(InModel);
 		if(const TObjectPtr<URigVMController>* Controller = Controllers.Find(Key))
 		{
-			checkf(Controller->Get()->GetGraph() == InModel, TEXT("Controller %s contains unexpected graph."), *Controller->GetPathName());
-			return Controller->Get();
+			if  (!URigVMHost::IsGarbageOrDestroyed(Controller->Get()))
+			{
+				checkf(Controller->Get()->GetGraph() == InModel, TEXT("Controller %s contains unexpected graph."), *Controller->GetPathName());
+				return Controller->Get();
+			}
 		}
 	}
 	return nullptr;
@@ -476,13 +479,16 @@ URigVMController* FRigVMClient::GetOrCreateController(const URigVMGraph* InModel
 	{
 		if(URigVMController* Controller = GetController(InModel))
 		{
-			// We associate controllers to graphs via soft path, so they can match newly created graphs
-			// If this happens make sure the graph is correctly bound to the controller
-			if (!InModel->ModifiedEvent.IsBound())
+			if (IsValid(Controller))
 			{
-				const_cast<URigVMGraph*>(InModel)->OnModified().AddUObject(Controller, &URigVMController::HandleModifiedEvent);
+				// We associate controllers to graphs via soft path, so they can match newly created graphs
+				// If this happens make sure the graph is correctly bound to the controller
+				if (!InModel->ModifiedEvent.IsBound())
+				{
+					const_cast<URigVMGraph*>(InModel)->OnModified().AddUObject(Controller, &URigVMController::HandleModifiedEvent);
+				}
+				return Controller;
 			}
-			return Controller;
 		}
 		return CreateController(InModel);
 	}
@@ -1112,7 +1118,7 @@ URigVMController* FRigVMClient::CreateController(const URigVMGraph* InModel)
 	const FString ModelName = InModel ? InModel->GetName() : TEXT("__NullGraph");
 	const FName SafeControllerName = GetUniqueName(*FString::Printf(TEXT("%s_Controller"), *ModelName));
 	URigVMController* Controller = NewObject<URigVMController>(GetOuter(), ControllerClass, SafeControllerName);
-	Controllers.Add(InModel, Controller);
+	Controllers.FindOrAdd(InModel, nullptr) = Controller;
 	if(InModel && InModel->GetSchemaClass() != nullptr)
 	{
 		Controller->SetSchemaClass(InModel->GetSchemaClass());
