RigVM: Solidify visitor pattern for archetype instances

#rnx

diff --git a/Engine/Plugins/Runtime/RigVM/Source/RigVMDeveloper/Private/RigVMAsset.cpp b/Engine/Plugins/Runtime/RigVM/Source/RigVMDeveloper/Private/RigVMAsset.cpp
--- a/Engine/Plugins/Runtime/RigVM/Source/RigVMDeveloper/Private/RigVMAsset.cpp
+++ b/Engine/Plugins/Runtime/RigVM/Source/RigVMDeveloper/Private/RigVMAsset.cpp
@@ -1788,7 +1788,8 @@ void IRigVMAssetInterface::BroadcastRigVMPackageDone()
 	TArray<UObject*> ArchetypeInstances = GetArchetypeInstances(true, false);
 	for (UObject* Instance : ArchetypeInstances)
 	{
-		if (URigVMHost* InstanceHost = Cast<URigVMHost>(Instance))
+		URigVMHost* InstanceHost = Cast<URigVMHost>(Instance);
+		if (!URigVMHost::IsGarbageOrDestroyed(InstanceHost))
 		{
 			InstanceHost->BroadCastEndLoadPackage();
 		}
@@ -2176,7 +2177,8 @@ void IRigVMAssetInterface::RequestRigVMInit() const
 	TArray<UObject*> Instances = GetArchetypeInstances(true, true);
 	for (UObject* Instance : Instances)
 	{
-		if (URigVMHost* RigVMHost = Cast<URigVMHost>(Instance))
+		URigVMHost* RigVMHost = Cast<URigVMHost>(Instance);
+		if (!URigVMHost::IsGarbageOrDestroyed(RigVMHost))
 		{
 			RigVMHost->RequestInit();
 		}
@@ -3293,7 +3295,8 @@ void IRigVMAssetInterface::HandleModifiedEvent(ERigVMGraphNotifType InNotifType,
 							TArray<UObject*> ArchetypeInstances = GetArchetypeInstances(false, false);
 							for (UObject* ArchetypeInstance : ArchetypeInstances)
 							{
-								if (URigVMHost* InstancedHost = Cast<URigVMHost>(ArchetypeInstance))
+								URigVMHost* InstancedHost = Cast<URigVMHost>(ArchetypeInstance);
+								if (!URigVMHost::IsGarbageOrDestroyed(InstancedHost))
 								{
 									if (!InstancedHost->HasAllFlags(RF_ClassDefaultObject))
 									{
@@ -3756,7 +3759,8 @@ void IRigVMAssetInterface::PropagateRuntimeSettingsFromBPToInstances()
 	TArray<UObject*> ArchetypeInstances = GetArchetypeInstances(true, false);
 	for (UObject* ArchetypeInstance : ArchetypeInstances)
 	{
-		if (URigVMHost* InstanceHost = Cast<URigVMHost>(ArchetypeInstance))
+		URigVMHost* InstanceHost = Cast<URigVMHost>(ArchetypeInstance);
+		if (!URigVMHost::IsGarbageOrDestroyed(InstanceHost))
 		{
 			InstanceHost->VMRuntimeSettings = GetVMRuntimeSettings();
 		}
@@ -3785,20 +3789,24 @@ void IRigVMAssetInterface::InitializeArchetypeInstances()
 		TArray<UObject*> ArchetypeInstances = GetArchetypeInstances(false, false);
 		for (UObject* Instance : ArchetypeInstances)
 		{
-			if (!Instance->HasAllFlags(RF_ClassDefaultObject))
+			URigVMHost* InstanceHost = Cast<URigVMHost>(Instance);
+			if (URigVMHost::IsGarbageOrDestroyed(InstanceHost))
 			{
-				if (URigVMHost* InstanceHost = Cast<URigVMHost>(Instance))
-				{
-					// No objects should be created during load, so PostInitInstanceIfRequired, which creates a new VM and
-					// DynamicHierarchy, should not be called during load
-					if (!InstanceHost->HasAllFlags(RF_NeedPostLoad))
-					{
-						InstanceHost->PostInitInstanceIfRequired();
-					}
-					InstanceHost->InstantiateVMFromCDO();
-					InstanceHost->CopyExternalVariableDefaultValuesFromCDO();
-				}
+				continue;
+			}
+			if (InstanceHost->HasAllFlags(RF_ClassDefaultObject))
+			{
+				continue;
+			}
+
+			// No objects should be created during load, so PostInitInstanceIfRequired, which creates a new VM and
+			// DynamicHierarchy, should not be called during load
+			if (!InstanceHost->HasAllFlags(RF_NeedPostLoad))
+			{
+				InstanceHost->PostInitInstanceIfRequired();
 			}
+			InstanceHost->InstantiateVMFromCDO();
+			InstanceHost->CopyExternalVariableDefaultValuesFromCDO();
 		}
 	}
 }
