RigVM: Fix crash with empty variant tag string

#rnx

diff --git a/Engine/Plugins/Runtime/RigVM/Source/RigVMEditor/Private/Widgets/SRigVMSwapAssetReferencesWidget.cpp b/Engine/Plugins/Runtime/RigVM/Source/RigVMEditor/Private/Widgets/SRigVMSwapAssetReferencesWidget.cpp
--- a/Engine/Plugins/Runtime/RigVM/Source/RigVMEditor/Private/Widgets/SRigVMSwapAssetReferencesWidget.cpp
+++ b/Engine/Plugins/Runtime/RigVM/Source/RigVMEditor/Private/Widgets/SRigVMSwapAssetReferencesWidget.cpp
@@ -72,7 +72,10 @@ bool FRigVMTreeAssetVariantFilter::Filters(TSharedRef<FRigVMTreeNode>& InNode, c
 				const FProperty* AssetVariantProperty = CastField<FProperty>(SwapContext->GetSourceAsset().GetClass()->FindPropertyByName(AssetVariantPropertyName));
 				const FString VariantStr = SwapContext->GetSourceAsset().GetTagValueRef<FString>(AssetVariantPropertyName);
 				FRigVMVariant AssetVariant;
-				AssetVariantProperty->ImportText_Direct(*VariantStr, &AssetVariant, nullptr, EPropertyPortFlags::PPF_None);
+				if (!VariantStr.IsEmpty())
+				{
+					AssetVariantProperty->ImportText_Direct(*VariantStr, &AssetVariant, nullptr, EPropertyPortFlags::PPF_None);
+				}
 
 				if (!AssetVariant.Guid.IsValid())
 				{
