RigVM: Fix issue with variable nodes on aggregate nodes

#rnx

diff --git a/Engine/Plugins/Runtime/RigVM/Source/RigVMDeveloper/Private/EdGraph/RigVMEdGraph.cpp b/Engine/Plugins/Runtime/RigVM/Source/RigVMDeveloper/Private/EdGraph/RigVMEdGraph.cpp
--- a/Engine/Plugins/Runtime/RigVM/Source/RigVMDeveloper/Private/EdGraph/RigVMEdGraph.cpp
+++ b/Engine/Plugins/Runtime/RigVM/Source/RigVMDeveloper/Private/EdGraph/RigVMEdGraph.cpp
@@ -504,6 +504,14 @@ bool URigVMEdGraph::HandleModifiedEvent_Internal(ERigVMGraphNotifType InNotifTyp
 							TargetRigPin->LinkedTo.Remove(nullptr);
 						}
 					}
+
+					if (const URigVMNode* SourceNode = SourcePin->GetNode())
+					{
+						if (SourceNode->IsA<URigVMVariableNode>() && TargetRigNode)
+						{
+							TargetRigNode->ModelPinsChanged(true);
+						}
+					}
 				}
 			}
 			break;
diff --git a/Engine/Plugins/Runtime/RigVM/Source/RigVMDeveloper/Private/RigVMModel/RigVMController.cpp b/Engine/Plugins/Runtime/RigVM/Source/RigVMDeveloper/Private/RigVMModel/RigVMController.cpp
--- a/Engine/Plugins/Runtime/RigVM/Source/RigVMDeveloper/Private/RigVMModel/RigVMController.cpp
+++ b/Engine/Plugins/Runtime/RigVM/Source/RigVMDeveloper/Private/RigVMModel/RigVMController.cpp
@@ -3936,6 +3936,8 @@ URigVMNode* URigVMController::EjectNodeFromPin(URigVMPin* InPin, bool bSetupUndo
 	Notify(ERigVMGraphNotifType::NodeAdded, NodeToEject);
 	if (InputPin)
 	{
+		(void)UnbindPinFromVariable(InputPin, false);
+		
 		TArray<URigVMLink*> SourceLinks = InputPin->GetSourceLinks(true);
 		if (SourceLinks.Num() > 0)
 		{
@@ -11278,6 +11280,8 @@ FString URigVMController::AddAggregatePin(URigVMNode* InNode, const FString& InP
 			return FString();
 		}
 
+		(void)EjectAllInjectedNodes(InNode, false);
+		
 		const FString AggregateArg1 = Arg1->GetName();
 		const FString AggregateArg2 = Arg2->GetName();
 		const FString AggregateOppositeArg = ArgOpposite->GetName();
diff --git a/Engine/Plugins/Runtime/RigVM/Source/RigVMDeveloper/Private/RigVMModel/RigVMPin.cpp b/Engine/Plugins/Runtime/RigVM/Source/RigVMDeveloper/Private/RigVMModel/RigVMPin.cpp
--- a/Engine/Plugins/Runtime/RigVM/Source/RigVMDeveloper/Private/RigVMModel/RigVMPin.cpp
+++ b/Engine/Plugins/Runtime/RigVM/Source/RigVMDeveloper/Private/RigVMModel/RigVMPin.cpp
@@ -41,7 +41,7 @@ URigVMGraph* URigVMInjectionInfo::GetGraph() const
 
 URigVMPin* URigVMInjectionInfo::GetPin() const
 {
-	return CastChecked<URigVMPin>(GetOuter());
+	return Cast<URigVMPin>(GetOuter());
 }
 
 URigVMInjectionInfo::FWeakInfo URigVMInjectionInfo::GetWeakInfo() const
