Control Rig: Add a framework to remove unused members from a Control Rig Blueprint, along with an implementation for variables and functions

* A dialog is now shown when a member is removed, allowing to select which members are removed
* Functions and Variables are now also considered unused if they are not connected in the graph
* Functions and Variables that are not connected in the graph are now removed from the graph when removing unused functions or variables
* Update highlights only when nodes were added or removed, or where their connected pins changed

diff --git a/Engine/Plugins/Runtime/RigVM/Source/RigVMDeveloper/Public/EdGraph/RigVMEdGraphNode.h b/Engine/Plugins/Runtime/RigVM/Source/RigVMDeveloper/Public/EdGraph/RigVMEdGraphNode.h
--- a/Engine/Plugins/Runtime/RigVM/Source/RigVMDeveloper/Public/EdGraph/RigVMEdGraphNode.h
+++ b/Engine/Plugins/Runtime/RigVM/Source/RigVMDeveloper/Public/EdGraph/RigVMEdGraphNode.h
@@ -239,8 +239,12 @@ public:
 	/** Sets the node's content from an external client, used for preview nodes without a graph */
 	UE_API void SetModelNode(URigVMNode* InModelNode);
 
+	/** Returns all input pins */
 	const TArray<URigVMPin*>& GetInputPins() const { return InputPins; }
 
+	/** Returns all execute pins */
+	const TArray<URigVMPin*>& GetExecutePins() const { return ExecutePins; };
+
 	// Returns true if this Node should be faded out
 	UE_API bool IsFadedOut() const;
 
diff --git a/Engine/Plugins/Runtime/RigVM/Source/RigVMEditor/Private/Editor/RemoveUnusedMembers/RigVMEditorRemoveUnusedMembers.cpp b/Engine/Plugins/Runtime/RigVM/Source/RigVMEditor/Private/Editor/RemoveUnusedMembers/RigVMEditorRemoveUnusedMembers.cpp
new file mode 100644
--- /dev/null
+++ b/Engine/Plugins/Runtime/RigVM/Source/RigVMEditor/Private/Editor/RemoveUnusedMembers/RigVMEditorRemoveUnusedMembers.cpp
@@ -0,0 +1,363 @@
+// Copyright Epic Games, Inc. All Rights Reserved.
+
+#include "RigVMEditorRemoveUnusedMembers.h"
+
+#include "Algo/Transform.h"
+#include "EdGraph/RigVMEdGraphSchema.h"
+
+#define LOCTEXT_NAMESPACE "RigVMEditorRemoveUnusedMembersDialog"
+
+namespace UE::RigVMEditor::Private
+{
+	template <>
+	FText GetDialogTitle<URigVMFunctionReferenceNode>()
+	{
+		return LOCTEXT("RemoveUnusedFunctionsDialogTitle", "Remove Unused Functions");
+	}
+
+	template <>
+	FText GetDialogTitle<URigVMVariableNode>()
+	{
+		return LOCTEXT("RemoveUnusedVariablesDialogTitle", "Remove Unused Variables");
+	}
+
+	/** Removes disconnected function nodes */
+	template <>
+	void RemoveDisconnectedNodes<URigVMFunctionReferenceNode>(
+		URigVMController& Controller,
+		const TSharedRef<FRigVMEdGraphNodeRegistry>& EdGraphNodeRegistry,
+		const TArray<FName>& MemberNamesToRemove)
+	{
+		TArray<URigVMNode*> NodesToRemove;
+		Algo::TransformIf(EdGraphNodeRegistry->GetDisconnectedEdGrapNodes(), NodesToRemove,
+			[&MemberNamesToRemove](const TWeakObjectPtr<URigVMEdGraphNode>& WeakEdGraphNode)
+			{
+				URigVMFunctionReferenceNode* FunctionReferenceNode = WeakEdGraphNode.IsValid() ?
+					Cast<URigVMFunctionReferenceNode>(WeakEdGraphNode->GetModelNode())
+					: nullptr;
+
+				if (FunctionReferenceNode &&
+					MemberNamesToRemove.Contains(FunctionReferenceNode->GetFunctionIdentifier().GetFunctionFName()))
+				{
+					return true;
+				}
+
+				return false;
+			},
+			[](const TWeakObjectPtr<URigVMEdGraphNode>& WeakEdGraphNode)
+			{
+				URigVMFunctionReferenceNode* FunctionReferenceNode = WeakEdGraphNode.IsValid() ?
+					Cast<URigVMFunctionReferenceNode>(WeakEdGraphNode->GetModelNode())
+					: nullptr;
+
+				check(FunctionReferenceNode);
+				return FunctionReferenceNode;
+			});
+
+		if (!NodesToRemove.IsEmpty())
+		{
+			constexpr bool bSetupUndoRedo = true;
+			constexpr bool bPrintPythonCommand = false;
+			Controller.RemoveNodes(NodesToRemove, bSetupUndoRedo, bPrintPythonCommand);
+		}
+	}
+
+	/** Removes disconnected variable nodes */
+	template <>
+	void RemoveDisconnectedNodes<URigVMVariableNode>(
+		URigVMController& Controller,
+		const TSharedRef<FRigVMEdGraphNodeRegistry>& EdGraphNodeRegistry,
+		const TArray<FName>& MemberNamesToRemove)
+	{
+		TArray<URigVMNode*> NodesToRemove;
+		Algo::TransformIf(EdGraphNodeRegistry->GetDisconnectedEdGrapNodes(), NodesToRemove,
+			[&MemberNamesToRemove](const TWeakObjectPtr<URigVMEdGraphNode>& WeakEdGraphNode)
+			{
+				URigVMVariableNode* VariableNode = WeakEdGraphNode.IsValid() ?
+					Cast<URigVMVariableNode>(WeakEdGraphNode->GetModelNode())
+					: nullptr;
+
+				if (VariableNode &&
+					MemberNamesToRemove.Contains(VariableNode->GetVariableName()))
+				{
+					return true;
+				}
+
+				return false;
+			},
+			[](const TWeakObjectPtr<URigVMEdGraphNode>& WeakEdGraphNode)
+			{
+				URigVMVariableNode* VariableNode = WeakEdGraphNode.IsValid() ?
+					Cast<URigVMVariableNode>(WeakEdGraphNode->GetModelNode())
+					: nullptr;
+					
+				check(VariableNode);
+				return VariableNode;
+			});
+
+		if (!NodesToRemove.IsEmpty())
+		{
+			constexpr bool bSetupUndoRedo = true;
+			constexpr bool bPrintPythonCommand = false;
+			Controller.RemoveNodes(NodesToRemove, bSetupUndoRedo, bPrintPythonCommand);
+		}
+	}
+
+	/** Removes function members */
+	template <>
+	void RemoveMembers<URigVMFunctionReferenceNode>(
+		const TScriptInterface<IRigVMAssetInterface>& Asset,
+		const TArray<FName>& MemberNamesToRemove)
+	{
+		if (!Asset)
+		{
+			return;
+		}
+
+		const URigVMEdGraphSchema* Schema = GetDefault<URigVMEdGraphSchema>();
+		if (!Schema)
+		{
+			return;
+		}
+
+		for (const FName& MemberNameToRemove : MemberNamesToRemove)
+		{
+			URigVMFunctionLibrary* Library = Asset->GetLocalFunctionLibrary();
+			URigVMLibraryNode* Function = Library ? Library->FindFunction(MemberNameToRemove) : nullptr;
+			if (Library && Function)
+			{
+				if (UEdGraph* EdGraph = Asset->GetEdGraph(Function->GetContainedGraph()))
+				{
+					Schema->TryDeleteGraph(EdGraph);
+				}
+			}
+		}
+	}
+
+	/** Removes variable members */
+	template <>
+	void RemoveMembers<URigVMVariableNode>(
+		const TScriptInterface<IRigVMAssetInterface>& Asset,
+		const TArray<FName>& MemberNamesToRemove)
+	{
+		if (!Asset)
+		{
+			return;
+		}
+
+		// Variables
+		for (const FName& MemberNameToRemove : MemberNamesToRemove)
+		{
+			Asset->RemoveMemberVariable(MemberNameToRemove);
+		}
+
+		// Local variables
+		URigVMGraph* FocusedGraph = Asset->GetFocusedModel();
+		FRigVMClient* RigVMClient = Asset->GetRigVMClient();
+		URigVMController* Controller = RigVMClient ? RigVMClient->GetController(FocusedGraph) : nullptr;
+		if (!FocusedGraph ||
+			!RigVMClient ||
+			!Controller)
+		{
+			return;
+		}
+
+		for (const FRigVMGraphVariableDescription& VariableDescription : FocusedGraph->GetLocalVariables())
+		{
+			if (MemberNamesToRemove.Contains(VariableDescription.Name))
+			{
+				constexpr bool bSetupUndo = true;
+				constexpr bool bPrintPythonCommand = false;
+				Controller->RemoveLocalVariable(VariableDescription.Name, bSetupUndo, bPrintPythonCommand);
+			}
+		}
+	}
+
+	template<>
+	void FindUnusedMemberNames<URigVMFunctionReferenceNode>(
+		const TScriptInterface<IRigVMAssetInterface>& InAsset,
+		const TSharedRef<FRigVMEdGraphNodeRegistry>& InEdGraphNodeRegistry,
+		TMap<FRigVMUnusedMemberCategory, TArray<FName>>& OutCategoryToUnusedMemberNamesMap)
+	{
+		TArray<FName> UnusedMemberNames;
+
+		const FRigVMClient* RigVMClient = InAsset ? InAsset->GetRigVMClient() : nullptr;
+		URigVMFunctionLibrary* Library = RigVMClient ? RigVMClient->GetFunctionLibrary() : nullptr;
+		if (!RigVMClient || !Library)
+		{
+			return;
+		}
+
+		const TArray<TWeakObjectPtr<URigVMEdGraphNode>> ConnectedEdGraphNodes = InEdGraphNodeRegistry->GetConnectedEdGrapNodes();
+
+		TArray<URigVMFunctionReferenceNode*> FunctionReferenceNodes;
+		Algo::TransformIf(ConnectedEdGraphNodes, FunctionReferenceNodes,
+			[](const TWeakObjectPtr<URigVMEdGraphNode>& WeakEdGraphNode)
+			{
+				return
+					WeakEdGraphNode.IsValid() &&
+					Cast<URigVMFunctionReferenceNode>(WeakEdGraphNode->GetModelNode()) != nullptr;
+			},
+			[](const TWeakObjectPtr<URigVMEdGraphNode>& WeakEdGraphNode)
+			{
+				check(WeakEdGraphNode.IsValid());
+				return Cast<URigVMFunctionReferenceNode>(WeakEdGraphNode->GetModelNode());
+			});
+
+		TArray<FName>* PublicFunctionMemberNamesPtr = nullptr;
+		TArray<FName>* PrivateFunctionMemberNamesPtr = nullptr;
+
+		for (const URigVMLibraryNode* Function : Library->GetFunctions())
+		{
+			if (!Function)
+			{
+				continue;
+			}
+
+			const bool bConnectedInGraph = FunctionReferenceNodes.ContainsByPredicate(
+				[Function](const URigVMFunctionReferenceNode* FunctionReferenceNode)
+				{
+					if (ensureMsgf(FunctionReferenceNode, TEXT("Unexpected nullptr in PrivateFunctionReferenceNodes")))
+					{
+						return FunctionReferenceNode->GetFunctionIdentifier().GetFunctionFName() == Function->GetFName();
+					}
+
+					return true;
+				});
+
+			if (bConnectedInGraph)
+			{
+				continue;
+			}
+
+			const FRigVMGraphFunctionIdentifier FunctionIdentifier = Function->GetFunctionIdentifier();
+			const FSoftObjectPath& FunctionHostPath = Function->GetFunctionIdentifier().HostObject;
+			const UObject* RigVMBlueprintGeneratedClassObj = FunctionHostPath.TryLoad();
+			const IRigVMGraphFunctionHost* FunctionHost = Cast<const IRigVMGraphFunctionHost>(RigVMBlueprintGeneratedClassObj);
+			const FRigVMGraphFunctionStore* FunctionStore = FunctionHost ? FunctionHost->GetRigVMGraphFunctionStore() : nullptr;
+
+			const bool bPublicFunction = FunctionStore && FunctionStore->IsFunctionPublic(FunctionIdentifier);
+
+			if (bPublicFunction)
+			{
+				if (!PublicFunctionMemberNamesPtr)
+				{
+					constexpr bool bSafeToRemove = false;
+					const FRigVMUnusedMemberCategory PublicFunctionsCategory(
+						"PublicFunctions",
+						LOCTEXT("PublicFunctionsCategory", "Public Functions"),
+						bSafeToRemove);
+
+					PublicFunctionMemberNamesPtr = &OutCategoryToUnusedMemberNamesMap.Add(PublicFunctionsCategory);
+				}
+				check(PublicFunctionMemberNamesPtr);
+				
+				PublicFunctionMemberNamesPtr->Add(Function->GetFName());
+			}
+			else
+			{
+				if (!PrivateFunctionMemberNamesPtr)
+				{
+					constexpr bool bSafeToRemove = true;
+					const FRigVMUnusedMemberCategory PrivateFunctionsCategory(
+						"PrivateFunctions",
+						LOCTEXT("PrivateFunctionsCategory", "Private Functions"),
+						bSafeToRemove);
+
+					PrivateFunctionMemberNamesPtr = &OutCategoryToUnusedMemberNamesMap.Add(PrivateFunctionsCategory);
+				}
+				check(PrivateFunctionMemberNamesPtr);
+				
+				PrivateFunctionMemberNamesPtr->Add(Function->GetFName());
+			}
+		}
+	}
+
+	template<>
+	void FindUnusedMemberNames<URigVMVariableNode>(
+		const TScriptInterface<IRigVMAssetInterface>& InAsset,
+		const TSharedRef<FRigVMEdGraphNodeRegistry>& InEdGraphNodeRegistry,
+		TMap<FRigVMUnusedMemberCategory, TArray<FName>>& OutCategoryToUnusedMemberNamesMap)
+	{
+		if (!InAsset)
+		{
+			return;
+		}
+
+		const auto IsVariableConnectedInGraph =
+			[&InEdGraphNodeRegistry](const FRigVMGraphVariableDescription& VariableDescription)
+			{
+				const bool bConnectedInGraph = InEdGraphNodeRegistry->GetConnectedEdGrapNodes().ContainsByPredicate(
+					[&VariableDescription](const TWeakObjectPtr<URigVMEdGraphNode>& WeakEdGraphNode)
+					{
+						URigVMVariableNode* VariableNode = WeakEdGraphNode.IsValid() ?
+							Cast<URigVMVariableNode>(WeakEdGraphNode->GetModelNode()) :
+							nullptr;
+
+						if (VariableNode)
+						{
+							return VariableNode->GetVariableName() == VariableDescription.Name;
+						}
+
+						return false;
+					});
+
+				return bConnectedInGraph;
+			};
+
+		// As of 5.7, variables are not publicly accessible hence all are considered safe to remove
+		constexpr bool bSafeToRemove = true;
+
+		TArray<FName>* VariableMemberNamesPtr = nullptr;
+		TArray<FName>* LocalVariableMemberNamesPtr = nullptr;
+
+		// Asset variables
+		for (const FRigVMGraphVariableDescription& VariableDescription : InAsset->GetAssetVariables())
+		{
+			if (IsVariableConnectedInGraph(VariableDescription))
+			{
+				continue;
+			}
+
+			if (!VariableMemberNamesPtr)
+			{
+				const FRigVMUnusedMemberCategory VariablesCategory(
+					"Variables",
+					LOCTEXT("VariablesCategory", "Variables"),
+					bSafeToRemove);
+
+				VariableMemberNamesPtr = &OutCategoryToUnusedMemberNamesMap.Add(VariablesCategory);
+			}
+			check(VariableMemberNamesPtr);
+			
+			VariableMemberNamesPtr->Add(VariableDescription.Name);
+		}
+
+		// Local variables
+		if (URigVMGraph* FocusedGraph = InAsset->GetFocusedModel())
+		{
+			for (const FRigVMGraphVariableDescription& VariableDescription : FocusedGraph->GetLocalVariables())
+			{
+				if (IsVariableConnectedInGraph(VariableDescription))
+				{
+					continue;
+				}
+
+				if (!LocalVariableMemberNamesPtr)
+				{
+					const FRigVMUnusedMemberCategory LocalVariablesCategory(
+						"LocalVariables",
+						LOCTEXT("LocalVariablesCategory", "Local Variables"),
+						bSafeToRemove);
+
+					LocalVariableMemberNamesPtr = &OutCategoryToUnusedMemberNamesMap.Add(LocalVariablesCategory);
+				}
+				check(LocalVariableMemberNamesPtr);
+				
+				LocalVariableMemberNamesPtr->Add(VariableDescription.Name);
+			}
+		}
+	}
+}
+
+#undef LOCTEXT_NAMESPACE
diff --git a/Engine/Plugins/Runtime/RigVM/Source/RigVMEditor/Private/Editor/RemoveUnusedMembers/RigVMEditorRemoveUnusedMembers.h b/Engine/Plugins/Runtime/RigVM/Source/RigVMEditor/Private/Editor/RemoveUnusedMembers/RigVMEditorRemoveUnusedMembers.h
new file mode 100644
--- /dev/null
+++ b/Engine/Plugins/Runtime/RigVM/Source/RigVMEditor/Private/Editor/RemoveUnusedMembers/RigVMEditorRemoveUnusedMembers.h
@@ -0,0 +1,129 @@
+// Copyright Epic Games, Inc. All Rights Reserved.
+
+#pragma once
+
+#include "Editor/RigVMEdGraphNodeRegistry.h"
+#include "Framework/Notifications/NotificationManager.h"
+#include "RigVMAsset.h"
+#include "RigVMEditorRemoveUnusedMembersCategory.h"
+#include "RigVMEditorRemoveUnusedMembersDialog.h"
+#include "RigVMModel/Nodes/RigVMFunctionReferenceNode.h"
+#include "RigVMModel/Nodes/RigVMVariableNode.h"
+#include "RigVMModel/RigVMClient.h"
+#include "Widgets/Notifications/SNotificationList.h"
+#include <type_traits>
+
+namespace UE::RigVMEditor
+{
+	/** Concept for types that support RemoveUnusedMembersFromAsset */
+	template <typename TRigVMNode>
+	concept ConceptTypeSupportsRemoveUnusedMembers = requires
+	{
+		std::is_base_of_v<URigVMFunctionReferenceNode, TRigVMNode> ||
+		std::is_base_of_v<URigVMVariableNode, TRigVMNode>;
+	};
+
+	namespace Private
+	{
+		/** Returns a type specific dialog title */
+		template <typename TRigVMNode> requires ConceptTypeSupportsRemoveUnusedMembers<TRigVMNode>
+		FText GetDialogTitle();
+
+		/** Finds unused member names. Considers members with disconnected nodes as unused */
+		template <typename TRigVMNode> requires ConceptTypeSupportsRemoveUnusedMembers<TRigVMNode>
+		void FindUnusedMemberNames(
+			const TScriptInterface<IRigVMAssetInterface>& InAsset,
+			const TSharedRef<FRigVMEdGraphNodeRegistry>& InEdGraphNodeRegistry,
+			TMap<FRigVMUnusedMemberCategory, TArray<FName>>& OutCategoryToUnusedMemberNamesMap);
+
+		/** Removes disconnected nodes */
+		template <typename TRigVMNode> requires ConceptTypeSupportsRemoveUnusedMembers<TRigVMNode>
+		void RemoveDisconnectedNodes(
+			URigVMController& Controller,
+			const TSharedRef<FRigVMEdGraphNodeRegistry>& EdGraphNodeRegistry,
+			const TArray<FName>& MemberNamesToRemove);
+
+		/** Removes members */
+		template <typename TRigVMNode> requires ConceptTypeSupportsRemoveUnusedMembers<TRigVMNode>
+		void RemoveMembers(
+			const TScriptInterface<IRigVMAssetInterface>& Asset, 
+			const TArray<FName>& MemberNamesToRemove);
+	}
+
+	/**
+	 * Utility to remove unused members from a Rig VM asset.
+	 *
+	 * @param Asset							The Asset for which the nodes are removed
+	 * @param TransactionText				(Optional) The transaction text used when nodes are being removed. If empty or already transacting, no transaction is created.
+	 * @param AllMembersReferencedMessage	(Optional) The notification message displayed if all members are referenced and none was removed. If empty, no message is displayed.
+	 */
+	template <typename TRigVMNode> requires ConceptTypeSupportsRemoveUnusedMembers<TRigVMNode>
+	void RemoveUnusedMembersFromAsset(
+		const TScriptInterface<IRigVMAssetInterface>& Asset,
+		const FText& TransactionText = FText::GetEmpty(),
+		const FText& AllMembersReferencedMessage = FText::GetEmpty())
+	{
+		using namespace UE::RigVMEditor::Private;
+
+		URigVMController* Controller = Asset ? Asset->GetController() : nullptr;
+		const FRigVMClient* RigVMClient = Asset ? Asset->GetRigVMClient() : nullptr;
+		if (!Asset || 
+			!Asset->GetObject() ||
+			!RigVMClient || 
+			!Controller)
+		{
+			return;
+		}
+
+		// Find unused member names		
+		constexpr bool bForceUpdate = true;
+		const TSharedRef<FRigVMEdGraphNodeRegistry> EdGraphNodeRegistry = 
+			FRigVMEdGraphNodeRegistry::GetOrCreateRegistry(Asset, TRigVMNode::StaticClass(), bForceUpdate);
+
+		TMap<FRigVMUnusedMemberCategory, TArray<FName>> CategoryToUnusedMemberNamesMap;
+		FindUnusedMemberNames<TRigVMNode>(Asset, EdGraphNodeRegistry, CategoryToUnusedMemberNamesMap);
+
+		if (CategoryToUnusedMemberNamesMap.IsEmpty())
+		{
+			return;
+		}
+
+		// Present a dialog to let the user select which members to remove
+		const FRigVMEditorRemoveUnusedMembersDialog::FResult DialogResult = 
+			FRigVMEditorRemoveUnusedMembersDialog::Open(GetDialogTitle<TRigVMNode>(), CategoryToUnusedMemberNamesMap);
+
+		if (DialogResult.AppReturnType != EAppReturnType::Ok)
+		{
+			return;
+		}
+
+		if (DialogResult.MemberNames.IsEmpty() && 
+			!AllMembersReferencedMessage.IsEmpty())
+		{
+			// Notify if nothing was removed
+			FNotificationInfo Info(AllMembersReferencedMessage);
+			Info.bFireAndForget = true;
+			Info.FadeOutDuration = 5.0f;
+			Info.ExpireDuration = 5.0f;
+
+			FSlateNotificationManager::Get().AddNotification(Info);
+		}
+		else
+		{
+			const bool bTransact = !TransactionText.IsEmpty();
+			if (bTransact)
+			{
+				Asset->GetObject()->Modify();
+				Controller->OpenUndoBracket(TransactionText.ToString());
+			}
+
+			Private::RemoveDisconnectedNodes<TRigVMNode>(*Controller, EdGraphNodeRegistry, DialogResult.MemberNames);
+			Private::RemoveMembers<TRigVMNode>(Asset, DialogResult.MemberNames);
+
+			if (bTransact)
+			{
+				Controller->CloseUndoBracket();
+			}
+		}
+	}
+}
diff --git a/Engine/Plugins/Runtime/RigVM/Source/RigVMEditor/Private/Editor/RemoveUnusedMembers/RigVMEditorRemoveUnusedMembersCategory.h b/Engine/Plugins/Runtime/RigVM/Source/RigVMEditor/Private/Editor/RemoveUnusedMembers/RigVMEditorRemoveUnusedMembersCategory.h
new file mode 100644
--- /dev/null
+++ b/Engine/Plugins/Runtime/RigVM/Source/RigVMEditor/Private/Editor/RemoveUnusedMembers/RigVMEditorRemoveUnusedMembersCategory.h
@@ -0,0 +1,49 @@
+// Copyright Epic Games, Inc. All Rights Reserved.
+
+#pragma once
+
+#include "HAL/Platform.h"
+#include "Internationalization/Text.h"
+#include "UObject/NameTypes.h"
+
+class IRigVMAssetInterface;
+template <typename InInterfaceType> class TScriptInterface;
+
+namespace UE::RigVMEditor
+{	
+	/** A category for related unused members */
+	struct FRigVMUnusedMemberCategory
+	{
+		FRigVMUnusedMemberCategory() = default;
+
+		FRigVMUnusedMemberCategory(const FName InID, const FText& InCategoryNameText, const bool bInSafeToRemove)
+			: ID(InID)
+			, CategoryNameText(InCategoryNameText)
+			, bSafeToRemove(bInSafeToRemove)
+		{}
+
+		/** The category ID */
+		FName ID;
+
+		/** The text for this category */
+		FText CategoryNameText;
+
+		/** If true, the members in this category can be safely removed from the Rig */
+		bool bSafeToRemove;
+
+		bool operator==(const FRigVMUnusedMemberCategory& Other) const
+		{
+			return Other.ID == ID;
+		}
+
+		bool operator!=(const FRigVMUnusedMemberCategory& Other) const
+		{
+			return Other.ID != ID;
+		}
+
+		friend uint32 GetTypeHash(const FRigVMUnusedMemberCategory& Category)
+		{
+			return GetTypeHash(Category.ID);
+		}
+	};
+}
diff --git a/Engine/Plugins/Runtime/RigVM/Source/RigVMEditor/Private/Editor/RemoveUnusedMembers/RigVMEditorRemoveUnusedMembersDialog.cpp b/Engine/Plugins/Runtime/RigVM/Source/RigVMEditor/Private/Editor/RemoveUnusedMembers/RigVMEditorRemoveUnusedMembersDialog.cpp
new file mode 100644
--- /dev/null
+++ b/Engine/Plugins/Runtime/RigVM/Source/RigVMEditor/Private/Editor/RemoveUnusedMembers/RigVMEditorRemoveUnusedMembersDialog.cpp
@@ -0,0 +1,70 @@
+// Copyright Epic Games, Inc. All Rights Reserved.
+
+#include "RigVMEditorRemoveUnusedMembersDialog.h"
+
+#include "Framework/Application/SlateApplication.h"
+#include "Interfaces/IMainFrameModule.h"
+#include "Modules/ModuleManager.h"
+#include "Widgets/RemoveUnusedMembers/SRigVMEditorRemoveUnusedMembersDialog.h"
+#include "Widgets/SWindow.h"
+
+#define LOCTEXT_NAMESPACE "RigVMEditorRemoveUnusedMembersDialog"
+
+namespace UE::RigVMEditor
+{
+	FRigVMEditorRemoveUnusedMembersDialog::FResult FRigVMEditorRemoveUnusedMembersDialog::Open(
+		const FText& DialogTitle,
+		const TMap<FRigVMUnusedMemberCategory, TArray<FName>>& CategoryToUnusedMemberNamesMap)
+	{
+		const TSharedPtr<SWindow> ParentWindow = []() -> TSharedPtr<SWindow>
+			{
+				if (FModuleManager::Get().IsModuleLoaded("MainFrame"))
+				{
+					IMainFrameModule& MainFrame = FModuleManager::LoadModuleChecked<IMainFrameModule>("MainFrame");
+					return MainFrame.GetParentWindow();
+				}
+
+				return nullptr;
+			}();
+
+		if (!ParentWindow.IsValid())
+		{	
+			return FResult();
+		}
+
+		const FSlateRect WorkAreaRect = FSlateApplicationBase::Get().GetPreferredWorkArea();
+		const FVector2D WorkAreaSize = WorkAreaRect.GetSize();
+
+		const double WindowWidth = FMath::Min(480.0f, WorkAreaSize.X); 
+		const double WindowHeight = FMath::Min(640.0f, WorkAreaSize.Y);
+		const FVector2D WindowSize = FVector2D(WindowWidth, WindowHeight);
+
+		const FVector2D WindowPosition = WorkAreaRect.GetTopLeft() + (WorkAreaSize - WindowSize) / 2.0f;
+
+		const TSharedRef<SWindow> Window = SNew(SWindow)
+			.Title(DialogTitle)
+			.SizingRule(ESizingRule::UserSized)
+			.AutoCenter(EAutoCenter::None)
+			.ScreenPosition(WindowPosition)
+			.ClientSize(WindowSize);
+
+		FResult Result;
+		Window->SetContent
+		(
+			SNew(SRigVMEditorRemoveUnusedMembersDialog, Window, CategoryToUnusedMemberNamesMap)
+			.OnDialogEnded_Lambda(
+				[&Result](const EAppReturnType::Type AppReturnType, const TArray<FName>& SelectedMemberNames)
+				{
+					Result.AppReturnType = AppReturnType;
+					Result.MemberNames = SelectedMemberNames;
+				})
+		);
+
+		constexpr bool bSlowTaskWindow = false;
+		FSlateApplication::Get().AddModalWindow(Window, ParentWindow, bSlowTaskWindow);
+
+		return Result;
+	}
+}
+
+#undef LOCTEXT_NAMESPACE
diff --git a/Engine/Plugins/Runtime/RigVM/Source/RigVMEditor/Private/Editor/RemoveUnusedMembers/RigVMEditorRemoveUnusedMembersDialog.h b/Engine/Plugins/Runtime/RigVM/Source/RigVMEditor/Private/Editor/RemoveUnusedMembers/RigVMEditorRemoveUnusedMembersDialog.h
new file mode 100644
--- /dev/null
+++ b/Engine/Plugins/Runtime/RigVM/Source/RigVMEditor/Private/Editor/RemoveUnusedMembers/RigVMEditorRemoveUnusedMembersDialog.h
@@ -0,0 +1,40 @@
+// Copyright Epic Games, Inc. All Rights Reserved.
+
+#pragma once
+
+#include "Containers/Array.h"
+#include "GenericPlatform/GenericPlatformMisc.h"
+#include "RigVMEditorRemoveUnusedMembersCategory.h"
+#include "UObject/NameTypes.h"
+
+class IRigVMAssetInterface;
+template <typename InInterfaceType> class TScriptInterface;
+
+namespace UE::RigVMEditor
+{
+	/** Displays a dialog to remove unsued members when opened */
+	struct FRigVMEditorRemoveUnusedMembersDialog
+	{
+		/** The result of the dialog */
+		struct FResult
+		{
+			/** The return type picked by the user */
+			EAppReturnType::Type AppReturnType = EAppReturnType::Cancel;
+
+			/** The array of member names selected to be removed */
+			TArray<FName> MemberNames;
+		};
+
+		/**
+		 * Opens the dialog.
+		 *
+		 * @param DialogTitle						The tile of the dialog
+		 * @param CategoryToUnusedMemberNamesMap	A map of categories with their unused member names
+		 *
+		 * @return									The Result of the dialog.
+		 */
+		static FResult Open(
+			const FText& DialogTitle, 
+			const TMap<FRigVMUnusedMemberCategory, TArray<FName>>& CategoryToUnusedMemberNamesMap);
+	};
+}
diff --git a/Engine/Plugins/Runtime/RigVM/Source/RigVMEditor/Private/Editor/RigVMEdGraphNodeRegistry.cpp b/Engine/Plugins/Runtime/RigVM/Source/RigVMEditor/Private/Editor/RigVMEdGraphNodeRegistry.cpp
--- a/Engine/Plugins/Runtime/RigVM/Source/RigVMEditor/Private/Editor/RigVMEdGraphNodeRegistry.cpp
+++ b/Engine/Plugins/Runtime/RigVM/Source/RigVMEditor/Private/Editor/RigVMEdGraphNodeRegistry.cpp
@@ -11,7 +11,10 @@ namespace UE::RigVMEditor
 {
 	TMap<FName, TWeakPtr<FRigVMEdGraphNodeRegistry>> FRigVMEdGraphNodeRegistry::TypeIDToRegistryMap;
 
-	TSharedRef<FRigVMEdGraphNodeRegistry> FRigVMEdGraphNodeRegistry::GetOrCreateRegistry(const TScriptInterface<IRigVMAssetInterface>& InRigVMAssetInterface, UClass* InRigVMNodeClass)
+	TSharedRef<FRigVMEdGraphNodeRegistry> FRigVMEdGraphNodeRegistry::GetOrCreateRegistry(
+		const TScriptInterface<IRigVMAssetInterface>& InRigVMAssetInterface, 
+		UClass* InRigVMNodeClass,
+		const bool bForceUpdate)
 	{		
 		if (!ensureMsgf(InRigVMNodeClass->IsChildOf(URigVMNode::StaticClass()), TEXT("Node class is not a RigVMNode type, creating empty Rig VM Node Registry")) ||
 			!ensureMsgf(InRigVMAssetInterface, TEXT("Rig VM Asset Interface is invalid, creating empty Rig VM Node Registry")))
@@ -34,13 +37,26 @@ namespace UE::RigVMEditor
 		const TWeakPtr<FRigVMEdGraphNodeRegistry>* InstancePtr = TypeIDToRegistryMap.Find(TypeID);
 		if (InstancePtr && InstancePtr->IsValid())
 		{
+			if (bForceUpdate)
+			{
+				(*InstancePtr).Pin()->ForceUpdate();
+			}
+
 			return (*InstancePtr).Pin().ToSharedRef();
 		}
 		else
 		{
 			const TSharedRef<FRigVMEdGraphNodeRegistry> NewRegistry = MakeShared<FRigVMEdGraphNodeRegistry>(InRigVMAssetInterface, InRigVMNodeClass);
 			NewRegistry->Initialize();
-			NewRegistry->Update();
+			
+			if (bForceUpdate)
+			{
+				NewRegistry->ForceUpdate();
+			}
+			else
+			{
+				NewRegistry->RequestUpdate();
+			}
 
 			TypeIDToRegistryMap.Add(TypeID, NewRegistry);
 
@@ -64,15 +80,40 @@ namespace UE::RigVMEditor
 		ClientPtr->GetPostGraphModified().AddSP(this, &FRigVMEdGraphNodeRegistry::PostGraphModified);
 	}
 
-	void FRigVMEdGraphNodeRegistry::Update()
+	void FRigVMEdGraphNodeRegistry::RequestUpdate()
 	{
-		TRACE_CPUPROFILER_EVENT_SCOPE(FRigVMEdGraphNodeRegistry::Update);
+		if (!RequestUpdateTimerHandle.IsValid())
+		{
+			RequestUpdateTimerHandle = GEditor->GetTimerManager()->SetTimerForNextTick(
+				[WeakThis = AsWeak(), this]()
+				{
+					if (WeakThis.IsValid())
+					{
+						ForceUpdate();
+						RequestUpdateTimerHandle.Invalidate();
+					}
+				});
+		}
+	}
 
-		WeakEdGraphNodes.Reset();
+	void FRigVMEdGraphNodeRegistry::ForceUpdate()
+	{
+		TRACE_CPUPROFILER_EVENT_SCOPE(FRigVMEdGraphNodeRegistry::ForceUpdate);
+
+		if (!WeakAssetInterface.IsValid())
+		{
+			return;
+		}
 
 		TArray<UEdGraph*> EdGraphs;
 		WeakAssetInterface->GetAllEdGraphs(EdGraphs);
 
+		TArray<TWeakObjectPtr<URigVMEdGraphNode>> NewWeakConnectedEdGraphNodes;
+		NewWeakConnectedEdGraphNodes.Reserve(WeakConnectedEdGraphNodes.Num());
+		TArray<TWeakObjectPtr<URigVMEdGraphNode>> NewWeakDisconnectedEdGraphNodes;
+		NewWeakDisconnectedEdGraphNodes.Reserve(WeakDisconnectedEdGraphNodes.Num());
+
+		const bool bOnlyConsiderExecPin = RigVMNodeClass->IsChildOf(URigVMFunctionReferenceNode::StaticClass());
 		for (UEdGraph* Graph : EdGraphs)
 		{
 			if (URigVMEdGraph* RigVMEdGraph = Cast<URigVMEdGraph>(Graph))
@@ -82,40 +123,69 @@ namespace UE::RigVMEditor
 
 				for (URigVMEdGraphNode* EdGraphNode : EdGraphNodes)
 				{
-					if (EdGraphNode->GetModelNode() &&
-						EdGraphNode->GetModelNode()->GetClass() == RigVMNodeClass)
+					URigVMNode* ModelNode = EdGraphNode->GetModelNode();
+					if (ModelNode &&
+						ModelNode->GetClass() == RigVMNodeClass)
 					{
-						WeakEdGraphNodes.Add(EdGraphNode);
+						const bool bConnectedNode = IsNodeConnected(*EdGraphNode, bOnlyConsiderExecPin);
+
+						if (bConnectedNode)
+						{
+							NewWeakConnectedEdGraphNodes.Add(EdGraphNode);
+						}
+						else
+						{
+							NewWeakDisconnectedEdGraphNodes.Add(EdGraphNode);
+						}
 					}
 				}
 			}
 		}
 
-		OnPostRegistryUpdated.Broadcast();
+		if (NewWeakConnectedEdGraphNodes != WeakConnectedEdGraphNodes ||
+			NewWeakDisconnectedEdGraphNodes != WeakDisconnectedEdGraphNodes)
+		{
+			WeakConnectedEdGraphNodes = NewWeakConnectedEdGraphNodes;
+			WeakDisconnectedEdGraphNodes = NewWeakDisconnectedEdGraphNodes;
+
+			OnPostRegistryUpdated.Broadcast();
+		}
 	}
 
 	void FRigVMEdGraphNodeRegistry::PostGraphModified(ERigVMGraphNotifType InNotifType, URigVMGraph* InGraph, UObject* InSubject)
 	{
 		if (InNotifType == ERigVMGraphNotifType::NodeAdded ||
-			InNotifType == ERigVMGraphNotifType::NodeRemoved)
+			InNotifType == ERigVMGraphNotifType::NodeRemoved ||
+			InNotifType == ERigVMGraphNotifType::LinkAdded ||
+			InNotifType == ERigVMGraphNotifType::LinkRemoved)
 		{
 			RequestUpdate();
 		}
 	}
 
-	void FRigVMEdGraphNodeRegistry::RequestUpdate()
+	bool FRigVMEdGraphNodeRegistry::IsNodeConnected(const URigVMEdGraphNode& EdGraphNode, const bool bOnlyConsiderExecPin) const
 	{
-		if (!RequestUpdateTimerHandle.IsValid())
+		if (EdGraphNode.Pins.IsEmpty())
 		{
-			RequestUpdateTimerHandle = GEditor->GetTimerManager()->SetTimerForNextTick(
-				[WeakThis = AsWeak(), this]()
-				{
-					if (WeakThis.IsValid())
-					{
-						Update();
+			return false;
+		}
 
-						RequestUpdateTimerHandle.Invalidate();
-					}
+		// For functions with exec pins, only consider them connected when an exec pin is connected
+		if (bOnlyConsiderExecPin && 
+			!EdGraphNode.GetExecutePins().IsEmpty())
+		{
+			return EdGraphNode.GetExecutePins().ContainsByPredicate(
+				[](const URigVMPin* Pin)
+				{
+					return Pin && !Pin->GetLinks().IsEmpty();
+				});
+		}
+		else
+		{
+			return EdGraphNode.Pins.ContainsByPredicate(
+				[](const UEdGraphPin* Pin)
+				{
+					return Pin && !Pin->LinkedTo.IsEmpty();
 				});
 		}
 	}
diff --git a/Engine/Plugins/Runtime/RigVM/Source/RigVMEditor/Private/Editor/RigVMEdGraphNodeRegistry.h b/Engine/Plugins/Runtime/RigVM/Source/RigVMEditor/Private/Editor/RigVMEdGraphNodeRegistry.h
--- a/Engine/Plugins/Runtime/RigVM/Source/RigVMEditor/Private/Editor/RigVMEdGraphNodeRegistry.h
+++ b/Engine/Plugins/Runtime/RigVM/Source/RigVMEditor/Private/Editor/RigVMEdGraphNodeRegistry.h
@@ -10,6 +10,7 @@
 
 class IRigVMAssetInterface;
 class UClass;
+class UObject;
 class URigVMEdGraphNode;
 class URigVMGraph;
 enum class ERigVMGraphNotifType : uint8;
@@ -30,19 +31,29 @@ namespace UE::RigVMEditor
 		 * Creates a registry instance. 
 		 * Note, The registry will be destroyed when the all references are released.
 		 *
-		 * @param InRigVMAssetInterface			The Rig VM Asset Interface for which this registry is constructed 
+		 * @param InRigVMAssetInterface			The Asset for which this registry is constructed 
 		 * @param InRigVMNodeClass				The Rig VM Node Class for which Ed Graph nodes are held
+		 * @param bForceUpdate					(Optional) If set to true, the registry instantly updates.
 		 * 
 		 * @return								The registry instance
 		 */
-		static TSharedRef<FRigVMEdGraphNodeRegistry> GetOrCreateRegistry(const TScriptInterface<IRigVMAssetInterface>& InRigVMAssetInterface, UClass* InRigVMNodeClass);
+		static TSharedRef<FRigVMEdGraphNodeRegistry> GetOrCreateRegistry(
+			const TScriptInterface<IRigVMAssetInterface>& InRigVMAssetInterface, 
+			UClass* InRigVMNodeClass,
+			const bool bForceUpdate = false);
 
-		/** Returns the currently registered ed graph nodes which correspond to the node type */
-		const TArray<TWeakObjectPtr<URigVMEdGraphNode>>& GetWeakEdGrapNodes() const { return WeakEdGraphNodes; }
+		/** Returns the currently registered ed graph nodes which correspond to the node type and have connected pins */
+		const TArray<TWeakObjectPtr<URigVMEdGraphNode>>& GetConnectedEdGrapNodes() const { return WeakConnectedEdGraphNodes; }
 
-		/** Delegate broadcasted when the registry was updated, raised one tick after changes occured */
+		/** Returns the currently registered ed graph nodes which correspond to the node type and do not have any connected pins */
+		const TArray<TWeakObjectPtr<URigVMEdGraphNode>>& GetDisconnectedEdGrapNodes() const { return WeakDisconnectedEdGraphNodes; }
+
+		/** Delegate broadcasted when the registry was updated */
 		FSimpleMulticastDelegate OnPostRegistryUpdated;
 
+		/** Returns the node type this registry is tracking */
+		const UClass* GetType() const { return RigVMNodeClass; }
+
 	private:
 		/** Private constructor, hidden on purpose. Use GetOrCreateRegistry to create instances. */
 		FRigVMEdGraphNodeRegistry(const TScriptInterface<IRigVMAssetInterface>& InRigVMAssetInterface, UClass* InRigVMNodeClass);
@@ -50,17 +61,26 @@ namespace UE::RigVMEditor
 		/** Initializes the registry */
 		void Initialize();
 
-		/** Updates the registry from existing nodes */
-		void Update();
+		/** Requests to Update the registry */
+		void RequestUpdate();
+
+		/**
+		 * Forces the registry to update. Mind the registry tracks changes and notifies via OnPostRegistryUpdated.
+		 * ForceUpdate only needs to be called where handling changes via OnPostRegistryUpdated is not adequate.
+		 */
+		void ForceUpdate();
 
 		/** Called after a graph was modified */
 		void PostGraphModified(ERigVMGraphNotifType InNotifType, URigVMGraph* InGraph, UObject* InSubject);
 
-		/** Requests to Update the registry */
-		void RequestUpdate();
+		/** Returns true if the node is considered connected */
+		bool IsNodeConnected(const URigVMEdGraphNode& EdGraphNode, const bool bOnlyConsiderExecPin) const;
+
+		/** The currently registered nodes that have any pin connected */
+		TArray<TWeakObjectPtr<URigVMEdGraphNode>> WeakConnectedEdGraphNodes;
 
-		/** The currently registered nodes */
-		TArray<TWeakObjectPtr<URigVMEdGraphNode>> WeakEdGraphNodes;
+		/** The currently registered nodes that do not have any pin connected */
+		TArray<TWeakObjectPtr<URigVMEdGraphNode>> WeakDisconnectedEdGraphNodes;
 
 		/** The asset interface this registry refers to */
 		TWeakInterfacePtr<IRigVMAssetInterface> WeakAssetInterface;
diff --git a/Engine/Plugins/Runtime/RigVM/Source/RigVMEditor/Private/Widgets/RemoveUnusedMembers/RigVMEditorRemoveUnusedMembersColumnID.cpp b/Engine/Plugins/Runtime/RigVM/Source/RigVMEditor/Private/Widgets/RemoveUnusedMembers/RigVMEditorRemoveUnusedMembersColumnID.cpp
new file mode 100644
--- /dev/null
+++ b/Engine/Plugins/Runtime/RigVM/Source/RigVMEditor/Private/Widgets/RemoveUnusedMembers/RigVMEditorRemoveUnusedMembersColumnID.cpp
@@ -0,0 +1,9 @@
+// Copyright Epic Games, Inc. All Rights Reserved.
+
+#include "RigVMEditorRemoveUnusedMembersColumnID.h"
+
+namespace UE::RigVMEditor
+{
+	const FName FRigVMEditorRemoveUnusedMembersColumnID::LabelRow = "Label";
+	const FName FRigVMEditorRemoveUnusedMembersColumnID::SelectedCheckBox = "SelectedCheckBox";
+}
diff --git a/Engine/Plugins/Runtime/RigVM/Source/RigVMEditor/Private/Widgets/RemoveUnusedMembers/RigVMEditorRemoveUnusedMembersColumnID.h b/Engine/Plugins/Runtime/RigVM/Source/RigVMEditor/Private/Widgets/RemoveUnusedMembers/RigVMEditorRemoveUnusedMembersColumnID.h
new file mode 100644
--- /dev/null
+++ b/Engine/Plugins/Runtime/RigVM/Source/RigVMEditor/Private/Widgets/RemoveUnusedMembers/RigVMEditorRemoveUnusedMembersColumnID.h
@@ -0,0 +1,17 @@
+// Copyright Epic Games, Inc. All Rights Reserved.
+
+#pragma once
+
+#include "UObject/NameTypes.h"
+
+namespace UE::RigVMEditor
+{
+	struct FRigVMEditorRemoveUnusedMembersColumnID
+	{
+		/** The column to display the member name label */
+		static const FName LabelRow;
+
+		/** The column to display a checkbox to chose if the member should be removed */
+		static const FName SelectedCheckBox;
+	};
+}
diff --git a/Engine/Plugins/Runtime/RigVM/Source/RigVMEditor/Private/Widgets/RemoveUnusedMembers/RigVMEditorRemoveUnusedMembersListItem.cpp b/Engine/Plugins/Runtime/RigVM/Source/RigVMEditor/Private/Widgets/RemoveUnusedMembers/RigVMEditorRemoveUnusedMembersListItem.cpp
new file mode 100644
--- /dev/null
+++ b/Engine/Plugins/Runtime/RigVM/Source/RigVMEditor/Private/Widgets/RemoveUnusedMembers/RigVMEditorRemoveUnusedMembersListItem.cpp
@@ -0,0 +1,13 @@
+// Copyright Epic Games, Inc. All Rights Reserved.
+
+#include "RigVMEditorRemoveUnusedMembersListItem.h"
+
+namespace UE::RigVMEditor
+{
+	FRigVMEditorRemoveUnusedMembersListItem::FRigVMEditorRemoveUnusedMembersListItem(
+		const FName& InMemberName, 
+		const bool bInInitiallySelected)
+		: MemberName(InMemberName)
+		, bSelected(bInInitiallySelected)
+	{}
+}
diff --git a/Engine/Plugins/Runtime/RigVM/Source/RigVMEditor/Private/Widgets/RemoveUnusedMembers/RigVMEditorRemoveUnusedMembersListItem.h b/Engine/Plugins/Runtime/RigVM/Source/RigVMEditor/Private/Widgets/RemoveUnusedMembers/RigVMEditorRemoveUnusedMembersListItem.h
new file mode 100644
--- /dev/null
+++ b/Engine/Plugins/Runtime/RigVM/Source/RigVMEditor/Private/Widgets/RemoveUnusedMembers/RigVMEditorRemoveUnusedMembersListItem.h
@@ -0,0 +1,31 @@
+// Copyright Epic Games, Inc. All Rights Reserved.
+
+#pragma once
+
+#include "UObject/NameTypes.h"
+
+namespace UE::RigVMEditor
+{
+	/** An item in the unused members list */
+	struct FRigVMEditorRemoveUnusedMembersListItem
+	{
+		FRigVMEditorRemoveUnusedMembersListItem() = default;
+		FRigVMEditorRemoveUnusedMembersListItem(const FName& InMemberName, const bool bInInitiallySelected);
+
+		/** Returns the member name this item stands for */
+		const FName& GetMemberName() const { return MemberName; }
+
+		/** Returns true if the item is selected */
+		bool IsSelected() const { return bSelected; }
+
+		/** Sets the itetm selected */
+		void SetSelected(const bool bInSelected) { bSelected = bInSelected; }
+
+	private:
+		/** The member name this item stands for */
+		FName MemberName;
+
+		/** If true the item is selected */
+		bool bSelected = false;
+	};
+}
diff --git a/Engine/Plugins/Runtime/RigVM/Source/RigVMEditor/Private/Widgets/RemoveUnusedMembers/SRigVMEditorRemoveUnusedMembersDialog.cpp b/Engine/Plugins/Runtime/RigVM/Source/RigVMEditor/Private/Widgets/RemoveUnusedMembers/SRigVMEditorRemoveUnusedMembersDialog.cpp
new file mode 100644
--- /dev/null
+++ b/Engine/Plugins/Runtime/RigVM/Source/RigVMEditor/Private/Widgets/RemoveUnusedMembers/SRigVMEditorRemoveUnusedMembersDialog.cpp
@@ -0,0 +1,150 @@
+// Copyright Epic Games, Inc. All Rights Reserved.
+
+#include "SRigVMEditorRemoveUnusedMembersDialog.h"
+
+#include "SRigVMEditorRemoveUnusedMembersList.h"
+#include "Widgets/Input/SButton.h"
+#include "Widgets/Layout/SScrollBox.h"
+#include "Widgets/SWindow.h"
+
+#define LOCTEXT_NAMESPACE "SRigVMEditorRemoveUnusedMembersDialog"
+
+namespace UE::RigVMEditor
+{
+	SRigVMEditorRemoveUnusedMembersDialog::~SRigVMEditorRemoveUnusedMembersDialog()
+	{
+		if (OnDialogEndedDelegate.IsBound())
+		{
+			OnDialogEndedDelegate.Execute(EAppReturnType::Cancel, TArray<FName>());
+		}
+	}
+
+	void SRigVMEditorRemoveUnusedMembersDialog::Construct(
+		const FArguments& InArgs,
+		const TSharedRef<SWindow>& InOwningWindow,
+		const TMap<FRigVMUnusedMemberCategory, TArray<FName>>& CategoryToUnusedMemberNamesMap)
+	{
+		WeakOwningWindow = InOwningWindow;
+
+		OnDialogEndedDelegate = InArgs._OnDialogEnded;
+
+		const TSharedRef<SScrollBox> ScrollBox = SNew(SScrollBox);
+		for (const TTuple<FRigVMUnusedMemberCategory, TArray<FName>>& CategoryToUnusedMemberNamesPair : CategoryToUnusedMemberNamesMap)
+		{
+			ScrollBox->AddSlot()
+				[
+					CreateCategory(CategoryToUnusedMemberNamesPair.Key, CategoryToUnusedMemberNamesPair.Value)
+				];
+		}
+
+		ChildSlot
+		[
+			SNew(SBorder)
+			.HAlign(HAlign_Fill)
+			.VAlign(VAlign_Fill)
+			.Padding(4.f)
+			[
+				SNew(SVerticalBox)
+
+				+ SVerticalBox::Slot()
+				.FillHeight(1.f)
+				[
+					ScrollBox
+				]
+
+				+ SVerticalBox::Slot()
+
+				.AutoHeight()
+				.HAlign(HAlign_Right)
+				[
+					CreateDialogButtons()
+				]
+			]
+		];
+	}
+
+	TSharedRef<SWidget> SRigVMEditorRemoveUnusedMembersDialog::CreateCategory(
+		const FRigVMUnusedMemberCategory& Category,
+		const TArray<FName>& UnusedMemberNames)
+	{
+		const TSharedRef<SRigVMEditorRemoveUnusedMembersList> UnusedMembersList = 
+			SNew(SRigVMEditorRemoveUnusedMembersList)
+			.Category(Category)
+			.MemberNames(UnusedMemberNames);
+
+		UnusedMembersListViews.Add(UnusedMembersList);
+		
+		return UnusedMembersList;
+	}
+
+	TSharedRef<SWidget> SRigVMEditorRemoveUnusedMembersDialog::CreateDialogButtons()
+	{
+		return
+			SNew(SHorizontalBox)
+
+			+ SHorizontalBox::Slot()
+			.AutoWidth()
+			.MinWidth(60.f)
+			.Padding(4.f, 0.f)
+			[
+				SNew(SButton)
+				.HAlign(HAlign_Center)
+				.Text(LOCTEXT("OkButtonText", "Ok"))
+				.OnClicked(this, &SRigVMEditorRemoveUnusedMembersDialog::OnOkClicked)
+			]
+
+			+ SHorizontalBox::Slot()
+			.AutoWidth()
+			.MinWidth(60.f)
+			[
+				SNew(SButton)
+				.HAlign(HAlign_Center)
+				.Text(LOCTEXT("CancelButtonText", "Cancel"))
+				.OnClicked(this, &SRigVMEditorRemoveUnusedMembersDialog::OnCancelClicked)
+			];
+	}
+
+	FReply SRigVMEditorRemoveUnusedMembersDialog::OnOkClicked()
+	{
+		if (OnDialogEndedDelegate.IsBound())
+		{
+			TArray<FName> SelectedMemberNames;
+			for (const TSharedPtr<SRigVMEditorRemoveUnusedMembersList>& UnusedMembersListView : UnusedMembersListViews)
+			{
+				SelectedMemberNames.Append(UnusedMembersListView->GetSelectedMemberNames());
+			}
+
+			OnDialogEndedDelegate.Execute(EAppReturnType::Ok, SelectedMemberNames);
+
+			// Only raise once
+			OnDialogEndedDelegate.Unbind();
+		}
+
+		if (WeakOwningWindow.IsValid())
+		{
+			WeakOwningWindow.Pin()->RequestDestroyWindow();
+		}
+
+		return FReply::Handled();
+	}
+
+	FReply SRigVMEditorRemoveUnusedMembersDialog::OnCancelClicked()
+	{
+		if (OnDialogEndedDelegate.IsBound())
+		{
+			OnDialogEndedDelegate.Execute(EAppReturnType::Cancel, TArray<FName>());
+
+			// Only raise once
+			OnDialogEndedDelegate.Unbind();
+		}
+
+		if (WeakOwningWindow.IsValid())
+		{
+			WeakOwningWindow.Pin()->RequestDestroyWindow();
+		}
+
+		return FReply::Handled();
+	}
+}
+
+#undef LOCTEXT_NAMESPACE
diff --git a/Engine/Plugins/Runtime/RigVM/Source/RigVMEditor/Private/Widgets/RemoveUnusedMembers/SRigVMEditorRemoveUnusedMembersDialog.h b/Engine/Plugins/Runtime/RigVM/Source/RigVMEditor/Private/Widgets/RemoveUnusedMembers/SRigVMEditorRemoveUnusedMembersDialog.h
new file mode 100644
--- /dev/null
+++ b/Engine/Plugins/Runtime/RigVM/Source/RigVMEditor/Private/Widgets/RemoveUnusedMembers/SRigVMEditorRemoveUnusedMembersDialog.h
@@ -0,0 +1,65 @@
+// Copyright Epic Games, Inc. All Rights Reserved.
+
+#pragma once
+
+#include "Containers/Array.h"
+#include "Containers/Map.h"
+#include "Editor/RemoveUnusedMembers/RigVMEditorRemoveUnusedMembersCategory.h"
+#include "GenericPlatform/GenericPlatformMisc.h"
+#include "Input/Reply.h"
+#include "Widgets/SCompoundWidget.h"
+
+template <typename ItemType> class SListView;
+
+namespace UE::RigVMEditor
+{
+	class SRigVMEditorRemoveUnusedMembersList;
+
+	DECLARE_DELEGATE_TwoParams(FRigVMEditorOnRemoveUnusedMembersDialogEnded, const EAppReturnType::Type /** AppReturnType */, const TArray<FName>& /** SelectedMemberNames*/);
+
+	/** The dialog displayed when removing unused members */
+	class SRigVMEditorRemoveUnusedMembersDialog
+		: public SCompoundWidget
+	{
+	public:
+		SLATE_BEGIN_ARGS(SRigVMEditorRemoveUnusedMembersDialog)
+		{}
+
+			/** Delegate executed when the dialog is closing */
+			SLATE_EVENT(FRigVMEditorOnRemoveUnusedMembersDialogEnded, OnDialogEnded)
+
+		SLATE_END_ARGS()
+
+		virtual ~SRigVMEditorRemoveUnusedMembersDialog();
+
+		/** Constructs this widget */
+		void Construct(
+			const FArguments& InArgs, 
+			const TSharedRef<SWindow>& InOwningWindow, 
+			const TMap<FRigVMUnusedMemberCategory, TArray<FName>>& CategoryToUnusedMemberNamesMap);
+
+	private:
+		/** Creates a category with its unused members */
+		TSharedRef<SWidget> CreateCategory(
+			const FRigVMUnusedMemberCategory& Category, 
+			const TArray<FName>& UnusedMemberNames);
+
+		/** Creates the widget that displays the dialog buttons */
+		TSharedRef<SWidget> CreateDialogButtons();
+
+		/** Called when the ok button was clicked */
+		FReply OnOkClicked();
+		
+		/** Called when the cancel button was clicked */
+		FReply OnCancelClicked();
+
+		/** The widget that owns this dialog */
+		TWeakPtr<SWindow> WeakOwningWindow;
+
+		/** The list views of unused members displayed in this dialog */
+		TArray<TSharedPtr<SRigVMEditorRemoveUnusedMembersList>> UnusedMembersListViews;
+
+		/** Delegate executed when this dialog ended */
+		FRigVMEditorOnRemoveUnusedMembersDialogEnded OnDialogEndedDelegate;
+	};
+}
diff --git a/Engine/Plugins/Runtime/RigVM/Source/RigVMEditor/Private/Widgets/RemoveUnusedMembers/SRigVMEditorRemoveUnusedMembersList.cpp b/Engine/Plugins/Runtime/RigVM/Source/RigVMEditor/Private/Widgets/RemoveUnusedMembers/SRigVMEditorRemoveUnusedMembersList.cpp
new file mode 100644
--- /dev/null
+++ b/Engine/Plugins/Runtime/RigVM/Source/RigVMEditor/Private/Widgets/RemoveUnusedMembers/SRigVMEditorRemoveUnusedMembersList.cpp
@@ -0,0 +1,245 @@
+// Copyright Epic Games, Inc. All Rights Reserved.
+
+#include "SRigVMEditorRemoveUnusedMembersList.h"
+
+#include "RigVMEditorRemoveUnusedMembersColumnID.h"
+#include "RigVMEditorRemoveUnusedMembersListItem.h"
+#include "SRigVMEditorRemoveUnusedMembersRow.h"
+#include "Widgets/Input/SCheckBox.h"
+
+#define LOCTEXT_NAMESPACE "SRigVMEditorRemoveUnusedMembersList"
+
+namespace UE::RigVMEditor
+{
+	void SRigVMEditorRemoveUnusedMembersList::Construct(const FArguments& InArgs)
+	{
+		if (InArgs._MemberNames.IsEmpty())
+		{
+			return;
+		}
+
+		Category = InArgs._Category;
+
+		const bool bInitiallySelectItems = Category.bSafeToRemove;
+		for (const FName& MemberName : InArgs._MemberNames)
+		{
+			Items.Add(MakeShared<FRigVMEditorRemoveUnusedMembersListItem>(MemberName, bInitiallySelectItems));
+		}
+
+		ChildSlot
+		[
+			SNew(SListView<TSharedPtr<FRigVMEditorRemoveUnusedMembersListItem>>)
+			.SelectionMode(ESelectionMode::None)
+			.ListItemsSource(&Items)
+			.HeaderRow(GenerateHeaderRow())
+			.OnGenerateRow(this, &SRigVMEditorRemoveUnusedMembersList::GenerateRow)
+			.Visibility(this, &SRigVMEditorRemoveUnusedMembersList::GetVisibility)
+		];
+	}
+
+	TSharedRef<SHeaderRow> SRigVMEditorRemoveUnusedMembersList::GenerateHeaderRow()
+	{
+		const TSharedRef<SHeaderRow> HeaderRow = SNew(SHeaderRow);
+
+		// Label Column
+		{
+			HeaderRow->AddColumn(
+				SHeaderRow::FColumn::FArguments()
+				.ColumnId(FRigVMEditorRemoveUnusedMembersColumnID::LabelRow)
+				.MinSize(40.f)
+				.FillWidth(0.9f)
+				.SortMode(this, &SRigVMEditorRemoveUnusedMembersList::GetColumnSortMode, FRigVMEditorRemoveUnusedMembersColumnID::LabelRow)
+				.OnSort(this, &SRigVMEditorRemoveUnusedMembersList::SortByColumnID)
+				.HeaderContent()
+				[
+					GenerateHeaderRowLabel()
+				]
+			);
+		}
+
+		// Select column
+		{
+			HeaderRow->AddColumn(
+				SHeaderRow::FColumn::FArguments()
+				.ColumnId(FRigVMEditorRemoveUnusedMembersColumnID::SelectedCheckBox)
+				.MinSize(40.f)
+				.FillWidth(0.5f)
+				.VAlignHeader(VAlign_Center)
+				.HeaderContent()
+				[
+					SNew(SBorder)
+					.BorderImage(FAppStyle::GetBrush("NoBorder"))
+					.HAlign(HAlign_Center)
+					[
+						SNew(SCheckBox)
+						.HAlign(HAlign_Center)
+						.IsChecked(this, &SRigVMEditorRemoveUnusedMembersList::IsHeaderRowChecked)
+						.OnCheckStateChanged(this, &SRigVMEditorRemoveUnusedMembersList::OnHeaderRowCheckStateChanged)
+					]
+				]
+			);
+		}
+
+		return HeaderRow;
+	}
+
+	TSharedRef<ITableRow> SRigVMEditorRemoveUnusedMembersList::GenerateRow(
+		TSharedPtr<FRigVMEditorRemoveUnusedMembersListItem> Item, 
+		const TSharedRef<STableViewBase>& OwnerTable)
+	{
+		return SNew(SRigVMEditorRemoveUnusedMembersRow, Item, OwnerTable);
+	}
+
+	TSharedRef<SWidget> SRigVMEditorRemoveUnusedMembersList::GenerateHeaderRowLabel() const
+	{
+		if (Category.bSafeToRemove)
+		{
+			const FText LabelText = FText::Format(
+				LOCTEXT("RemoveUnusedPrivateMemberNameLabel", "Unused {0}"),
+				Category.CategoryNameText);
+
+			return 
+				SNew(SBorder)
+				.BorderImage(FAppStyle::GetBrush("NoBorder"))
+				.HAlign(HAlign_Left)
+				.VAlign(VAlign_Center)
+				[
+					SNew(STextBlock)
+					.Text(LabelText)
+					.Font(FAppStyle::GetFontStyle("DetailsView.CategoryFontStyle"))
+				];
+		}
+		else
+		{
+			const FText LabelText = FText::Format(
+				LOCTEXT("ForceRemoveUnusedPublicMemberNameLabel", "Force Remove Unused {0}"),
+				Category.CategoryNameText);
+
+			const FText WarningToolTip = FText::Format(
+				LOCTEXT("PublicMemberWarningTooltip", "Warning! {0} might be referenced in other Rig Assets."),
+				Category.CategoryNameText);
+
+			const TSharedRef<SWidget> WarningIcon = 
+				SNew(SImage)
+				.Image(FAppStyle::Get().GetBrush("Icons.Warning"))
+				.ColorAndOpacity(FAppStyle::Get().GetSlateColor("Colors.AccentYellow"))
+				.ToolTipText(WarningToolTip);
+
+			const TSharedRef<SWidget> Label =
+				SNew(STextBlock)
+				.Text(LabelText)
+				.Font(FAppStyle::GetFontStyle("DetailsView.CategoryFontStyle"))
+				.ToolTipText(WarningToolTip);
+
+			return 
+				SNew(SBorder)
+				.BorderImage(FAppStyle::GetBrush("NoBorder"))
+				.HAlign(HAlign_Left)
+				.VAlign(VAlign_Center)
+				[
+					SNew(SHorizontalBox)
+
+					+ SHorizontalBox::Slot()
+					.AutoWidth()
+					.HAlign(HAlign_Left)
+					.VAlign(VAlign_Center)
+					[
+						WarningIcon
+					]
+
+					+ SHorizontalBox::Slot()
+					.FillWidth(1.f)
+					.Padding(4.f, 0.f)
+					.HAlign(HAlign_Left)
+					.VAlign(VAlign_Center)
+					[
+						Label
+					]
+				];
+		}
+	}
+
+	TArray<FName> SRigVMEditorRemoveUnusedMembersList::GetSelectedMemberNames() const
+	{
+		TArray<FName> SelectedMemberNames;
+
+		Algo::TransformIf(Items, SelectedMemberNames,
+			[](const TSharedPtr<FRigVMEditorRemoveUnusedMembersListItem>& Item)
+			{
+				return Item.IsValid() && Item->IsSelected();
+			},
+			[](const TSharedPtr<FRigVMEditorRemoveUnusedMembersListItem>& Item)
+			{
+				return Item->GetMemberName();
+			});
+
+		return SelectedMemberNames;
+	}
+
+	ECheckBoxState SRigVMEditorRemoveUnusedMembersList::IsHeaderRowChecked() const
+	{
+		bool bSelected = false;
+		for (const TSharedPtr<FRigVMEditorRemoveUnusedMembersListItem>& Item : Items)
+		{
+			if (!bSelected && Item->IsSelected())
+			{
+				bSelected = true;
+			}
+
+			if (bSelected && !Item->IsSelected())
+			{
+				return ECheckBoxState::Undetermined;
+			}
+		}
+
+		if (bSelected)
+		{
+			return ECheckBoxState::Checked;
+		}
+		else
+		{
+			return ECheckBoxState::Unchecked;
+		}
+	}
+
+	void SRigVMEditorRemoveUnusedMembersList::OnHeaderRowCheckStateChanged(ECheckBoxState NewState)
+	{
+		const bool bSelect = NewState == ECheckBoxState::Checked;
+
+		for (const TSharedPtr<FRigVMEditorRemoveUnusedMembersListItem>& Item : Items)
+		{
+			Item->SetSelected(bSelect);
+		}
+	}
+
+	EColumnSortMode::Type SRigVMEditorRemoveUnusedMembersList::GetColumnSortMode(const FName ColumnID) const
+	{
+		if (SortedByColumnID != ColumnID)
+		{
+			return EColumnSortMode::None;
+		}
+
+		return ColumnSortMode;
+	}
+
+	void SRigVMEditorRemoveUnusedMembersList::SortByColumnID(
+		const EColumnSortPriority::Type SortPriority, 
+		const FName& ColumnID, 
+		const EColumnSortMode::Type InSortMode)
+	{
+		const bool bAscending = InSortMode == EColumnSortMode::Ascending ? true : false;
+
+		if (ColumnID == FRigVMEditorRemoveUnusedMembersColumnID::LabelRow)
+		{
+			Algo::StableSort(Items, 
+				[bAscending](
+					const TSharedPtr<FRigVMEditorRemoveUnusedMembersListItem>& ItemA,
+					const TSharedPtr<FRigVMEditorRemoveUnusedMembersListItem>& ItemB)
+				{
+					return ItemA->GetMemberName().ToString() > ItemB->GetMemberName().ToString();
+				});
+		}
+	}
+}
+
+#undef LOCTEXT_NAMESPACE
diff --git a/Engine/Plugins/Runtime/RigVM/Source/RigVMEditor/Private/Widgets/RemoveUnusedMembers/SRigVMEditorRemoveUnusedMembersList.h b/Engine/Plugins/Runtime/RigVM/Source/RigVMEditor/Private/Widgets/RemoveUnusedMembers/SRigVMEditorRemoveUnusedMembersList.h
new file mode 100644
--- /dev/null
+++ b/Engine/Plugins/Runtime/RigVM/Source/RigVMEditor/Private/Widgets/RemoveUnusedMembers/SRigVMEditorRemoveUnusedMembersList.h
@@ -0,0 +1,74 @@
+// Copyright Epic Games, Inc. All Rights Reserved.
+
+#pragma once
+
+#include "Editor/RemoveUnusedMembers/RigVMEditorRemoveUnusedMembersCategory.h"
+#include "Styling/SlateTypes.h"
+#include "Widgets/SCompoundWidget.h"
+#include "Widgets/Views/SHeaderRow.h"
+
+class ITableRow;
+class SHeaderRow;
+class STableViewBase;
+
+namespace UE::RigVMEditor
+{
+	struct FRigVMEditorRemoveUnusedMembersListItem;
+
+	/** A list displaying unused members along with an option to remove each item in the list */
+	class SRigVMEditorRemoveUnusedMembersList
+		: public SCompoundWidget
+	{
+	public:
+		SLATE_BEGIN_ARGS(SRigVMEditorRemoveUnusedMembersList)
+			{}
+
+			/** The category to display in the header */
+			SLATE_ARGUMENT(FRigVMUnusedMemberCategory, Category)
+
+			/** The member names to display */
+			SLATE_ARGUMENT(TArray<FName>, MemberNames)
+
+		SLATE_END_ARGS()
+
+		/** Constructs this widget */
+		void Construct(const FArguments& InArgs);
+
+		/** Returns the selected member names in this list */
+		TArray<FName> GetSelectedMemberNames() const;
+
+	private:
+		/** Generates the header row for the list */
+		TSharedRef<SHeaderRow> GenerateHeaderRow();
+
+		/** Generates a row in the list */
+		TSharedRef<ITableRow> GenerateRow(TSharedPtr<FRigVMEditorRemoveUnusedMembersListItem> Item, const TSharedRef<STableViewBase>& OwnerTable);
+
+		/** Generates the header row label */
+		TSharedRef<SWidget> GenerateHeaderRowLabel() const;
+
+		/** Returns true if the checkbox in the header row is checked */
+		ECheckBoxState IsHeaderRowChecked() const;
+
+		/** Called when the check state of the checkbox in the header row changed */
+		void OnHeaderRowCheckStateChanged(ECheckBoxState NewState);
+
+		/** Returns the sort mode for the column */
+		EColumnSortMode::Type GetColumnSortMode(const FName ColumnID) const;
+
+		/** Sorts the list by column ID */
+		void SortByColumnID(const EColumnSortPriority::Type SortPriority, const FName& ColumnID, const EColumnSortMode::Type InSortMode);
+
+		/** The category to display in the header */
+		FRigVMUnusedMemberCategory Category;
+
+		/** The member names in this list */
+		TArray<TSharedPtr<FRigVMEditorRemoveUnusedMembersListItem>> Items;
+
+		/** The collumn ID by which this list is sorted */
+		FName SortedByColumnID;
+
+		/** The sort mode by which sorted by column ID is currently sorted */
+		EColumnSortMode::Type ColumnSortMode = EColumnSortMode::None;
+	};
+}
diff --git a/Engine/Plugins/Runtime/RigVM/Source/RigVMEditor/Private/Widgets/RemoveUnusedMembers/SRigVMEditorRemoveUnusedMembersRow.cpp b/Engine/Plugins/Runtime/RigVM/Source/RigVMEditor/Private/Widgets/RemoveUnusedMembers/SRigVMEditorRemoveUnusedMembersRow.cpp
new file mode 100644
--- /dev/null
+++ b/Engine/Plugins/Runtime/RigVM/Source/RigVMEditor/Private/Widgets/RemoveUnusedMembers/SRigVMEditorRemoveUnusedMembersRow.cpp
@@ -0,0 +1,87 @@
+// Copyright Epic Games, Inc. All Rights Reserved.
+
+#include "SRigVMEditorRemoveUnusedMembersRow.h"
+
+#include "RigVMEditorRemoveUnusedMembersColumnID.h"
+#include "RigVMEditorRemoveUnusedMembersListItem.h"
+#include "Widgets/Input/SCheckBox.h"
+#include "Widgets/Text/STextBlock.h"
+
+namespace UE::RigVMEditor
+{
+	void SRigVMEditorRemoveUnusedMembersRow::Construct(
+		const FArguments& InArgs, 
+		const TSharedPtr<FRigVMEditorRemoveUnusedMembersListItem>& InItem,
+		const TSharedRef<STableViewBase>& InOwnerTableView)
+	{
+		WeakItem = InItem;
+
+		SMultiColumnTableRow<TSharedPtr<FRigVMEditorRemoveUnusedMembersListItem>>::Construct(
+			FSuperRowType::FArguments()
+			.Padding(FMargin(8.f, 0.f, 0.f, 0.f)),
+			InOwnerTableView);
+	}
+
+	TSharedRef<SWidget> SRigVMEditorRemoveUnusedMembersRow::GenerateWidgetForColumn(const FName& ColumnID)
+	{
+		if (WeakItem.IsValid())
+		{
+			const TSharedRef<FRigVMEditorRemoveUnusedMembersListItem> Item = WeakItem.Pin().ToSharedRef();
+
+			if (ColumnID == FRigVMEditorRemoveUnusedMembersColumnID::LabelRow)
+			{
+				return
+					SNew(SBorder)
+					.BorderImage(FAppStyle::GetBrush("NoBorder"))
+					.HAlign(HAlign_Left)
+					.VAlign(VAlign_Center)
+					[
+						SNew(STextBlock)
+						.Text(FText::FromName(Item->GetMemberName()))
+						.Font(FAppStyle::GetFontStyle("PropertyWindow.NormalFont"))
+					];
+			}
+			else if (ColumnID == FRigVMEditorRemoveUnusedMembersColumnID::SelectedCheckBox)
+			{
+				return
+					SNew(SBorder)
+					.BorderImage(FAppStyle::GetBrush("NoBorder"))
+					.HAlign(HAlign_Center)
+					[
+						SNew(SCheckBox)
+						.HAlign(HAlign_Center)
+						.IsChecked(this, &SRigVMEditorRemoveUnusedMembersRow::IsChecked)
+						.OnCheckStateChanged(this, &SRigVMEditorRemoveUnusedMembersRow::OnCheckStateChanged)
+					];
+			}
+			else
+			{
+				ensureMsgf(0, TEXT("Unhandled column ID"));
+			}
+		}
+
+		return SNullWidget::NullWidget;
+	}
+
+	ECheckBoxState SRigVMEditorRemoveUnusedMembersRow::IsChecked() const
+	{
+		const bool bSelected = WeakItem.IsValid() && WeakItem.Pin()->IsSelected() ? true : false;
+		
+		return bSelected ? ECheckBoxState::Checked : ECheckBoxState::Unchecked;
+	}
+
+	void SRigVMEditorRemoveUnusedMembersRow::OnCheckStateChanged(ECheckBoxState NewState)
+	{
+		if (WeakItem.IsValid())
+		{
+			if (NewState == ECheckBoxState::Checked)
+			{
+				WeakItem.Pin()->SetSelected(true);
+			}
+			else if (NewState == ECheckBoxState::Unchecked)
+			{
+				WeakItem.Pin()->SetSelected(false);
+			}
+		}
+	}
+}
diff --git a/Engine/Plugins/Runtime/RigVM/Source/RigVMEditor/Private/Widgets/RemoveUnusedMembers/SRigVMEditorRemoveUnusedMembersRow.h b/Engine/Plugins/Runtime/RigVM/Source/RigVMEditor/Private/Widgets/RemoveUnusedMembers/SRigVMEditorRemoveUnusedMembersRow.h
new file mode 100644
--- /dev/null
+++ b/Engine/Plugins/Runtime/RigVM/Source/RigVMEditor/Private/Widgets/RemoveUnusedMembers/SRigVMEditorRemoveUnusedMembersRow.h
@@ -0,0 +1,43 @@
+// Copyright Epic Games, Inc. All Rights Reserved.
+
+#pragma once
+
+#include "Styling/SlateTypes.h"
+#include "UObject/NameTypes.h"
+#include "Widgets/Views/SListView.h"
+#include "Widgets/Views/STableRow.h"
+
+namespace UE::RigVMEditor
+{
+	struct FRigVMEditorRemoveUnusedMembersListItem;
+
+	/** A row displaying a member along with a checkbox to select if it should be removed */
+	class SRigVMEditorRemoveUnusedMembersRow
+		: public SMultiColumnTableRow<TSharedPtr<FRigVMEditorRemoveUnusedMembersListItem>>
+	{
+	public:
+		SLATE_BEGIN_ARGS(SRigVMEditorRemoveUnusedMembersRow)
+		{}
+
+		SLATE_END_ARGS()
+
+		/** Constructs this widget */
+		void Construct(
+			const FArguments& InArgs, 
+			const TSharedPtr<FRigVMEditorRemoveUnusedMembersListItem>& InItem,
+			const TSharedRef<STableViewBase>& InOwnerTableView);
+
+	private:
+		/** Generates a widget for this column ID */
+		virtual TSharedRef<SWidget> GenerateWidgetForColumn(const FName& ColumnID) override;
+
+		/** Returns true if the checkbox in this row is checked */
+		ECheckBoxState IsChecked() const;
+
+		/** Called when the check state of the checkbox in this row changed */
+		void OnCheckStateChanged(ECheckBoxState NewState);
+
+		/** The item this row displays */
+		TWeakPtr<FRigVMEditorRemoveUnusedMembersListItem> WeakItem;
+	};
+}
diff --git a/Engine/Plugins/Runtime/RigVM/Source/RigVMEditor/Private/Widgets/SRigVMEditorGraphExplorer.cpp b/Engine/Plugins/Runtime/RigVM/Source/RigVMEditor/Private/Widgets/SRigVMEditorGraphExplorer.cpp
--- a/Engine/Plugins/Runtime/RigVM/Source/RigVMEditor/Private/Widgets/SRigVMEditorGraphExplorer.cpp
+++ b/Engine/Plugins/Runtime/RigVM/Source/RigVMEditor/Private/Widgets/SRigVMEditorGraphExplorer.cpp
@@ -2,6 +2,7 @@
 
 #include "Widgets/SRigVMEditorGraphExplorer.h"
 
+#include "Editor/RemoveUnusedMembers/RigVMEditorRemoveUnusedMembers.h"
 #include "Editor/RigVMEdGraphNodeRegistry.h"
 #include "Editor/RigVMNewEditor.h"
 #include "EdGraph/RigVMEdGraphSchema.h"
@@ -38,7 +39,6 @@
 
 #define LOCTEXT_NAMESPACE "RigVMGraphExplorer"
 
-
 TSharedRef<FRigVMGraphExplorerDragDropOp> FRigVMGraphExplorerDragDropOp::New(const FRigVMExplorerElementKey& InElement, TObjectPtr<URigVMBlueprint> InBlueprint)
 {
 	const FRigVMAssetInterfacePtr Asset(InBlueprint);
@@ -204,7 +204,8 @@ void FRigVMEditorGraphExplorerCommands::RegisterCommands()
 	UI_COMMAND(PasteFunction, "Paste Function", "Pastes the function.", EUserInterfaceActionType::Button, FInputChord());
 	UI_COMMAND(PasteVariable, "Paste Variable", "Pastes the variable.", EUserInterfaceActionType::Button, FInputChord());
 	UI_COMMAND(PasteLocalVariable, "Paste Local Variable", "Pastes the local variable.", EUserInterfaceActionType::Button, FInputChord());
-	UI_COMMAND(RemoveUnusedVariables, "Remove Unused Variables", "Removes variables not used in any graph.", EUserInterfaceActionType::Button, FInputChord());
+	UI_COMMAND(RemoveUnusedFunctions, "Remove Unused Functions", "Removes unsued functions.", EUserInterfaceActionType::Button, FInputChord());
+	UI_COMMAND(RemoveUnusedVariables, "Remove Unused Variables", "Removes unused variables.", EUserInterfaceActionType::Button, FInputChord());
 }
 
 
@@ -225,11 +226,11 @@ void SRigVMEditorGraphExplorer::Construct(const FArguments& InArgs, TWeakPtr<IRi
 	{
 		Editor->OnRefresh().AddSP(this, &SRigVMEditorGraphExplorer::Refresh);
 
-		// Acquire a registry for variable nodes
+		// Acquire a registry for function and variable nodes
 		if (const TScriptInterface<IRigVMAssetInterface> AssetInterface = Editor->GetRigVMAssetInterface())
 		{
+			EdGraphNodeFunctionRegistry = FRigVMEdGraphNodeRegistry::GetOrCreateRegistry(AssetInterface, URigVMFunctionReferenceNode::StaticClass());
 			EdGraphNodeVariableRegistry = FRigVMEdGraphNodeRegistry::GetOrCreateRegistry(AssetInterface, URigVMVariableNode::StaticClass());
-			EdGraphNodeVariableRegistry->OnPostRegistryUpdated.AddSP(this, &SRigVMEditorGraphExplorer::Refresh);
 		}
 	}
 }
@@ -371,6 +372,12 @@ void SRigVMEditorGraphExplorer::RegisterCommands()
 			FIsActionChecked(),
 			FIsActionButtonVisible::CreateSP(this, &SRigVMEditorGraphExplorer::CanFindReference) );
 
+		ToolKitCommandList->MapAction(FRigVMEditorGraphExplorerCommands::Get().RemoveUnusedFunctions,
+			FExecuteAction::CreateSP(this, &SRigVMEditorGraphExplorer::OnRemoveUnusedFunctions),
+			FCanExecuteAction(),
+			FIsActionChecked(),
+			FIsActionButtonVisible::CreateSP(this, &SRigVMEditorGraphExplorer::CanRemoveUnusedFunctions));
+
 		ToolKitCommandList->MapAction(FRigVMEditorGraphExplorerCommands::Get().RemoveUnusedVariables,
 			FExecuteAction::CreateSP(this, &SRigVMEditorGraphExplorer::OnRemoveUnusedVariables),
 			FCanExecuteAction(),
@@ -1212,13 +1219,14 @@ TArray<TSharedPtr<IPinTypeSelectorFilter>> SRigVMEditorGraphExplorer::GetCustomP
 TSharedPtr<SWidget> SRigVMEditorGraphExplorer::OnContextMenuOpening()
 {
 	TSharedPtr<IRigVMEditor> Editor = RigVMEditor.Pin();
-	if (!Editor.IsValid())
+	if (!Editor.IsValid() || 
+		!TreeView.IsValid())
 	{
 		return TSharedPtr<SWidget>();
 	}
 
 	FMenuBuilder MenuBuilder(/*Close after selection*/true, Editor->GetToolkitCommands());
-
+	
 	TArray<FRigVMExplorerElementKey> Selection = TreeView->GetSelectedKeys();
 
 	if (Selection.Num() == 1)
@@ -1227,6 +1235,21 @@ TSharedPtr<SWidget> SRigVMEditorGraphExplorer::OnContextMenuOpening()
 		{
 			case ERigVMExplorerElementType::Section:
 			{
+				const ERigVMGraphExplorerSectionType SectionType = TreeView->GetSectionType(Selection[0]);
+				if (SectionType == ERigVMGraphExplorerSectionType::Functions)
+				{
+					MenuBuilder.BeginSection("BasicOperations");
+					MenuBuilder.AddMenuEntry(FRigVMEditorGraphExplorerCommands::Get().RemoveUnusedFunctions);
+					MenuBuilder.EndSection();
+				}
+				else if (SectionType == ERigVMGraphExplorerSectionType::Variables ||
+					SectionType == ERigVMGraphExplorerSectionType::LocalVariables)
+				{
+					MenuBuilder.BeginSection("BasicOperations");
+					MenuBuilder.AddMenuEntry(FRigVMEditorGraphExplorerCommands::Get().RemoveUnusedVariables);
+					MenuBuilder.EndSection();
+				}
+
 				break;
 			}
 			case ERigVMExplorerElementType::Graph:
@@ -1257,6 +1280,7 @@ TSharedPtr<SWidget> SRigVMEditorGraphExplorer::OnContextMenuOpening()
 				MenuBuilder.AddMenuEntry(FGenericCommands::Get().Copy);
 				MenuBuilder.AddMenuEntry(FGenericCommands::Get().Delete);
 				MenuBuilder.AddMenuEntry(FRigVMEditorGraphExplorerCommands::Get().AddFunctionVariant);
+				MenuBuilder.AddMenuEntry(FRigVMEditorGraphExplorerCommands::Get().RemoveUnusedFunctions);
 				MenuBuilder.EndSection();
 				break;
 			}
@@ -2084,64 +2108,36 @@ bool SRigVMEditorGraphExplorer::CanFindReference() const
 	return true;
 }
 
-void SRigVMEditorGraphExplorer::OnRemoveUnusedVariables()
+void SRigVMEditorGraphExplorer::OnRemoveUnusedFunctions()
 {
-	 const TScriptInterface<IRigVMAssetInterface> Asset = RigVMEditor.IsValid() ? RigVMEditor.Pin()->GetRigVMAssetInterface() : nullptr;
-	if (!TreeView.IsValid() ||
-		!Asset)
+	const TScriptInterface<IRigVMAssetInterface> Asset = RigVMEditor.IsValid() ? RigVMEditor.Pin()->GetRigVMAssetInterface() : nullptr;
+	if (Asset)
 	{
-		return;
-	}
-	
-	const TArray<TWeakObjectPtr<URigVMEdGraphNode>> WeakEdGraphNodes = EdGraphNodeVariableRegistry->GetWeakEdGrapNodes();
-
-	TArray<FRigVMExplorerElementKey> KeysToRemove;
-	Algo::TransformIf(TreeView->GetKeys(), KeysToRemove,
-		[&WeakEdGraphNodes](const FRigVMExplorerElementKey& Key)
-		{
-			if (Key.Type == ERigVMExplorerElementType::Variable ||
-				Key.Type == ERigVMExplorerElementType::LocalVariable)
-			{
-				const bool bContainedInGraph = WeakEdGraphNodes.ContainsByPredicate(
-					[&Key](const TWeakObjectPtr<URigVMEdGraphNode>& WeakEdGraphNode)
-					{
-						const URigVMVariableNode* VariableNode = WeakEdGraphNode.IsValid() ? Cast<URigVMVariableNode>(WeakEdGraphNode->GetModelNode()) : nullptr;
-						if (VariableNode)
-						{
-							return VariableNode->GetVariableName() == Key.Name;
-						}
-
-						return false;
-					});
-
-				return !bContainedInGraph;
-			}
+		UE::RigVMEditor::RemoveUnusedMembersFromAsset<URigVMFunctionReferenceNode>(
+			Asset,
+			LOCTEXT("RemoveUnusedFunctions", "Remove Unused Functions"), 
+			LOCTEXT("AllFunctionsReferencedNotification", "All functions are referenced"));
 
-			return false;
-		},
-		[](const FRigVMExplorerElementKey& Key)
-		{
-			return Key;
-		});
+		Refresh();
+	}
+}
 
-	if (KeysToRemove.IsEmpty())
-	{
-		FNotificationInfo Info(LOCTEXT("AllVariablesReferencedNotification", "All variables are referenced"));
-		Info.bFireAndForget = true;
-		Info.FadeOutDuration = 5.0f;
-		Info.ExpireDuration = 5.0f;
+bool SRigVMEditorGraphExplorer::CanRemoveUnusedFunctions() const
+{
+	return true;
+}
 
-		FSlateNotificationManager::Get().AddNotification(Info);
-	}
-	else if (Asset->GetObject())
+void SRigVMEditorGraphExplorer::OnRemoveUnusedVariables()
+{
+	const TScriptInterface<IRigVMAssetInterface> Asset = RigVMEditor.IsValid() ? RigVMEditor.Pin()->GetRigVMAssetInterface() : nullptr;
+	if (Asset)
 	{
-		const FScopedTransaction RemoveUnusedVariablesTransaction(LOCTEXT("RemoveUnusedVariables", "Remove Unused Variables"), !GIsTransacting);
+		UE::RigVMEditor::RemoveUnusedMembersFromAsset<URigVMVariableNode>(
+			Asset,
+			LOCTEXT("RemoveUnusedVariables", "Remove Unused Variables"),
+			LOCTEXT("AllVariablesReferencedNotification", "All variables are referenced"));
 
-		Asset->GetObject()->Modify();
-		for (const FRigVMExplorerElementKey& KeyToRemove : KeysToRemove)
-		{
-			Asset->RemoveMemberVariable(*KeyToRemove.Name);
-		}
+		Refresh();
 	}
 }
 
diff --git a/Engine/Plugins/Runtime/RigVM/Source/RigVMEditor/Private/Widgets/SRigVMEditorGraphExplorerTreeView.cpp b/Engine/Plugins/Runtime/RigVM/Source/RigVMEditor/Private/Widgets/SRigVMEditorGraphExplorerTreeView.cpp
--- a/Engine/Plugins/Runtime/RigVM/Source/RigVMEditor/Private/Widgets/SRigVMEditorGraphExplorerTreeView.cpp
+++ b/Engine/Plugins/Runtime/RigVM/Source/RigVMEditor/Private/Widgets/SRigVMEditorGraphExplorerTreeView.cpp
@@ -25,11 +25,19 @@ void SRigVMEditorGraphExplorerItem::Construct(const FArguments& InArgs,
                                               TSharedRef<FRigVMEditorGraphExplorerTreeElement> InElement,
                                               TSharedPtr<SRigVMEditorGraphExplorerTreeView> InTreeView)
 {
+	using namespace UE::RigVMEditor;
+
 	WeakExplorerElement = InElement;
+	WeakTreeView = InTreeView;
 	Delegates = InTreeView->GetRigTreeDelegates();
 	OnAddClickedOnSection = InArgs._OnAddClickedOnSection;
 
-	FSlateFontInfo NameFont = FCoreStyle::GetDefaultFontStyle("Regular", 10);
+	if (const TSharedPtr<UE::RigVMEditor::FRigVMEdGraphNodeRegistry> EdGraphNodeRegistry = GetEdGraphNodeRegistry())
+	{
+		UpdateIsUsedInGraph();
+
+		EdGraphNodeRegistry->OnPostRegistryUpdated.AddSP(this, &SRigVMEditorGraphExplorerItem::UpdateIsUsedInGraph);
+	}
 
 	if (InElement->Key.Type == ERigVMExplorerElementType::Section)
 	{
@@ -109,11 +117,6 @@ void SRigVMEditorGraphExplorerItem::Construct(const FArguments& InArgs,
 	else if (InElement->Key.Type == ERigVMExplorerElementType::Variable || InElement->Key.Type == ERigVMExplorerElementType::LocalVariable)
 	{
 		TArray<TSharedPtr<IPinTypeSelectorFilter>> CustomPinTypeFilters = Delegates.GetCustomPinFilters();
-		
-		// Lower text color for variables that are not used in the graph
-		const FLinearColor TextColor = InTreeView->IsVariableElementUsedInGraph(InElement) ?
-			FLinearColor::White :
-			FLinearColor::White.CopyWithNewOpacity(0.7f);
 
 		STableRow<TSharedPtr<FRigVMEditorGraphExplorerTreeElement>>::Construct(
 		   STableRow<TSharedPtr<FRigVMEditorGraphExplorerTreeElement>>::FArguments()
@@ -127,7 +130,7 @@ void SRigVMEditorGraphExplorerItem::Construct(const FArguments& InArgs,
 				   .VAlign(VAlign_Center)
 				   .Padding(3.0f, 0.0f)
 				   [
-					   CreateTextSlotWidget(InElement->Key, InTreeView->FilterText, TextColor)
+					   CreateTextSlotWidget(InElement->Key, InTreeView->FilterText)
 				   ]
 
 				+SHorizontalBox::Slot()
@@ -260,18 +263,27 @@ TSharedRef<SWidget> SRigVMEditorGraphExplorerItem::CreateIconWidget(const FRigVM
 
 TSharedRef<SWidget> SRigVMEditorGraphExplorerItem::CreateTextSlotWidget(
 	const FRigVMExplorerElementKey& Key, 
-	const FText& InHighlightText, 
-	const FSlateColor& TextColor)
+	const FText& InHighlightText)
 {
-	TSharedPtr<SInlineEditableTextBlock> EditableTextElement = SNew(SInlineEditableTextBlock)
+	InlineRenameWidget = SNew(SInlineEditableTextBlock)
 		.Text(this, &SRigVMEditorGraphExplorerItem::GetDisplayText)
 		.HighlightText(InHighlightText)
 		.ToolTipText(this, &SRigVMEditorGraphExplorerItem::GetItemTooltip)
 		.OnVerifyTextChanged(this, &SRigVMEditorGraphExplorerItem::OnVerifyNameChanged)
 		.OnTextCommitted(this, &SRigVMEditorGraphExplorerItem::OnNameCommitted)
-		.ColorAndOpacity(TextColor);
-
-	InlineRenameWidget = EditableTextElement.ToSharedRef();
+		.ColorAndOpacity_Lambda([this]()
+			{
+				if (OptionalIsUsedInGraph.IsSet())
+				{
+					return OptionalIsUsedInGraph.GetValue() ?
+						FLinearColor::White :
+						FLinearColor::White.CopyWithNewOpacity(.35f);
+				}
+				else
+				{
+					return FLinearColor::White;
+				}
+			});
 
 	return InlineRenameWidget.ToSharedRef();
 }
@@ -462,6 +474,75 @@ FReply SRigVMEditorGraphExplorerItem::OnAddButtonClickedOnSection(TSharedRef<FRi
 	return FReply::Unhandled();
 }
 
+void SRigVMEditorGraphExplorerItem::UpdateIsUsedInGraph()
+{
+	using namespace UE::RigVMEditor;
+
+	TRACE_CPUPROFILER_EVENT_SCOPE(SRigVMEditorGraphExplorerTreeView::UpdateIsUsedInGraph);
+
+	if (!WeakExplorerElement.IsValid())
+	{
+		return;
+	}
+
+	if (const TSharedPtr<UE::RigVMEditor::FRigVMEdGraphNodeRegistry> EdGraphNodeRegistry = GetEdGraphNodeRegistry())
+	{
+		OptionalIsUsedInGraph.Reset();
+
+		const FRigVMExplorerElementKey& Key = WeakExplorerElement.Pin()->Key;
+
+		OptionalIsUsedInGraph = EdGraphNodeRegistry->GetConnectedEdGrapNodes().ContainsByPredicate(
+				[&Key](const TWeakObjectPtr<URigVMEdGraphNode>& WeakEdGraphNode)
+				{
+					const URigVMVariableNode* VariableNode = WeakEdGraphNode.IsValid() ? 
+						Cast<URigVMVariableNode>(WeakEdGraphNode->GetModelNode()) : 
+						nullptr;
+
+					const URigVMFunctionReferenceNode* ConnectedFunctionReferenceNode = 
+						WeakEdGraphNode.IsValid() ? Cast<URigVMFunctionReferenceNode>(WeakEdGraphNode->GetModelNode()) : 
+						nullptr;
+
+					if (VariableNode)
+					{
+						return VariableNode->GetVariableName() == Key.Name;
+					}
+					else if (ConnectedFunctionReferenceNode)
+					{
+						return ConnectedFunctionReferenceNode->GetFunctionIdentifier().GetFunctionName() == Key.Name;
+					}
+
+					return false;
+				});
+	}
+}
+
+TSharedPtr<UE::RigVMEditor::FRigVMEdGraphNodeRegistry> SRigVMEditorGraphExplorerItem::GetEdGraphNodeRegistry() const
+{
+	const TSharedPtr<SRigVMEditorGraphExplorer> GraphExplorer = WeakTreeView.IsValid() ? WeakTreeView.Pin()->GetGraphExplorer() : nullptr;
+	if (!WeakExplorerElement.IsValid() ||
+		!GraphExplorer.IsValid())
+	{
+		return nullptr;
+	}
+
+	const FRigVMExplorerElementKey& Key = WeakExplorerElement.Pin()->Key;
+
+	if (Key.Type == ERigVMExplorerElementType::Variable ||
+		Key.Type == ERigVMExplorerElementType::LocalVariable)
+	{
+		return GraphExplorer->GetEdGraphNodeVariableRegistry();
+	}
+	else if (Key.Type == ERigVMExplorerElementType::Function)
+	{
+		return GraphExplorer->GetEdGraphNodeFunctionRegistry();
+	}
+	else
+	{
+		return nullptr;
+	}
+}
+
+
 TSharedRef<ITableRow> FRigVMEditorGraphExplorerTreeElement::MakeTreeRowWidget(const TSharedRef<STableViewBase>& InOwnerTable,
 	TSharedRef<FRigVMEditorGraphExplorerTreeElement> InElement, TSharedPtr<SRigVMEditorGraphExplorerTreeView> InTreeView)
 {
@@ -977,32 +1058,34 @@ FReply SRigVMEditorGraphExplorerTreeView::OnAddButtonClickedOnSection(const FRig
 	return FReply::Unhandled();
 }
 
-bool SRigVMEditorGraphExplorerTreeView::IsVariableElementUsedInGraph(const TSharedRef<FRigVMEditorGraphExplorerTreeElement>& Element) const
+ERigVMGraphExplorerSectionType SRigVMEditorGraphExplorerTreeView::GetSectionType(const FRigVMExplorerElementKey& Key) const
 {
-	TRACE_CPUPROFILER_EVENT_SCOPE(SRigVMEditorGraphExplorerTreeView::IsVariableElementUsedInGraph);
-
-	const TSharedPtr<FRigVMEdGraphNodeRegistry> EdGraphNodeVariableRegistry = WeakGraphExplorer.IsValid() ? 
-		WeakGraphExplorer.Pin()->GetEdGraphNodeVariableRegistry() : 
-		nullptr;
-
-	if (!EdGraphNodeVariableRegistry.IsValid())
+	if (Key.Type == ERigVMExplorerElementType::Section)
 	{
-		return false;
+		if (Key.Name == GraphsSectionName.ToString())
+		{
+			return ERigVMGraphExplorerSectionType::Graphs;
+		}
+		else if (Key.Name == FunctionsSectionName.ToString())
+		{
+			return ERigVMGraphExplorerSectionType::Functions;
+		}
+		else if (Key.Name == VariablesSectionName.ToString())
+		{
+			return ERigVMGraphExplorerSectionType::Variables;
+		}
+		else if (Key.Name == LocalVariablesSectionName.ToString())
+		{
+			return ERigVMGraphExplorerSectionType::LocalVariables;
+		}
 	}
 
-	return
-		EdGraphNodeVariableRegistry->GetWeakEdGrapNodes()
-		.ContainsByPredicate(
-			[&Element](const TWeakObjectPtr<URigVMEdGraphNode>& WeakEdGraphNode)
-			{
-				const URigVMVariableNode* VariableNode = WeakEdGraphNode.IsValid() ?
-					Cast<URigVMVariableNode>(WeakEdGraphNode->GetModelNode()) :
-					nullptr;
+	return ERigVMGraphExplorerSectionType::None;
+}
 
-				return VariableNode ?
-					VariableNode->GetVariableName() == Element->Key.Name :
-					false;
-			});
+TSharedPtr<SRigVMEditorGraphExplorer> SRigVMEditorGraphExplorerTreeView::GetGraphExplorer() const
+{
+	return WeakGraphExplorer.IsValid() ? WeakGraphExplorer.Pin() : nullptr;
 }
 
-#undef LOCTEXT_NAMESPACE
\ No newline at end of file
+#undef LOCTEXT_NAMESPACE
diff --git a/Engine/Plugins/Runtime/RigVM/Source/RigVMEditor/Public/Widgets/SRigVMEditorGraphExplorer.h b/Engine/Plugins/Runtime/RigVM/Source/RigVMEditor/Public/Widgets/SRigVMEditorGraphExplorer.h
--- a/Engine/Plugins/Runtime/RigVM/Source/RigVMEditor/Public/Widgets/SRigVMEditorGraphExplorer.h
+++ b/Engine/Plugins/Runtime/RigVM/Source/RigVMEditor/Public/Widgets/SRigVMEditorGraphExplorer.h
@@ -80,6 +80,7 @@ public:
 	TSharedPtr<FUICommandInfo> PasteFunction;
 	TSharedPtr<FUICommandInfo> PasteVariable;
 	TSharedPtr<FUICommandInfo> PasteLocalVariable;
+	TSharedPtr<FUICommandInfo> RemoveUnusedFunctions;
 	TSharedPtr<FUICommandInfo> RemoveUnusedVariables;
 
 	void RegisterCommands() override;
@@ -110,10 +111,12 @@ public:
 
 	void ClearSelection();
 
+	/** Returns a registry for Ed Graph Nodes that hold RigVMFunctionReferenceNodes */
+	const TSharedPtr<FRigVMEdGraphNodeRegistry>& GetEdGraphNodeFunctionRegistry() const { return EdGraphNodeFunctionRegistry; }
+
 	/** Returns a registry for Ed Graph Nodes that hold RigVMVariableNodes */
 	const TSharedPtr<FRigVMEdGraphNodeRegistry>& GetEdGraphNodeVariableRegistry() const { return EdGraphNodeVariableRegistry; }
 
-
 private:
 	bool bNeedsRefresh;
 
@@ -213,6 +216,8 @@ private:
 
 	void OnFindReference(bool bSearchAllBlueprints);
 	bool CanFindReference() const;
+	void OnRemoveUnusedFunctions();
+	bool CanRemoveUnusedFunctions() const;
 	void OnRemoveUnusedVariables();
 	bool CanRemoveUnusedVariables() const;
 
@@ -224,6 +229,9 @@ private:
 
 	TSharedPtr<SSearchBox> FilterBox;
 
+	/** A registry for Ed Graph Nodes that hold URigVMFunctionReferenceNode */
+	TSharedPtr<FRigVMEdGraphNodeRegistry> EdGraphNodeFunctionRegistry;
+
 	/** A registry for Ed Graph Nodes that hold RigVMVariableNodes */
 	TSharedPtr<FRigVMEdGraphNodeRegistry> EdGraphNodeVariableRegistry;
 
diff --git a/Engine/Plugins/Runtime/RigVM/Source/RigVMEditor/Public/Widgets/SRigVMEditorGraphExplorerTreeView.h b/Engine/Plugins/Runtime/RigVM/Source/RigVMEditor/Public/Widgets/SRigVMEditorGraphExplorerTreeView.h
--- a/Engine/Plugins/Runtime/RigVM/Source/RigVMEditor/Public/Widgets/SRigVMEditorGraphExplorerTreeView.h
+++ b/Engine/Plugins/Runtime/RigVM/Source/RigVMEditor/Public/Widgets/SRigVMEditorGraphExplorerTreeView.h
@@ -33,6 +33,16 @@ namespace ERigVMExplorerElementType
 	};
 }
 
+/** The type of section in the graph explorer */
+enum class ERigVMGraphExplorerSectionType : uint8
+{
+	None,
+	Graphs,
+	Functions,
+	Variables, 
+	LocalVariables
+};
+
 struct FRigVMExplorerElementKey
 {
 
@@ -514,6 +524,8 @@ class SRigVMEditorGraphExplorerItem  : public STableRow<TSharedPtr<FRigVMEditorG
 {
 	DECLARE_DELEGATE_RetVal_OneParam(FReply, FOnAddClickedOnSection, const FRigVMExplorerElementKey&);
 	
+	using FRigVMEdGraphNodeRegistry = UE::RigVMEditor::FRigVMEdGraphNodeRegistry;
+
 public:
 	SLATE_BEGIN_ARGS(SRigVMEditorGraphExplorerItem)
 	{}
@@ -526,19 +538,28 @@ public:
 		TSharedPtr<SRigVMEditorGraphExplorerTreeView> InTreeView);
 
 	TSharedRef<SWidget> CreateIconWidget(const FRigVMExplorerElementKey& Key);
-	TSharedRef<SWidget> CreateTextSlotWidget(const FRigVMExplorerElementKey& Key, const FText& InHighlightText, const FSlateColor& TextColor = FLinearColor::White);
+	TSharedRef<SWidget> CreateTextSlotWidget(const FRigVMExplorerElementKey& Key, const FText& InHighlightText);
 	FText GetDisplayText() const;
 	FText GetItemTooltip() const;
 	bool OnVerifyNameChanged(const FText& InText, FText& OutErrorMessage);
 	void OnNameCommitted(const FText& InText, ETextCommit::Type InCommitType) const;
 	FReply OnAddButtonClickedOnSection(TSharedRef<FRigVMEditorGraphExplorerTreeElement> InElement);
 	
-	private:
-	
+private:
+	/** Updates OptionalIsUsedInGraph */
+	void UpdateIsUsedInGraph();
+
+	/** Returns a registry for ed graph nodes this item uses, or nullptr if no registry exists */
+	TSharedPtr<FRigVMEdGraphNodeRegistry> GetEdGraphNodeRegistry() const;
+
 	TWeakPtr<FRigVMEditorGraphExplorerTreeElement> WeakExplorerElement;
+	TWeakPtr<SRigVMEditorGraphExplorerTreeView> WeakTreeView;
 	TSharedPtr<SInlineEditableTextBlock> InlineRenameWidget;
 	FRigVMEditorGraphExplorerTreeDelegates Delegates;
 
+	/** If set and true, there are occurences in the graph. Only set for supported types (functions and variables as of 5.7) */
+	TOptional<bool> OptionalIsUsedInGraph;
+
 	FOnAddClickedOnSection OnAddClickedOnSection;
 };
 
@@ -618,6 +639,12 @@ public:
 	void OnItemDoubleClicked(TSharedPtr<FRigVMEditorGraphExplorerTreeElement> InElement);
 	FReply OnAddButtonClickedOnSection(const FRigVMExplorerElementKey& InSectionKey);
 
+	/** Returns the type of section the key references, or ERigVMGraphExplorerSectionType::None if the key doesn't reference a section.  */
+	ERigVMGraphExplorerSectionType GetSectionType(const FRigVMExplorerElementKey& Key) const;
+
+	/** Returns the outer graph explorer or nullptr if there is no valid outer graph explorer */
+	TSharedPtr<SRigVMEditorGraphExplorer> GetGraphExplorer() const;
+
 private:
 
 	/** Returns true if the eleemnt refers to a variable used in the graph. The element's key type has to correspond to a variable (ensured) */
