Control Rig: Fix minor issues with Remove Unused Members

* fix sorting functionality for the dialog list view
* Avoid creating transactions in RemoveUnusedMembers if a transaction is already ongoing
* Properly initialize FRigVMUnusedMemberCategory::bSafeToRemove

diff --git a/Engine/Plugins/Runtime/RigVM/Source/RigVMEditor/Private/Editor/RemoveUnusedMembers/RigVMEditorRemoveUnusedMembers.h b/Engine/Plugins/Runtime/RigVM/Source/RigVMEditor/Private/Editor/RemoveUnusedMembers/RigVMEditorRemoveUnusedMembers.h
--- a/Engine/Plugins/Runtime/RigVM/Source/RigVMEditor/Private/Editor/RemoveUnusedMembers/RigVMEditorRemoveUnusedMembers.h
+++ b/Engine/Plugins/Runtime/RigVM/Source/RigVMEditor/Private/Editor/RemoveUnusedMembers/RigVMEditorRemoveUnusedMembers.h
@@ -110,7 +110,7 @@ namespace UE::RigVMEditor
 		}
 		else
 		{
-			const bool bTransact = !TransactionText.IsEmpty();
+			const bool bTransact = !GIsTransacting && !TransactionText.IsEmpty();
 			if (bTransact)
 			{
 				Asset->GetObject()->Modify();
diff --git a/Engine/Plugins/Runtime/RigVM/Source/RigVMEditor/Private/Editor/RemoveUnusedMembers/RigVMEditorRemoveUnusedMembersCategory.h b/Engine/Plugins/Runtime/RigVM/Source/RigVMEditor/Private/Editor/RemoveUnusedMembers/RigVMEditorRemoveUnusedMembersCategory.h
--- a/Engine/Plugins/Runtime/RigVM/Source/RigVMEditor/Private/Editor/RemoveUnusedMembers/RigVMEditorRemoveUnusedMembersCategory.h
+++ b/Engine/Plugins/Runtime/RigVM/Source/RigVMEditor/Private/Editor/RemoveUnusedMembers/RigVMEditorRemoveUnusedMembersCategory.h
@@ -29,7 +29,7 @@ namespace UE::RigVMEditor
 		FText CategoryNameText;
 
 		/** If true, the members in this category can be safely removed from the Rig */
-		bool bSafeToRemove;
+		bool bSafeToRemove = false;
 
 		bool operator==(const FRigVMUnusedMemberCategory& Other) const
 		{
diff --git a/Engine/Plugins/Runtime/RigVM/Source/RigVMEditor/Private/Widgets/RemoveUnusedMembers/SRigVMEditorRemoveUnusedMembersList.cpp b/Engine/Plugins/Runtime/RigVM/Source/RigVMEditor/Private/Widgets/RemoveUnusedMembers/SRigVMEditorRemoveUnusedMembersList.cpp
--- a/Engine/Plugins/Runtime/RigVM/Source/RigVMEditor/Private/Widgets/RemoveUnusedMembers/SRigVMEditorRemoveUnusedMembersList.cpp
+++ b/Engine/Plugins/Runtime/RigVM/Source/RigVMEditor/Private/Widgets/RemoveUnusedMembers/SRigVMEditorRemoveUnusedMembersList.cpp
@@ -28,7 +28,7 @@ namespace UE::RigVMEditor
 
 		ChildSlot
 		[
-			SNew(SListView<TSharedPtr<FRigVMEditorRemoveUnusedMembersListItem>>)
+			SAssignNew(ListView, SListView<TSharedPtr<FRigVMEditorRemoveUnusedMembersListItem>>)
 			.SelectionMode(ESelectionMode::None)
 			.ListItemsSource(&Items)
 			.HeaderRow(GenerateHeaderRow())
@@ -48,6 +48,7 @@ namespace UE::RigVMEditor
 				.ColumnId(FRigVMEditorRemoveUnusedMembersColumnID::LabelRow)
 				.MinSize(40.f)
 				.FillWidth(0.9f)
+				.InitialSortMode(EColumnSortMode::Ascending)
 				.SortMode(this, &SRigVMEditorRemoveUnusedMembersList::GetColumnSortMode, FRigVMEditorRemoveUnusedMembersColumnID::LabelRow)
 				.OnSort(this, &SRigVMEditorRemoveUnusedMembersList::SortByColumnID)
 				.HeaderContent()
@@ -55,6 +56,9 @@ namespace UE::RigVMEditor
 					GenerateHeaderRowLabel()
 				]
 			);
+
+			ColumnSortMode = EColumnSortMode::Ascending;
+			SortedByColumnID = FRigVMEditorRemoveUnusedMembersColumnID::LabelRow;
 		}
 
 		// Select column
@@ -227,7 +231,10 @@ namespace UE::RigVMEditor
 		const FName& ColumnID, 
 		const EColumnSortMode::Type InSortMode)
 	{
-		const bool bAscending = InSortMode == EColumnSortMode::Ascending ? true : false;
+		ColumnSortMode = InSortMode;
+		SortedByColumnID = ColumnID;
+
+		const bool bAscending = ColumnSortMode == EColumnSortMode::Ascending ? true : false;
 
 		if (ColumnID == FRigVMEditorRemoveUnusedMembersColumnID::LabelRow)
 		{
@@ -236,8 +243,20 @@ namespace UE::RigVMEditor
 					const TSharedPtr<FRigVMEditorRemoveUnusedMembersListItem>& ItemA,
 					const TSharedPtr<FRigVMEditorRemoveUnusedMembersListItem>& ItemB)
 				{
-					return ItemA->GetMemberName().ToString() > ItemB->GetMemberName().ToString();
+					if (bAscending)
+					{
+						return ItemA->GetMemberName().ToString() > ItemB->GetMemberName().ToString();
+					}
+					else
+					{
+						return ItemA->GetMemberName().ToString() < ItemB->GetMemberName().ToString();
+					}
 				});
+
+			if (ListView.IsValid())
+			{
+				ListView->RequestListRefresh();
+			}
 		}
 	}
 }
diff --git a/Engine/Plugins/Runtime/RigVM/Source/RigVMEditor/Private/Widgets/RemoveUnusedMembers/SRigVMEditorRemoveUnusedMembersList.h b/Engine/Plugins/Runtime/RigVM/Source/RigVMEditor/Private/Widgets/RemoveUnusedMembers/SRigVMEditorRemoveUnusedMembersList.h
--- a/Engine/Plugins/Runtime/RigVM/Source/RigVMEditor/Private/Widgets/RemoveUnusedMembers/SRigVMEditorRemoveUnusedMembersList.h
+++ b/Engine/Plugins/Runtime/RigVM/Source/RigVMEditor/Private/Widgets/RemoveUnusedMembers/SRigVMEditorRemoveUnusedMembersList.h
@@ -10,6 +10,7 @@
 class ITableRow;
 class SHeaderRow;
 class STableViewBase;
+template <typename ItemType> class SListView;
 
 namespace UE::RigVMEditor
 {
@@ -62,6 +63,9 @@ namespace UE::RigVMEditor
 		/** The category to display in the header */
 		FRigVMUnusedMemberCategory Category;
 
+		/** The actual list view */
+		TSharedPtr<SListView<TSharedPtr<FRigVMEditorRemoveUnusedMembersListItem>>> ListView;
+
 		/** The member names in this list */
 		TArray<TSharedPtr<FRigVMEditorRemoveUnusedMembersListItem>> Items;
 
diff --git a/Engine/Plugins/Runtime/RigVM/Source/RigVMEditor/Private/Widgets/SRigVMEditorGraphExplorerTreeView.cpp b/Engine/Plugins/Runtime/RigVM/Source/RigVMEditor/Private/Widgets/SRigVMEditorGraphExplorerTreeView.cpp
--- a/Engine/Plugins/Runtime/RigVM/Source/RigVMEditor/Private/Widgets/SRigVMEditorGraphExplorerTreeView.cpp
+++ b/Engine/Plugins/Runtime/RigVM/Source/RigVMEditor/Private/Widgets/SRigVMEditorGraphExplorerTreeView.cpp
@@ -498,7 +498,7 @@ void SRigVMEditorGraphExplorerItem::UpdateIsUsedInGraph()
 						Cast<URigVMVariableNode>(WeakEdGraphNode->GetModelNode()) : 
 						nullptr;
 
-					const URigVMFunctionReferenceNode* ConnectedFunctionReferenceNode = 
+					const URigVMFunctionReferenceNode* FunctionReferenceNode = 
 						WeakEdGraphNode.IsValid() ? Cast<URigVMFunctionReferenceNode>(WeakEdGraphNode->GetModelNode()) : 
 						nullptr;
 
@@ -506,9 +506,9 @@ void SRigVMEditorGraphExplorerItem::UpdateIsUsedInGraph()
 					{
 						return VariableNode->GetVariableName() == Key.Name;
 					}
-					else if (ConnectedFunctionReferenceNode)
+					else if (FunctionReferenceNode)
 					{
-						return ConnectedFunctionReferenceNode->GetFunctionIdentifier().GetFunctionName() == Key.Name;
+						return FunctionReferenceNode->GetFunctionIdentifier().GetFunctionName() == Key.Name;
 					}
 
 					return false;
