RigVM: Fix crash during undo of add aggregate pin

#rnx

diff --git a/Engine/Plugins/Runtime/RigVM/Source/RigVMDeveloper/Private/RigVMAsset.cpp b/Engine/Plugins/Runtime/RigVM/Source/RigVMDeveloper/Private/RigVMAsset.cpp
--- a/Engine/Plugins/Runtime/RigVM/Source/RigVMDeveloper/Private/RigVMAsset.cpp
+++ b/Engine/Plugins/Runtime/RigVM/Source/RigVMDeveloper/Private/RigVMAsset.cpp
@@ -3333,6 +3333,10 @@ void IRigVMAssetInterface::HandleModifiedEvent(ERigVMGraphNotifType InNotifType,
 			case ERigVMGraphNotifType::NodeRemoved:
 			{
 				const bool bAdded = InNotifType == ERigVMGraphNotifType::NodeAdded;
+				if (URigVMHost::IsGarbageOrDestroyed(InSubject))
+				{
+					break;
+				}
 
 				if (URigVMCollapseNode* CollapseNode = Cast<URigVMCollapseNode>(InSubject))
 				{
@@ -3475,6 +3479,11 @@ void IRigVMAssetInterface::HandleModifiedEvent(ERigVMGraphNotifType InNotifType,
 			case ERigVMGraphNotifType::PinRemoved:
 			case ERigVMGraphNotifType::PinRenamed:
 			{
+				if (URigVMHost::IsGarbageOrDestroyed(InSubject))
+				{
+					break;
+				}
+
 				// exposed pin changes like this (as well as type change etc)
 				// require to mark the blueprint as structurally modified,
 				// so that the instance actions work out.
diff --git a/Engine/Plugins/Runtime/RigVM/Source/RigVMDeveloper/Private/RigVMModel/RigVMController.cpp b/Engine/Plugins/Runtime/RigVM/Source/RigVMDeveloper/Private/RigVMModel/RigVMController.cpp
--- a/Engine/Plugins/Runtime/RigVM/Source/RigVMDeveloper/Private/RigVMModel/RigVMController.cpp
+++ b/Engine/Plugins/Runtime/RigVM/Source/RigVMDeveloper/Private/RigVMModel/RigVMController.cpp
@@ -3936,8 +3936,6 @@ URigVMNode* URigVMController::EjectNodeFromPin(URigVMPin* InPin, bool bSetupUndo
 	Notify(ERigVMGraphNotifType::NodeAdded, NodeToEject);
 	if (InputPin)
 	{
-		(void)UnbindPinFromVariable(InputPin, false);
-		
 		TArray<URigVMLink*> SourceLinks = InputPin->GetSourceLinks(true);
 		if (SourceLinks.Num() > 0)
 		{
@@ -4484,6 +4482,7 @@ TArray<FName> URigVMController::ImportNodesFromText(const FString& InText, bool
 			for(URigVMPin* Pin : CreatedNode->Pins)
 			{
 				ControllerForGraph->EnsurePinValidity(Pin, true);
+				Pin->InjectionInfos.Remove(nullptr);
 			}
 
 			ControllerForGraph->Notify(ERigVMGraphNotifType::NodeAdded, CreatedNode);
@@ -11231,7 +11230,6 @@ FString URigVMController::AddAggregatePin(URigVMNode* InNode, const FString& InP
 	}
 
 	FRigVMControllerCompileBracketScope CompileScope(this);
-	const TMap<FString, FPinState> PinStates = GetPinStates(InNode);
 
 	FRigVMReplaceNodesAction Action;
 	if(bSetupUndoRedo)
@@ -11240,6 +11238,8 @@ FString URigVMController::AddAggregatePin(URigVMNode* InNode, const FString& InP
 		Action.SetTitle(TEXT("Add aggregate pin"));
 	}
 
+	const TMap<FString, FPinState> PinStates = GetPinStates(InNode);
+
 	if (AggregateNode == nullptr)
 	{
 		bool bAggregateInputs = false;
@@ -11280,6 +11280,17 @@ FString URigVMController::AddAggregatePin(URigVMNode* InNode, const FString& InP
 			return FString();
 		}
 
+		TMap<FName, URigVMNode*> EjectedNodes;
+		for (URigVMPin* Pin : InNode->GetPins())
+		{
+			if (Pin->HasInjectedNodes())
+			{
+				if (URigVMNode* EjectedNode = EjectNodeFromPin(Pin, false, false))
+				{
+					EjectedNodes.Add(Pin->GetFName(), EjectedNode);
+				}
+			}
+		}
 		(void)EjectAllInjectedNodes(InNode, false);
 		
 		const FString AggregateArg1 = Arg1->GetName();
@@ -11389,6 +11400,31 @@ FString URigVMController::AddAggregatePin(URigVMNode* InNode, const FString& InP
 		FRestoreLinkedPathSettings Settings;
 		Settings.NodeNameMap = {{PreviousNodeName.ToString(), AggregateNode->GetName()}};
 		RestoreLinkedPaths(LinkedPaths, Settings, false);
+
+		for (const TPair<FName, URigVMNode*>& Pair : EjectedNodes)
+		{
+			if (URigVMPin* Pin = AggregateNode->FindRootPinByName(Pair.Key))
+			{
+				URigVMNode* EjectedNode = Pair.Value;
+				const bool bIsInput = Pin->GetNode()->GetLinkedSourceNodes().Contains(EjectedNode);
+				FName OtherPinName = NAME_None;
+
+				TArray<URigVMPin*> LinkedPins = bIsInput ? Pin->GetLinkedSourcePins() : Pin->GetLinkedTargetPins();
+				for (URigVMPin* LinkedPin : LinkedPins)
+				{
+					if (LinkedPin->GetNode() == EjectedNode)
+					{
+						OtherPinName = LinkedPin->GetFName();
+						break;
+					}
+				}
+				
+				if (!InjectNodeIntoPin(Pin->GetPinPath(), bIsInput, bIsInput ? FName() : OtherPinName, bIsInput ? OtherPinName : FName(), bSetupUndoRedo))
+				{
+					(void)RemoveNode(Pair.Value, false, false);
+				}
+			}
+		}
 	}
 	
 	if (!AggregateNode)
