Avoid a RigVM Graph node crash at postload, when trying to find source and target pins for a link that has been lost

diff --git a/Engine/Plugins/Runtime/RigVM/Source/RigVMDeveloper/Private/RigVMModel/RigVMNode.cpp b/Engine/Plugins/Runtime/RigVM/Source/RigVMDeveloper/Private/RigVMModel/RigVMNode.cpp
--- a/Engine/Plugins/Runtime/RigVM/Source/RigVMDeveloper/Private/RigVMModel/RigVMNode.cpp
+++ b/Engine/Plugins/Runtime/RigVM/Source/RigVMDeveloper/Private/RigVMModel/RigVMNode.cpp
@@ -531,6 +531,10 @@ URigVMPin* URigVMNode::FindPin(const FString& InPinPath) const
 
 	for (URigVMPin* Pin : GetPins())
 	{
+		if (Pin == nullptr)
+		{
+			continue;
+		}
 		if (Pin->NameEquals(Left, true))
 		{
 			if (Right.IsEmpty())
@@ -545,6 +549,10 @@ URigVMPin* URigVMNode::FindPin(const FString& InPinPath) const
 	{
 		for (URigVMPin* Pin : OrphanedPins)
 		{
+			if (Pin == nullptr)
+			{
+				continue;
+			}
 			if (Pin->GetName() == InPinPath)
 			{
 				return Pin;
