RigVM: make sure to establish write lock before computing function ref node structure hash

#rnx

diff --git a/Engine/Plugins/Runtime/RigVM/Source/RigVMDeveloper/Private/RigVMModel/Nodes/RigVMFunctionReferenceNode.cpp b/Engine/Plugins/Runtime/RigVM/Source/RigVMDeveloper/Private/RigVMModel/Nodes/RigVMFunctionReferenceNode.cpp
--- a/Engine/Plugins/Runtime/RigVM/Source/RigVMDeveloper/Private/RigVMModel/Nodes/RigVMFunctionReferenceNode.cpp
+++ b/Engine/Plugins/Runtime/RigVM/Source/RigVMDeveloper/Private/RigVMModel/Nodes/RigVMFunctionReferenceNode.cpp
@@ -207,10 +207,9 @@ FName URigVMFunctionReferenceNode::GetOuterVariableName(const FName& InInnerVari
 
 uint32 URigVMFunctionReferenceNode::GetStructureHash() const
 {
-	const FRigVMRegistry& Registry = FRigVMRegistry::Get();
-	
 	uint32 Hash = Super::GetStructureHash();
 
+	const FRigVMRegistryWriteLock Registry;
 	Hash = HashCombine(Hash, GetTypeHash(ReferencedFunctionHeader.Name.ToString()));
 	Hash = HashCombine(Hash, GetTypeHash(ReferencedFunctionHeader.NodeTitle));
 	Hash = HashCombine(Hash, GetTypeHash(ReferencedFunctionHeader.LibraryPointer.GetLibraryNodePath()));
@@ -222,8 +221,8 @@ uint32 URigVMFunctionReferenceNode::GetStructureHash() const
 	{
 		Hash = HashCombine(Hash, GetTypeHash(Argument.Name.ToString()));
 		Hash = HashCombine(Hash, GetTypeHash((int32)Argument.Direction));
-		const TRigVMTypeIndex TypeIndex = Registry.GetTypeIndexFromCPPType(Argument.CPPType.ToString());
-		Hash = HashCombine(Hash, Registry.GetHashForType(TypeIndex));
+		const TRigVMTypeIndex TypeIndex = Registry->GetTypeIndexFromCPPType_NoLock(Argument.CPPType.ToString());
+		Hash = HashCombine(Hash, Registry->GetHashForType_NoLock(TypeIndex));
 
 		for(const TPair<FString, FText>& Pair : Argument.PathToTooltip)
 		{
@@ -235,8 +234,8 @@ uint32 URigVMFunctionReferenceNode::GetStructureHash() const
 	for(const FRigVMExternalVariable& ExternalVariable : ReferencedFunctionHeader.ExternalVariables)
 	{
 		Hash = HashCombine(Hash, GetTypeHash(ExternalVariable.Name.ToString()));
-		const TRigVMTypeIndex TypeIndex = Registry.GetTypeIndexFromCPPType(ExternalVariable.TypeName.ToString());
-		Hash = HashCombine(Hash, Registry.GetHashForType(TypeIndex));
+		const TRigVMTypeIndex TypeIndex = Registry->GetTypeIndexFromCPPType_NoLock(ExternalVariable.TypeName.ToString());
+		Hash = HashCombine(Hash, Registry->GetHashForType_NoLock(TypeIndex));
 	}
 
 	return Hash;
