Control Rig: Nodes and hierarchy elements now offers options to find nodes by name locally and globally

diff --git a/Engine/Plugins/Runtime/RigVM/Source/RigVMEditor/Private/Editor/RigVMEditor.cpp b/Engine/Plugins/Runtime/RigVM/Source/RigVMEditor/Private/Editor/RigVMEditor.cpp
--- a/Engine/Plugins/Runtime/RigVM/Source/RigVMEditor/Private/Editor/RigVMEditor.cpp
+++ b/Engine/Plugins/Runtime/RigVM/Source/RigVMEditor/Private/Editor/RigVMEditor.cpp
@@ -131,6 +131,8 @@ void FRigVMEditorBase::UnbindEditor()
 		}
 	}
 
+	FRigVMEditorModule::GetOnRequestFindNodeReferences().RemoveAll(this);
+
 	if (bRequestedReopen)
 	{
 		// Sometimes FPersonaToolkit::SetPreviewMesh will request an asset editor close and reopen. If
@@ -490,6 +492,8 @@ void FRigVMEditorBase::InitRigVMEditor(const EToolkitMode::Type Mode, const TSha
 	}
 
 	PropertyChangedHandle = FCoreUObjectDelegates::OnObjectPropertyChanged.AddSP(StaticCastSharedRef<FRigVMEditorBase>(SharedRef()), &FRigVMEditorBase::OnPropertyChanged);
+
+	FRigVMEditorModule::GetOnRequestFindNodeReferences().AddSP(StaticCastSharedRef<FRigVMEditorBase>(SharedRef()), &FRigVMEditorBase::OnRequestFindNodeReferences);
 }
 
 void FRigVMEditorBase::HandleAssetRequestedOpen(UObject* InObject)
@@ -4456,4 +4460,34 @@ bool FRigVMEditorBase::IsEditingSingleBlueprint() const
 	return GetRigVMAssetInterface() != nullptr;
 }
 
-#undef LOCTEXT_NAMESPACE
\ No newline at end of file
+void FRigVMEditorBase::OnRequestFindNodeReferences(const UE::RigVMEditor::FRigVMEditorFindNodeReferencesParams Params)
+{
+	if (!Params.WeakRigVMAsset.IsValid() ||
+		Params.WeakRigVMAsset.Get() != GetRigVMAssetInterface().GetInterface())
+	{
+		return;
+	}
+
+	UClass* GeneratedClass = Params.WeakRigVMAsset.Get()->GetRigVMGeneratedClass();
+	const URigVMNode* Node = Params.WeakEdGraphNode.IsValid() ? Params.WeakEdGraphNode->GetModelNode() : nullptr;
+	if (!GeneratedClass ||
+		!Node)
+	{
+		return;
+	}
+
+	const FString SearchTerm = [Node, &Params]()
+		{
+			if (const URigVMVariableNode* VariableNode = Cast<const URigVMVariableNode>(Node))
+			{
+				return VariableNode->GetVariableName().ToString();
+			}
+			
+			return Node->GetNodeTitle();
+		}();
+
+	const bool bFindWithinBlueprint = !Params.bSearchInAllBlueprints;
+	SummonSearchUI(bFindWithinBlueprint, SearchTerm);
+}
+
+#undef LOCTEXT_NAMESPACE
diff --git a/Engine/Plugins/Runtime/RigVM/Source/RigVMEditor/Private/Editor/RigVMEditorCommands.cpp b/Engine/Plugins/Runtime/RigVM/Source/RigVMEditor/Private/Editor/RigVMEditorCommands.cpp
--- a/Engine/Plugins/Runtime/RigVM/Source/RigVMEditor/Private/Editor/RigVMEditorCommands.cpp
+++ b/Engine/Plugins/Runtime/RigVM/Source/RigVMEditor/Private/Editor/RigVMEditorCommands.cpp
@@ -2,6 +2,8 @@
 
 #include "RigVMEditorCommands.h"
 
+#include "Framework/MultiBox/MultiBoxBuilder.h"
+
 #define LOCTEXT_NAMESPACE "RigVMEditorCommands"
 
 TSharedPtr<const FRigVMEditorCommands> FRigVMEditorCommands::GetSharedPtr()
@@ -33,6 +35,20 @@ void FRigVMEditorCommands::RegisterCommands()
 	UI_COMMAND(ToggleProfiling, "Toggle Profiling", "Enables or disables heat map profiling for the graph.", EUserInterfaceActionType::ToggleButton, FInputChord(EKeys::P));
 	UI_COMMAND(TogglePreviewHere, "Toggle Preview Here", "Run the graph up to this node", EUserInterfaceActionType::Button, FInputChord(EKeys::F9));
 	UI_COMMAND(PreviewHereStepForward, "Step Forwards", "Run the graph up to the next node.", EUserInterfaceActionType::Button, FInputChord(EKeys::F10));
+
+	UI_COMMAND(FindReferences, "Find References", "Find references of this item", EUserInterfaceActionType::Button, FInputChord(EKeys::F, EModifierKey::Shift | EModifierKey::Alt));
+	UI_COMMAND(FindReferencesByNameLocal, "By Name", "Find references of this item by name", EUserInterfaceActionType::Button, FInputChord());
+	UI_COMMAND(FindReferencesByNameGlobal, "By Name (All)", "Find references of this item by name in all blueprints", EUserInterfaceActionType::Button, FInputChord());
+}
+
+void FRigVMEditorCommands::BuildFindReferencesMenu(FMenuBuilder& MenuBuilder)
+{
+	MenuBuilder.BeginSection("FindReferences");
+	{
+		MenuBuilder.AddMenuEntry(FRigVMEditorCommands::Get().FindReferencesByNameLocal);
+		MenuBuilder.AddMenuEntry(FRigVMEditorCommands::Get().FindReferencesByNameGlobal);
+	}
+	MenuBuilder.EndSection();
 }
 
 #undef LOCTEXT_NAMESPACE
diff --git a/Engine/Plugins/Runtime/RigVM/Source/RigVMEditor/Private/Editor/RigVMEditorCommands.h b/Engine/Plugins/Runtime/RigVM/Source/RigVMEditor/Private/Editor/RigVMEditorCommands.h
--- a/Engine/Plugins/Runtime/RigVM/Source/RigVMEditor/Private/Editor/RigVMEditorCommands.h
+++ b/Engine/Plugins/Runtime/RigVM/Source/RigVMEditor/Private/Editor/RigVMEditorCommands.h
@@ -6,6 +6,8 @@
 #include "Framework/Commands/Commands.h"
 #include "Editor/RigVMEditorStyle.h"
 
+class FMenuBuilder;
+
 class FRigVMEditorCommands : public TCommands<FRigVMEditorCommands>
 {
 public:
@@ -65,8 +67,20 @@ public:
 	/** Step to next node during preview */
 	TSharedPtr< FUICommandInfo > PreviewHereStepForward;
 
+	/** Find references for non-functions and variables */
+	TSharedPtr< FUICommandInfo > FindReferences;
+
+	/** Find local references for functions and variables by name */
+	TSharedPtr< FUICommandInfo > FindReferencesByNameLocal;
+
+	/** Find member references for functions and variables by name */
+	TSharedPtr< FUICommandInfo > FindReferencesByNameGlobal;
+
 	/**
 	 * Initialize commands
 	 */
 	virtual void RegisterCommands() override;
+
+	/** Build "Find References" submenu when a context allows for it */
+	static void BuildFindReferencesMenu(FMenuBuilder& MenuBuilder);
 };
diff --git a/Engine/Plugins/Runtime/RigVM/Source/RigVMEditor/Private/RigVMEditorModule.cpp b/Engine/Plugins/Runtime/RigVM/Source/RigVMEditor/Private/RigVMEditorModule.cpp
--- a/Engine/Plugins/Runtime/RigVM/Source/RigVMEditor/Private/RigVMEditorModule.cpp
+++ b/Engine/Plugins/Runtime/RigVM/Source/RigVMEditor/Private/RigVMEditorModule.cpp
@@ -57,6 +57,8 @@ IMPLEMENT_MODULE(FRigVMEditorModule, RigVMEditor)
 TAutoConsoleVariable<bool> CVarRigVMUseNewEditor(TEXT("RigVM.UseNewEditor"), 0, TEXT("Use editor which does not inherit from BlueprintEditor."));
 TAutoConsoleVariable<bool> CVarRigVMUseDualEditor(TEXT("RigVM.UseDualEditor"), 0, TEXT("Use both legacy and new editor for A/B comparison."));
 
+UE::RigVMEditor::FRigVMEditorOnRequestFindNodeReferencesDelegate FRigVMEditorModule::OnRequestFindNodeReferences;
+
 FRigVMEditorModule& FRigVMEditorModule::Get()
 {
 	return FModuleManager::LoadModuleChecked< FRigVMEditorModule >(TEXT("RigVMEditor"));
@@ -666,18 +668,125 @@ void FRigVMEditorModule::GetNodeContextMenuActions(TScriptInterface<IRigVMClient
 
 void FRigVMEditorModule::GetNodeActionsContextMenuActions(TScriptInterface<IRigVMClientHost> RigVMClientHost, const URigVMEdGraphNode* EdGraphNode, URigVMNode* ModelNode, UToolMenu* Menu) const
 {
-	if (!EdGraphNode || !ModelNode || !Menu)
+	const TScriptInterface<IRigVMAssetInterface> RigVMAsset = RigVMClientHost.GetObject();
+	if (!EdGraphNode || 
+		!ModelNode || 
+		!Menu || 
+		!RigVMAsset)
 	{
 		return;
 	}
 
-	// This section displays actions that are to be found under 'Node Actions' in other blueprint editors as well
-	FToolMenuSection& NodeActionsSection = Menu->AddSection("RigVMEditorContextMenuNodeActions", LOCTEXT("NodeActionsHeader", "Node Actions"));
+	const TWeakInterfacePtr<IRigVMAssetInterface> WeakRigVMAsset = RigVMAsset;
+	const TWeakObjectPtr<URigVMNode> WeakModelNode = ModelNode;
+	const TWeakObjectPtr<const URigVMEdGraphNode> WeakEdGraphNode = EdGraphNode;
+
+	const TSharedRef<FUICommandList> FindNodeFreferencesCommandList = MakeShared<FUICommandList>();
+
+	FindNodeFreferencesCommandList->MapAction(
+		FRigVMEditorCommands::Get().FindReferences,
+		FExecuteAction::CreateLambda(
+			[WeakRigVMAsset, WeakEdGraphNode]()
+			{
+				constexpr bool bSearchInAllBlueprints = false;
+
+				const UE::RigVMEditor::FRigVMEditorFindNodeReferencesParams FindNodeReferencesParams(
+					WeakRigVMAsset,
+					WeakEdGraphNode,
+					bSearchInAllBlueprints
+				);
+
+				OnRequestFindNodeReferences.Broadcast(FindNodeReferencesParams);
+			})
+	);
+
+	FindNodeFreferencesCommandList->MapAction(
+		FRigVMEditorCommands::Get().FindReferencesByNameLocal,
+		FExecuteAction::CreateLambda(
+			[WeakRigVMAsset, WeakEdGraphNode]()
+			{
+				constexpr bool bSearchInAllBlueprints = false;
+
+				const UE::RigVMEditor::FRigVMEditorFindNodeReferencesParams FindNodeReferencesParams(
+					WeakRigVMAsset,
+					WeakEdGraphNode,
+					bSearchInAllBlueprints
+				);
+
+				OnRequestFindNodeReferences.Broadcast(FindNodeReferencesParams);
+			})
+	);
 
-	NodeActionsSection.AddMenuEntry(
-		FGraphEditorCommands::Get().FindReferencesByNameLocal, 
-		LOCTEXT("FindReferencesByNameLabel", "Find References By Name")
+
+	FindNodeFreferencesCommandList->MapAction(
+		FRigVMEditorCommands::Get().FindReferencesByNameGlobal,
+		FExecuteAction::CreateLambda(
+			[WeakRigVMAsset, WeakEdGraphNode]()
+			{
+				constexpr bool bSearchInAllBlueprints = true;
+
+				const UE::RigVMEditor::FRigVMEditorFindNodeReferencesParams FindNodeReferencesParams(
+					WeakRigVMAsset,
+					WeakEdGraphNode,
+					bSearchInAllBlueprints
+				);
+
+				OnRequestFindNodeReferences.Broadcast(FindNodeReferencesParams);
+			})
 	);
+
+	// This section displays actions that are to be found under 'Node Actions' in other blueprint editors as well
+	FToolMenuSection& NodeActionsSection = Menu->AddSection(
+		"RigVMEditorContextMenuNodeActions",
+		LOCTEXT("NodeActionsHeader", "Node Actions"));
+
+	const bool bShowSubMenu = [EdGraphNode]()
+		{
+			if (const URigVMNode* Node = EdGraphNode->GetModelNode())
+			{
+				return
+					Node->IsA(URigVMVariableNode::StaticClass()) ||
+					Node->IsA(URigVMUnitNode::StaticClass()) ||
+					Node->IsA(URigVMFunctionReferenceNode::StaticClass());
+			}
+
+			return false;
+		}();
+
+	const FText FindReferencesLabel = LOCTEXT("FindReferencesLabel", "Find References");
+	const FText FindReferencesTooltip = LOCTEXT("FindReferencesTooltip", "Options for finding references to class members");
+
+	if (bShowSubMenu)
+	{
+		NodeActionsSection.AddSubMenu(
+			"FindReferenceSubMenu",
+			FindReferencesLabel,
+			FindReferencesTooltip,
+			FNewToolMenuDelegate::CreateLambda([FindNodeFreferencesCommandList](UToolMenu* InMenu)
+				{
+					if (InMenu)
+					{
+						FToolMenuSection& FindReferencesSection = InMenu->AddSection(
+							"FindReferencesSection",
+							FText::GetEmpty());
+
+						FindReferencesSection.AddMenuEntry(FRigVMEditorCommands::Get().FindReferencesByNameLocal)
+							.SetCommandList(FindNodeFreferencesCommandList);
+
+						FindReferencesSection.AddMenuEntry(FRigVMEditorCommands::Get().FindReferencesByNameGlobal)
+							.SetCommandList(FindNodeFreferencesCommandList);
+					}
+				})
+		);		
+	}
+	else
+	{
+		NodeActionsSection.AddMenuEntry(
+			FRigVMEditorCommands::Get().FindReferences,
+			FindReferencesLabel,
+			FindReferencesTooltip)
+			.SetCommandList(FindNodeFreferencesCommandList);
+	}
 }
 
 void FRigVMEditorModule::GetNodeWorkflowContextMenuActions(TScriptInterface<IRigVMClientHost> RigVMClientHost, const URigVMEdGraphNode* EdGraphNode, URigVMNode* ModelNode, UToolMenu* Menu) const
diff --git a/Engine/Plugins/Runtime/RigVM/Source/RigVMEditor/Private/Widgets/SRigVMEditorGraphExplorer.cpp b/Engine/Plugins/Runtime/RigVM/Source/RigVMEditor/Private/Widgets/SRigVMEditorGraphExplorer.cpp
--- a/Engine/Plugins/Runtime/RigVM/Source/RigVMEditor/Private/Widgets/SRigVMEditorGraphExplorer.cpp
+++ b/Engine/Plugins/Runtime/RigVM/Source/RigVMEditor/Private/Widgets/SRigVMEditorGraphExplorer.cpp
@@ -4,6 +4,7 @@
 
 #include "Editor/RemoveUnusedMembers/RigVMEditorRemoveUnusedMembers.h"
 #include "Editor/RigVMEdGraphNodeRegistry.h"
+#include "Editor/RigVMEditorCommands.h"
 #include "Editor/RigVMNewEditor.h"
 #include "EdGraph/RigVMEdGraphSchema.h"
 #include "Framework/Notifications/NotificationManager.h"
@@ -366,23 +367,48 @@ void SRigVMEditorGraphExplorer::RegisterCommands()
 			FExecuteAction::CreateSP(this, &SRigVMEditorGraphExplorer::OnDeleteEntry),
 			FCanExecuteAction::CreateSP(this, &SRigVMEditorGraphExplorer::CanDeleteEntry));
 
-		ToolKitCommandList->MapAction( FGraphEditorCommands::Get().FindReferences,
-			FExecuteAction::CreateSP(this, &SRigVMEditorGraphExplorer::OnFindReference, /*bSearchAllBlueprints=*/false),
-			FCanExecuteAction(),
-			FIsActionChecked(),
-			FIsActionButtonVisible::CreateSP(this, &SRigVMEditorGraphExplorer::CanFindReference) );
+		// Find references local
+		{
+			constexpr bool bSearchAllBlueprints = false;
+
+			ToolKitCommandList->MapAction(FRigVMEditorCommands::Get().FindReferences,
+				FExecuteAction::CreateSP(this, &SRigVMEditorGraphExplorer::OnFindReferences, bSearchAllBlueprints),
+				FCanExecuteAction(),
+				FIsActionChecked(),
+				FIsActionButtonVisible::CreateSP(this, &SRigVMEditorGraphExplorer::CanFindReferences));
+
+			ToolKitCommandList->MapAction(FRigVMEditorCommands::Get().FindReferencesByNameLocal,
+				FExecuteAction::CreateSP(this, &SRigVMEditorGraphExplorer::OnFindReferences, bSearchAllBlueprints),
+				FCanExecuteAction(),
+				FIsActionChecked(),
+				FIsActionButtonVisible::CreateSP(this, &SRigVMEditorGraphExplorer::CanFindReferences));
+		}
+
+		// Find references global
+		{
+			constexpr bool bSearchAllBlueprints = true;
 
-		ToolKitCommandList->MapAction(FRigVMEditorGraphExplorerCommands::Get().RemoveUnusedFunctions,
-			FExecuteAction::CreateSP(this, &SRigVMEditorGraphExplorer::OnRemoveUnusedFunctions),
-			FCanExecuteAction(),
-			FIsActionChecked(),
-			FIsActionButtonVisible::CreateSP(this, &SRigVMEditorGraphExplorer::CanRemoveUnusedFunctions));
+			ToolKitCommandList->MapAction(FRigVMEditorCommands::Get().FindReferencesByNameGlobal,
+				FExecuteAction::CreateSP(this, &SRigVMEditorGraphExplorer::OnFindReferences, bSearchAllBlueprints),
+				FCanExecuteAction(),
+				FIsActionChecked(),
+				FIsActionButtonVisible::CreateSP(this, &SRigVMEditorGraphExplorer::CanFindReferences));
+		}
 
-		ToolKitCommandList->MapAction(FRigVMEditorGraphExplorerCommands::Get().RemoveUnusedVariables,
-			FExecuteAction::CreateSP(this, &SRigVMEditorGraphExplorer::OnRemoveUnusedVariables),
-			FCanExecuteAction(),
-			FIsActionChecked(),
-			FIsActionButtonVisible::CreateSP(this, &SRigVMEditorGraphExplorer::CanRemoveUnusedVariables));
+		// Remove unused members
+		{
+			ToolKitCommandList->MapAction(FRigVMEditorGraphExplorerCommands::Get().RemoveUnusedFunctions,
+				FExecuteAction::CreateSP(this, &SRigVMEditorGraphExplorer::OnRemoveUnusedFunctions),
+				FCanExecuteAction(),
+				FIsActionChecked(),
+				FIsActionButtonVisible::CreateSP(this, &SRigVMEditorGraphExplorer::CanRemoveUnusedFunctions));
+
+			ToolKitCommandList->MapAction(FRigVMEditorGraphExplorerCommands::Get().RemoveUnusedVariables,
+				FExecuteAction::CreateSP(this, &SRigVMEditorGraphExplorer::OnRemoveUnusedVariables),
+				FCanExecuteAction(),
+				FIsActionChecked(),
+				FIsActionButtonVisible::CreateSP(this, &SRigVMEditorGraphExplorer::CanRemoveUnusedVariables));
+		}
 	}
 }
 
@@ -1231,6 +1257,9 @@ TSharedPtr<SWidget> SRigVMEditorGraphExplorer::OnContextMenuOpening()
 
 	if (Selection.Num() == 1)
 	{
+		const FText FindReferencesLabel = LOCTEXT("FindReferencesLabel", "Find References");
+		const FText FindReferencesTooltip = LOCTEXT("FindReferencesTooltip", "Options for finding references to class members");
+
 		switch (Selection[0].Type)
 		{
 			case ERigVMExplorerElementType::Section:
@@ -1258,7 +1287,7 @@ TSharedPtr<SWidget> SRigVMEditorGraphExplorer::OnContextMenuOpening()
 				MenuBuilder.AddMenuEntry(FRigVMEditorGraphExplorerCommands::Get().OpenGraph);
 				MenuBuilder.AddMenuEntry(FRigVMEditorGraphExplorerCommands::Get().OpenGraphInNewTab);
 				MenuBuilder.AddMenuEntry(FGenericCommands::Get().Rename);
-				MenuBuilder.AddMenuEntry(FGraphEditorCommands::Get().FindReferences);
+				MenuBuilder.AddMenuEntry(FRigVMEditorCommands::Get().FindReferences);
 				MenuBuilder.AddMenuEntry(FGenericCommands::Get().Cut);
 				MenuBuilder.AddMenuEntry(FGenericCommands::Get().Copy);
 				MenuBuilder.AddMenuEntry(FGenericCommands::Get().Delete);
@@ -1275,7 +1304,10 @@ TSharedPtr<SWidget> SRigVMEditorGraphExplorer::OnContextMenuOpening()
 				MenuBuilder.AddMenuEntry(FRigVMEditorGraphExplorerCommands::Get().OpenGraph);
 				MenuBuilder.AddMenuEntry(FRigVMEditorGraphExplorerCommands::Get().OpenGraphInNewTab);
 				MenuBuilder.AddMenuEntry(FGenericCommands::Get().Rename);
-				MenuBuilder.AddMenuEntry(FGraphEditorCommands::Get().FindReferences);
+				MenuBuilder.AddSubMenu(
+					FindReferencesLabel,
+					FindReferencesTooltip,
+					FNewMenuDelegate::CreateStatic(&FRigVMEditorCommands::BuildFindReferencesMenu));
 				MenuBuilder.AddMenuEntry(FGenericCommands::Get().Cut);
 				MenuBuilder.AddMenuEntry(FGenericCommands::Get().Copy);
 				MenuBuilder.AddMenuEntry(FGenericCommands::Get().Delete);
@@ -1289,7 +1321,10 @@ TSharedPtr<SWidget> SRigVMEditorGraphExplorer::OnContextMenuOpening()
 			{
 				MenuBuilder.BeginSection("BasicOperations");
 				MenuBuilder.AddMenuEntry(FGenericCommands::Get().Rename);
-				MenuBuilder.AddMenuEntry(FGraphEditorCommands::Get().FindReferences);
+				MenuBuilder.AddSubMenu(
+					FindReferencesLabel,
+					FindReferencesTooltip,
+					FNewMenuDelegate::CreateStatic(&FRigVMEditorCommands::BuildFindReferencesMenu));
 				MenuBuilder.AddMenuEntry(FGenericCommands::Get().Cut);
 				MenuBuilder.AddMenuEntry(FGenericCommands::Get().Copy);
 				MenuBuilder.AddMenuEntry(FGenericCommands::Get().Duplicate);
@@ -2056,56 +2091,69 @@ bool SRigVMEditorGraphExplorer::OnCanRenameVariable(const FRigVMExplorerElementK
 	return true;
 }
 
-void SRigVMEditorGraphExplorer::OnFindReference(bool bSearchAllBlueprints)
+void SRigVMEditorGraphExplorer::OnFindReferences(bool bSearchInAllBlueprints)
 {
-	FString SearchTerm;
+	if (!RigVMEditor.IsValid() ||
+		!TreeView.IsValid())
+	{
+		return;
+	}
+	const TSharedRef<IRigVMEditor> PinnedEditor = RigVMEditor.Pin().ToSharedRef();
+
+	const TArray<FRigVMExplorerElementKey> SelectedKeys = TreeView->GetSelectedKeys();
 
-	if (!RigVMEditor.IsValid())
+	// Only searching for a single keys is supported for the time being
+	if (!CanFindReferences() ||
+		SelectedKeys.Num() != 1)
 	{
 		return;
 	}
 
-	TSharedPtr<IRigVMEditor> Editor = RigVMEditor.Pin();
-	FRigVMAssetInterfacePtr Blueprint = Editor->GetRigVMAssetInterface();
-	const URigVMEdGraphSchema* Schema = GetDefault<URigVMEdGraphSchema>();
+	FString SearchTerm;
 
-	TArray<FRigVMExplorerElementKey> Selection = TreeView->GetSelectedKeys();
-	if (Selection.Num() != 1)
+	const TScriptInterface<IRigVMAssetInterface> Asset = PinnedEditor->GetRigVMAssetInterface();
+	const URigVMEdGraphSchema* Schema = GetDefault<URigVMEdGraphSchema>();
+	if (!Asset || !Schema)
 	{
 		return;
 	}
 
-	switch (Selection[0].Type)
+	switch (SelectedKeys[0].Type)
 	{
-		case ERigVMExplorerElementType::Graph:
-		{
-			if (UEdGraph* EdGraph = Editor->GetRigVMAssetInterface()->GetEdGraph(Selection[0].Name))
-			{
-				FGraphDisplayInfo DisplayInfo;
-				Schema->GetGraphDisplayInformation(*EdGraph, DisplayInfo);
-				SearchTerm = DisplayInfo.DisplayName.ToString();
-			}
-			break;
-		}
-		case ERigVMExplorerElementType::Function:
-		case ERigVMExplorerElementType::Variable:
-		case ERigVMExplorerElementType::LocalVariable:
+	case ERigVMExplorerElementType::Graph:
+	{
+		if (UEdGraph* EdGraph = Asset->GetEdGraph(SelectedKeys[0].Name))
 		{
-			SearchTerm = Selection[0].Name;
-			break;
+			FGraphDisplayInfo DisplayInfo;
+			Schema->GetGraphDisplayInformation(*EdGraph, DisplayInfo);
+			SearchTerm = DisplayInfo.DisplayName.ToString();
 		}
+		break;
 	}
-	
+	case ERigVMExplorerElementType::Function:
+	case ERigVMExplorerElementType::Variable:
+	case ERigVMExplorerElementType::LocalVariable:
+	{
+		SearchTerm = SelectedKeys[0].Name;
+		break;
+	}
+	default:
+		ensureMsgf(0, TEXT("Unsupported element type"));
+		break;
+	}
+
 	if (!SearchTerm.IsEmpty())
 	{
-		const bool bSetFindWithinBlueprint = !bSearchAllBlueprints;
-		Editor->SummonSearchUI(bSetFindWithinBlueprint, SearchTerm);
+		const bool bSetFindWithinBlueprint = !bSearchInAllBlueprints;
+		PinnedEditor->SummonSearchUI(bSetFindWithinBlueprint, SearchTerm);
 	}
 }
 
-bool SRigVMEditorGraphExplorer::CanFindReference() const
-{
-	return true;
+bool SRigVMEditorGraphExplorer::CanFindReferences() const
+{	
+	// Only searching for a single keys is supported for the time being
+	const TArray<FRigVMExplorerElementKey> SelectedKeys = TreeView.IsValid() ? TreeView->GetSelectedKeys() : TArray<FRigVMExplorerElementKey>();
+	return SelectedKeys.Num() == 1;
 }
 
 void SRigVMEditorGraphExplorer::OnRemoveUnusedFunctions()
diff --git a/Engine/Plugins/Runtime/RigVM/Source/RigVMEditor/Public/Editor/RigVMEditor.h b/Engine/Plugins/Runtime/RigVM/Source/RigVMEditor/Public/Editor/RigVMEditor.h
--- a/Engine/Plugins/Runtime/RigVM/Source/RigVMEditor/Public/Editor/RigVMEditor.h
+++ b/Engine/Plugins/Runtime/RigVM/Source/RigVMEditor/Public/Editor/RigVMEditor.h
@@ -614,6 +614,10 @@ protected:
 
 	friend class SRigVMExecutionStackView;
 	friend class SRigVMEditorGraphExplorer;
+
+private:
+	/** Called when the editor module requests finding nodes */
+	void OnRequestFindNodeReferences(const UE::RigVMEditor::FRigVMEditorFindNodeReferencesParams Params);
 };
 
 struct FRigVMEditorZoomLevelsContainer : public FZoomLevelsContainer
diff --git a/Engine/Plugins/Runtime/RigVM/Source/RigVMEditor/Public/RigVMEditorModule.h b/Engine/Plugins/Runtime/RigVM/Source/RigVMEditor/Public/RigVMEditorModule.h
--- a/Engine/Plugins/Runtime/RigVM/Source/RigVMEditor/Public/RigVMEditorModule.h
+++ b/Engine/Plugins/Runtime/RigVM/Source/RigVMEditor/Public/RigVMEditorModule.h
@@ -21,6 +21,7 @@
 
 class URigVMAssetInterface;
 class IRigVMAssetInterface;
+class IRigVMEditorModule;
 class URigVMUserWorkflowOptions;
 class URigVMPin;
 class URigVMNode;
@@ -46,6 +47,33 @@ namespace ERefreshRigVMEditorReason
 	};
 }
 
+namespace UE::RigVMEditor
+{
+	/** Params when finding references for a node */
+	struct FRigVMEditorFindNodeReferencesParams
+	{
+		FRigVMEditorFindNodeReferencesParams(
+			const TWeakInterfacePtr<IRigVMAssetInterface> InRigVMAsset,
+			const TWeakObjectPtr<const URigVMEdGraphNode> InEdGraphNode,
+			const bool bInSearchInAllBlueprints)
+			: WeakRigVMAsset(InRigVMAsset)
+			, WeakEdGraphNode(InEdGraphNode)
+			, bSearchInAllBlueprints(bInSearchInAllBlueprints)
+		{}
+
+		/** The asset that contains the node */
+		const TWeakInterfacePtr<IRigVMAssetInterface> WeakRigVMAsset;
+
+		/** The ed graph node */
+		const TWeakObjectPtr<const URigVMEdGraphNode> WeakEdGraphNode;
+
+		/** If true, searches in all blueprints */
+		const bool bSearchInAllBlueprints;
+	};
+
+	DECLARE_EVENT_OneParam(IRigVMEditorModule, FRigVMEditorOnRequestFindNodeReferencesDelegate, FRigVMEditorFindNodeReferencesParams /** FindNodeReferencesParams */);
+}
+
 // shallow interface declaration for use within RigVMDeveloper 
 class IRigVMEditorModule : public IModuleInterface, public IHasMenuExtensibility, public IHasToolBarExtensibility, public FStructureEditorUtils::INotifyOnStructChanged
 {
@@ -110,6 +138,9 @@ public:
 
 	UE_API virtual FConnectionDrawingPolicy* CreateConnectionDrawingPolicy(int32 InBackLayerID, int32 InFrontLayerID, float InZoomFactor, const FSlateRect& InClippingRect, FSlateWindowElementList& InDrawElements, UEdGraph* InGraphObj) const override;
 
+	/** Returns an event broadcast when finding references for a node is requested */
+	static UE::RigVMEditor::FRigVMEditorOnRequestFindNodeReferencesDelegate& GetOnRequestFindNodeReferences() { return OnRequestFindNodeReferences; }
+
 protected:
 
 	UE_API bool IsRigVMEditorModuleBase() const;
@@ -166,6 +197,9 @@ private:
 
 	UE_API bool ShowWorkflowOptionsDialog(URigVMUserWorkflowOptions* InOptions) const;
 
+	/** Event broadcast when finding references for a node is requested */
+	static UE::RigVMEditor::FRigVMEditorOnRequestFindNodeReferencesDelegate OnRequestFindNodeReferences;
+
 	friend class URigVMEdGraphNode;
 };
 
diff --git a/Engine/Plugins/Runtime/RigVM/Source/RigVMEditor/Public/Widgets/SRigVMEditorGraphExplorer.h b/Engine/Plugins/Runtime/RigVM/Source/RigVMEditor/Public/Widgets/SRigVMEditorGraphExplorer.h
--- a/Engine/Plugins/Runtime/RigVM/Source/RigVMEditor/Public/Widgets/SRigVMEditorGraphExplorer.h
+++ b/Engine/Plugins/Runtime/RigVM/Source/RigVMEditor/Public/Widgets/SRigVMEditorGraphExplorer.h
@@ -214,8 +214,9 @@ private:
 	bool OnRenameVariable(const FRigVMExplorerElementKey& InOldKey, const FString& InNewName);
 	bool OnCanRenameVariable(const FRigVMExplorerElementKey& InOldKey, const FString& InNewName, FText& OutErrorMessage) const;
 
-	void OnFindReference(bool bSearchAllBlueprints);
-	bool CanFindReference() const;
+	void OnFindReferences(const bool bSearchAllBlueprints);
+	bool CanFindReferences() const;
+
 	void OnRemoveUnusedFunctions();
 	bool CanRemoveUnusedFunctions() const;
 	void OnRemoveUnusedVariables();
