Use `UE_AUTORTFM_DECLARE_THREAD_LOCAL_VAR` for `thread_local` variable used in a transaction.

#rnx

diff --git a/Engine/Plugins/Runtime/RigVM/Source/RigVM/Private/RigVMFunctions/RigVMDispatch_Array.cpp b/Engine/Plugins/Runtime/RigVM/Source/RigVM/Private/RigVMFunctions/RigVMDispatch_Array.cpp
--- a/Engine/Plugins/Runtime/RigVM/Source/RigVM/Private/RigVMFunctions/RigVMDispatch_Array.cpp
+++ b/Engine/Plugins/Runtime/RigVM/Source/RigVM/Private/RigVMFunctions/RigVMDispatch_Array.cpp
@@ -5,6 +5,7 @@
 #include "RigVMCore/RigVMStruct.h"
 #include "RigVMStringUtils.h"
 #include "Math/GuardedInt.h"
+#include "AutoRTFM.h"
 
 #include UE_INLINE_GENERATED_CPP_BY_NAME(RigVMDispatch_Array)
 #define LOCTEXT_NAMESPACE "RigVMDispatch_Array"
@@ -1514,7 +1515,8 @@ const TArray<FRigVMTemplateArgumentInfo>& FRigVMDispatch_ArrayIterator::GetArgum
 const TArray<FRigVMExecuteArgument>& FRigVMDispatch_ArrayIterator::GetExecuteArguments_Impl(
 	const FRigVMDispatchContext& InContext) const
 {
-	thread_local TArray<FRigVMExecuteArgument> ExecuteArguments;
+	UE_AUTORTFM_DECLARE_THREAD_LOCAL_VAR(TArray<FRigVMExecuteArgument>, ExecuteArguments);
+
 	if(ExecuteArguments.IsEmpty())
 	{
 		ExecuteArguments = FRigVMDispatch_ArrayBaseMutable::GetExecuteArguments_Impl(InContext);
diff --git a/Engine/Plugins/Runtime/RigVM/Source/RigVM/Private/RigVMFunctions/RigVMDispatch_Switch.cpp b/Engine/Plugins/Runtime/RigVM/Source/RigVM/Private/RigVMFunctions/RigVMDispatch_Switch.cpp
--- a/Engine/Plugins/Runtime/RigVM/Source/RigVM/Private/RigVMFunctions/RigVMDispatch_Switch.cpp
+++ b/Engine/Plugins/Runtime/RigVM/Source/RigVM/Private/RigVMFunctions/RigVMDispatch_Switch.cpp
@@ -4,6 +4,7 @@
 #include "RigVMStringUtils.h"
 #include "RigVMCore/RigVMRegistry.h"
 #include "RigVMCore/RigVMStruct.h"
+#include "AutoRTFM.h"
 
 #include UE_INLINE_GENERATED_CPP_BY_NAME(RigVMDispatch_Switch)
 
@@ -146,7 +147,9 @@ void FRigVMDispatch_SwitchInt32::Execute(FRigVMExtendedExecuteContext& InContext
 
 FName FRigVMDispatch_SwitchInt32::GetCaseName(int32 InIndex)
 {
-	thread_local TMap<int32, FName> NamesFromInt;
+	using MapType = TMap<int32, FName>;
+	UE_AUTORTFM_DECLARE_THREAD_LOCAL_VAR(MapType, NamesFromInt);
+
 	if (FName* Found = NamesFromInt.Find(InIndex))
 	{
 		return *Found;
@@ -159,7 +162,9 @@ FName FRigVMDispatch_SwitchInt32::GetCaseName(int32 InIndex)
 
 FName FRigVMDispatch_SwitchInt32::GetCaseDisplayName(int32 InIndex)
 {
-	thread_local TMap<int32, FName> NamesFromInt;
+	using MapType = TMap<int32, FName>;
+	UE_AUTORTFM_DECLARE_THREAD_LOCAL_VAR(MapType, NamesFromInt);
+
     if (FName* Found = NamesFromInt.Find(InIndex))
     {
     	return *Found;
@@ -170,4 +175,3 @@ FName FRigVMDispatch_SwitchInt32::GetCaseDisplayName(int32 InIndex)
     NamesFromInt.Add(InIndex, Result);
     return Result;
 }
-
