Add Damp and Critical Spring Damp RigVM nodes, exposing it to UAF and Control Rig

[RN]

#FailedValidations Preflight

diff --git a/Engine/Plugins/Runtime/RigVM/Source/RigVM/Private/RigVMFunctions/Math/RigVMFunction_SpringMath.cpp b/Engine/Plugins/Runtime/RigVM/Source/RigVM/Private/RigVMFunctions/Math/RigVMFunction_SpringMath.cpp
new file mode 100644
--- /dev/null
+++ b/Engine/Plugins/Runtime/RigVM/Source/RigVM/Private/RigVMFunctions/Math/RigVMFunction_SpringMath.cpp
@@ -0,0 +1,45 @@
+// Copyright Epic Games, Inc. All Rights Reserved.
+
+#include "RigVMFunctions/Math/RigVMFunction_SpringMath.h"
+
+#include "Animation/SpringMath.h"
+#include "Math/UnrealMathUtility.h"
+
+FRigVMFunction_DampFloat_Execute()
+{
+	float DeltaTime = ExecuteContext.GetDeltaTime<float>();
+	Result = Value;
+	FMath::ExponentialSmoothingApprox(Result, Target, DeltaTime, SmoothingTime);
+}
+
+FRigVMFunction_DampVector_Execute()
+{
+	float DeltaTime = ExecuteContext.GetDeltaTime<float>();
+	Result = Value;
+	FMath::ExponentialSmoothingApprox(Result, Target, DeltaTime, SmoothingTime);
+}
+
+FRigVMFunction_DampQuaternion_Execute()
+{
+	float DeltaTime = ExecuteContext.GetDeltaTime<float>();
+	Result = Value;
+	SpringMath::ExponentialSmoothingApproxQuat(Result, Target, DeltaTime, SmoothingTime);
+}
+
+FRigVMFunction_SpringDampFloat_Execute()
+{
+	float DeltaTime = ExecuteContext.GetDeltaTime<float>();
+	SpringMath::CriticalSpringDamper(Value, ValueVelocity, Target, SmoothingTime, DeltaTime);
+}
+
+FRigVMFunction_SpringDampVector_Execute()
+{
+	float DeltaTime = ExecuteContext.GetDeltaTime<float>();
+	SpringMath::CriticalSpringDamper(Value, ValueVelocity, Target, SmoothingTime, DeltaTime);
+}
+
+FRigVMFunction_SpringDampQuat_Execute()
+{
+	float DeltaTime = ExecuteContext.GetDeltaTime<float>();
+	SpringMath::CriticalSpringDamperQuat(Value, ValueVelocity, Target, SmoothingTime, DeltaTime);
+}
diff --git a/Engine/Plugins/Runtime/RigVM/Source/RigVM/Public/RigVMFunctions/Math/RigVMFunction_SpringMath.h b/Engine/Plugins/Runtime/RigVM/Source/RigVM/Public/RigVMFunctions/Math/RigVMFunction_SpringMath.h
new file mode 100644
--- /dev/null
+++ b/Engine/Plugins/Runtime/RigVM/Source/RigVM/Public/RigVMFunctions/Math/RigVMFunction_SpringMath.h
@@ -0,0 +1,156 @@
+// Copyright Epic Games, Inc. All Rights Reserved.
+
+#pragma once
+
+#include "RigVMFunction_MathBase.h"
+#include "RigVMFunction_SpringMath.generated.h"
+
+/**
+ * Damps a float value using exponential decay damping
+ */
+UE_EXPERIMENTAL(5.7, "SpringMath API is experimental")
+USTRUCT(meta=(DisplayName="Damp (Float)", Category="Math|Damp|Experimental", TemplateName="Damp"))
+struct FRigVMFunction_DampFloat : public FRigVMFunction_MathBase
+{
+	GENERATED_BODY()
+
+	RIGVM_METHOD()
+	RIGVM_API virtual void Execute() override;
+
+	UPROPERTY(meta = (Input))
+	float Value = 0.0f;
+
+	UPROPERTY(meta = (Input))
+	float Target = 0.0f;
+
+	UPROPERTY(meta = (Input, ClampMin = "0", UIMin = "0", ForceUnits = "s"))
+	float SmoothingTime = 0.2f;
+
+	UPROPERTY(meta = (Output))
+	float Result = 0.0f;
+};
+
+/**
+ * Damps a vector value using exponential decay damping
+ */
+UE_EXPERIMENTAL(5.7, "SpringMath API is experimental")
+USTRUCT(meta=(DisplayName="Damp (Vector)", Category="Math|Damp|Experimental", TemplateName="Damp"))
+struct FRigVMFunction_DampVector : public FRigVMFunction_MathBase
+{
+	GENERATED_BODY()
+
+	RIGVM_METHOD()
+	RIGVM_API virtual void Execute() override;
+
+	UPROPERTY(meta = (Input))
+	FVector Value = FVector::ZeroVector;
+
+	UPROPERTY(meta = (Input))
+	FVector Target = FVector::ZeroVector;
+
+	UPROPERTY(meta = (Input, ClampMin = "0", UIMin = "0", ForceUnits = "s"))
+	float SmoothingTime = 0.2f;
+
+	UPROPERTY(meta = (Output))
+	FVector Result = FVector::ZeroVector;
+};
+
+/**
+ * Damps a quaternion value using exponential decay damping
+ */
+UE_EXPERIMENTAL(5.7, "SpringMath API is experimental") 
+USTRUCT(meta=(DisplayName="Damp (Quaternion)", Category="Math|Damp|Experimental", TemplateName="Damp"))
+struct FRigVMFunction_DampQuaternion : public FRigVMFunction_MathBase
+{
+	GENERATED_BODY()
+
+	RIGVM_METHOD()
+	RIGVM_API virtual void Execute() override;
+
+	UPROPERTY(meta = (Input))
+	FQuat Value = FQuat::Identity;
+
+	UPROPERTY(meta = (Input))
+	FQuat Target = FQuat::Identity;
+
+	UPROPERTY(meta = (Input, ClampMin = "0", UIMin = "0", ForceUnits = "s"))
+	float SmoothingTime = 0.2f;
+
+	UPROPERTY(meta = (Output))
+	FQuat Result = FQuat::Identity;
+};
+
+/**
+ * Damps a float using a spring damper
+ */
+UE_EXPERIMENTAL(5.7, "SpringMath API is experimental")
+USTRUCT(meta=(DisplayName="Critical Spring Damp (Float)", Category="Math|Damp|Experimental", TemplateName="Critical Spring Damp"))
+struct FRigVMFunction_SpringDampFloat : public FRigVMFunction_MathMutableBase
+{
+	GENERATED_BODY()
+
+	RIGVM_METHOD()
+	RIGVM_API virtual void Execute() override;
+
+	UPROPERTY(meta = (Input, Output))
+	float Value = 0.0f;
+
+	UPROPERTY(meta = (Input, Output))
+	float ValueVelocity = 0.0f;
+
+	UPROPERTY(meta = (Input))
+	float Target = 0.0f;
+
+	UPROPERTY(meta = (Input, ClampMin = "0", UIMin = "0", ForceUnits = "s"))
+	float SmoothingTime = 0.2f;
+};
+
+/**
+ * Damps a vector using a spring damper
+ */
+UE_EXPERIMENTAL(5.7, "SpringMath API is experimental")
+USTRUCT(meta=(DisplayName="Critical Spring Damp (Vector)", Category="Math|Damp|Experimental", TemplateName="Critical Spring Damp"))
+struct FRigVMFunction_SpringDampVector : public FRigVMFunction_MathMutableBase
+{
+	GENERATED_BODY()
+
+	RIGVM_METHOD()
+	RIGVM_API virtual void Execute() override;
+
+	UPROPERTY(meta = (Input, Output))
+	FVector Value = FVector::ZeroVector;
+
+	UPROPERTY(meta = (Input, Output))
+	FVector ValueVelocity = FVector::ZeroVector;
+
+	UPROPERTY(meta = (Input))
+	FVector Target = FVector::ZeroVector;
+
+	UPROPERTY(meta = (Input, ClampMin = "0", UIMin = "0", ForceUnits = "s"))
+	float SmoothingTime = 0.2f;
+};
+
+/**
+ * Damps a quaternion using a spring damper
+ */
+UE_EXPERIMENTAL(5.7, "SpringMath API is experimental")
+USTRUCT(meta=(DisplayName="Critical Spring Damp (Quat)", Category="Math|Damp|Experimental", TemplateName="Critical Spring Damp"))
+struct FRigVMFunction_SpringDampQuat : public FRigVMFunction_MathMutableBase
+{
+	GENERATED_BODY()
+
+	RIGVM_METHOD()
+	RIGVM_API virtual void Execute() override;
+
+	UPROPERTY(meta = (Input, Output))
+	FQuat Value = FQuat::Identity;
+
+	UPROPERTY(meta = (Input, Output))
+	FVector ValueVelocity = FVector::ZeroVector;
+
+	UPROPERTY(meta = (Input))
+	FQuat Target = FQuat::Identity;
+
+	UPROPERTY(meta = (Input, ClampMin = "0", UIMin = "0", ForceUnits = "s"))
+	float SmoothingTime = 0.2f;
+};
