Fix for not being able to move nodes within newly created comment nodes in Ctrl Rig and UAF

diff --git a/Engine/Plugins/Runtime/RigVM/Source/RigVMDeveloper/Private/EdGraph/RigVMEdGraph.cpp b/Engine/Plugins/Runtime/RigVM/Source/RigVMDeveloper/Private/EdGraph/RigVMEdGraph.cpp
--- a/Engine/Plugins/Runtime/RigVM/Source/RigVMDeveloper/Private/EdGraph/RigVMEdGraph.cpp
+++ b/Engine/Plugins/Runtime/RigVM/Source/RigVMDeveloper/Private/EdGraph/RigVMEdGraph.cpp
@@ -1096,6 +1096,7 @@ const URigVMEdGraph* URigVMEdGraph::GetRootGraph() const
 void URigVMEdGraph::AddNode(UEdGraphNode* NodeToAdd, bool bUserAction, bool bSelectNewNode)
 {
 	// Comments are added outside of the ControlRigEditor, so we add here the node to the model
+	bool bCommentNodeAdded = false;
 	if (const UEdGraphNode_Comment* CommentNode = Cast<const UEdGraphNode_Comment>(NodeToAdd))
 	{
 		if (URigVMController* Controller = GetController())
@@ -1122,11 +1123,20 @@ void URigVMEdGraph::AddNode(UEdGraphNode* NodeToAdd, bool bUserAction, bool bSel
 				FLinearColor NodeColor = CommentNode->CommentColor;
 				Controller->AddCommentNode(CommentNode->NodeComment, NodePos, NodeSize, NodeColor, CommentNode->GetName(), !GIsTransacting, true);
 				ModelNodePathToEdNode.Add(NodeToAdd->GetFName(), NodeToAdd);
+
+				bCommentNodeAdded = true;
 			}
 		}
 	}
 
 	Super::AddNode(NodeToAdd, bUserAction, bSelectNewNode);
+
+	// If a comment node was added the graph needs to be refreshed to re-sort its nodes
+	// Otherwise when selecting nodes they won't be depth sorted correctly and the selection can be wrong
+	if (bCommentNodeAdded)
+	{
+		NotifyGraphChanged();
+	}
 }
 
 void URigVMEdGraph::RemoveNode(UEdGraphNode* InNode)
