Control Rig: Fix editor reopening after closing
#rnx

diff --git a/Engine/Plugins/Runtime/RigVM/Source/RigVMEditor/Private/Editor/RigVMEditor.cpp b/Engine/Plugins/Runtime/RigVM/Source/RigVMEditor/Private/Editor/RigVMEditor.cpp
--- a/Engine/Plugins/Runtime/RigVM/Source/RigVMEditor/Private/Editor/RigVMEditor.cpp
+++ b/Engine/Plugins/Runtime/RigVM/Source/RigVMEditor/Private/Editor/RigVMEditor.cpp
@@ -501,6 +501,20 @@ void FRigVMEditorBase::HandleAssetRequestedOpen(UObject* InObject)
 	if (InObject == GetRigVMAssetInterface()->GetObject())
 	{
 		bRequestedReopen = true;
+
+		// This reopen request is useful only if when Persona SetPreviewMesh is set, and it automatically closes and reopens the editor
+		// If this open editor request happens without closing the editor (for some other reason), we actually want to forget the reopen
+		// Lets set a timer so that if a close editor does not happen, we forget the reopen request
+		TWeakPtr<IRigVMEditor> WeakThis = SharedRef().ToWeakPtr();
+		FTimerHandle RequestedReopenTimerHandle;
+		GEditor->GetTimerManager()->SetTimer(RequestedReopenTimerHandle, [WeakThis]()
+			{
+				if (TSharedPtr<IRigVMEditor> StrongThis = WeakThis.Pin())
+				{
+					FRigVMEditorBase* Editor = static_cast<FRigVMEditorBase*>(StrongThis.Get());
+					Editor->bRequestedReopen = false;
+				}
+			}, 2.0f, false, 1.0f);
 	}
 }
 
