RigVM: Avoid data race when accessing client outer
Also avoid data race when controllers - UE-229807

#rnx

#FailedValidations Preflight

diff --git a/Engine/Plugins/Runtime/RigVM/Source/RigVMDeveloper/Private/RigVMModel/RigVMClient.cpp b/Engine/Plugins/Runtime/RigVM/Source/RigVMDeveloper/Private/RigVMModel/RigVMClient.cpp
--- a/Engine/Plugins/Runtime/RigVM/Source/RigVMDeveloper/Private/RigVMModel/RigVMClient.cpp
+++ b/Engine/Plugins/Runtime/RigVM/Source/RigVMDeveloper/Private/RigVMModel/RigVMClient.cpp
@@ -34,6 +34,7 @@ void FRigVMClient::SetDefaultSchemaClass(TSubclassOf<URigVMSchema> InSchemaClass
 
 	DefaultSchemaClass = InSchemaClass;
 
+	UE::TScopeLock Lock(ControllersLock);
 	for(TPair<FSoftObjectPath, TObjectPtr<URigVMController>>& ObjectControllerPair : Controllers)
 	{
 		ObjectControllerPair.Value->SetSchemaClass(InSchemaClass);
@@ -59,20 +60,21 @@ void FRigVMClient::SetControllerClass(TSubclassOf<URigVMController> InController
 
 void FRigVMClient::SetOuterClientHost(UObject* InOuterClientHost, const FName& InOuterClientHostPropertyName)
 {
-	OuterClientHost = InOuterClientHost;
-	OuterClientPropertyName = InOuterClientHostPropertyName;
-
-	check(OuterClientHost->Implements<URigVMClientHost>());
-	check(GetOuterClientProperty() != nullptr);
+	{
+		UE::TScopeLock Lock(OuterClientHostLock);
+	
+		OuterClientHost = InOuterClientHost;
+		OuterClientPropertyName = InOuterClientHostPropertyName;
+	
+		check(OuterClientHost->Implements<URigVMClientHost>());
+		check(GetOuterClientProperty() != nullptr);
+	}
 
 	// create the null graph / default controller.
 	// we need this to react to notifs without a valid graph
-	// such as interaction brackets
+	// such as interaction brackets.
 	static const URigVMGraph* NullGraph = nullptr;
-	if(!Controllers.Contains(NullGraph))
-	{
-		CreateController(NullGraph);
-	}
+	GetOrCreateController(NullGraph);
 }
 
 void FRigVMClient::SetFromDeprecatedData(URigVMGraph* InDefaultGraph, URigVMFunctionLibrary* InFunctionLibrary)
@@ -110,20 +112,62 @@ void FRigVMClient::SetExternalModelHost(IRigVMClientExternalModelHost* InExterna
 	ExternalModelHost = InExternalModelHost;
 }
 
+FRigVMClient& FRigVMClient::operator=(const FRigVMClient& InOther)
+{
+	OnGetFocusedGraphDelegate = InOther.OnGetFocusedGraphDelegate;
+	DefaultSchemaClass = InOther.DefaultSchemaClass;
+	ControllerClass = InOther.ControllerClass;
+	Models = InOther.Models;
+	FunctionLibrary = InOther.FunctionLibrary;
+
+	{
+		UE::TScopeLock Lock(ControllersLock);
+		Controllers = InOther.Controllers;
+	}
+	ActionStack = InOther.ActionStack;
+
+	UndoRedoIndex = 0;
+	UndoStack.Reset();
+	RedoStack.Reset();
+
+	bSuspendNotifications = InOther.bSuspendNotifications;
+	bIgnoreModelNotifications = InOther.bIgnoreModelNotifications;
+	bDefaultModelCanBeRemoved = InOther.bDefaultModelCanBeRemoved;
+	bSuspendModelNotificationsForOthers = InOther.bSuspendModelNotificationsForOthers;
+
+	{
+		UE::TScopeLock Lock(OuterClientHostLock);
+		OuterClientHost = InOther.OuterClientHost;
+		OuterClientPropertyName = InOther.OuterClientPropertyName;
+	}
+	ExternalModelHost = InOther.ExternalModelHost;
+
+#if WITH_EDITOR
+	PostGraphModifiedDelegate = InOther.PostGraphModifiedDelegate;
+#endif // WITH_EDITOR
+	
+	return *this;
+}
+
 void FRigVMClient::Reset()
 {
 	for(URigVMGraph* Model : GetModels())
 	{
 		DestroyObject(Model);
 	}
-	for(auto Pair : Controllers)
+
 	{
-		DestroyObject(Pair.Value);
+		UE::TScopeLock Lock(ControllersLock);
+		for(auto Pair : Controllers)
+		{
+			DestroyObject(Pair.Value);
+		}
+		Controllers.Reset();
 	}
+	
 	DestroyObject(FunctionLibrary);
 
 	Models.Reset();
-	Controllers.Reset();
 	FunctionLibrary = nullptr;
 
 	ResetActionStack();
@@ -433,12 +477,14 @@ URigVMController* FRigVMClient::GetController(const URigVMGraph* InModel) const
 	if(InModel)
 	{
 		const FSoftObjectPath Key(InModel);
-		if(const TObjectPtr<URigVMController>* Controller = Controllers.Find(Key))
+		UE::TScopeLock Lock(ControllersLock);
+
+		if(const TObjectPtr<URigVMController>* ExistingController = Controllers.Find(Key))
 		{
-			if  (!URigVMHost::IsGarbageOrDestroyed(Controller->Get()))
+			if  (!URigVMHost::IsGarbageOrDestroyed(ExistingController->Get()))
 			{
-				checkf(Controller->Get()->GetGraph() == InModel, TEXT("Controller %s contains unexpected graph."), *Controller->GetPathName());
-				return Controller->Get();
+				checkf(ExistingController->Get()->GetGraph() == InModel, TEXT("Controller %s contains unexpected graph."), *ExistingController->GetPathName());
+				return ExistingController->Get();
 			}
 		}
 	}
@@ -531,10 +577,23 @@ URigVMController* FRigVMClient::GetControllerByName(const FString InGraphName) c
 
 bool FRigVMClient::RemoveController(const URigVMGraph* InModel)
 {
+	UE::TScopeLock Lock(ControllersLock);
+
+	if(InModel == nullptr)
+	{
+		InModel = GetDefaultModel();
+	}
+	
 	const FSoftObjectPath Key(InModel);
-	URigVMController* Controller = GetController(InModel);
+	
+	URigVMController* Controller = nullptr;
+	if(const TObjectPtr<URigVMController>* ExistingController = Controllers.Find(Key))
+	{
+		Controller = ExistingController->Get();
+	}
+
 	const bool bSuccess = Controllers.Remove(Key) > 0;
-	if(Controller)
+	if(Controller && !URigVMHost::IsGarbageOrDestroyed(Controller))
 	{
 		Controller->SetActionStack(nullptr);
 		Controller->Rename(nullptr, GetTransientPackage(), REN_DoNotDirty | REN_DontCreateRedirectors | REN_NonTransactional);
@@ -910,14 +969,19 @@ FName FRigVMClient::RenameModel(const FString& InNodePathOrName, const FName& In
 		}
 #endif
 
-		TObjectPtr<URigVMController>* Controller = Controllers.Find(Model);
 		const FString OldNodePath = Model->GetNodePath();
 		const FName SafeNewName = GetUniqueName(InNewName);
-		Model->Rename(*SafeNewName.ToString(), nullptr, REN_DontCreateRedirectors);
-		const FString NewNodePath = Model->GetNodePath();
-		if (Controller)
+		FString NewNodePath;
+
 		{
-			Controllers.Add(Model, *Controller);
+			UE::TScopeLock Lock(ControllersLock);
+			TObjectPtr<URigVMController>* Controller = Controllers.Find(Model);
+			Model->Rename(*SafeNewName.ToString(), nullptr, REN_DontCreateRedirectors);
+			NewNodePath = Model->GetNodePath();
+			if (Controller)
+			{
+				Controllers.Add(Model, *Controller);
+			}
 		}
 
 		if(bSetupUndoRedo)
@@ -1017,6 +1081,8 @@ void FRigVMClient::PostTransacted(const FTransactionObjectEvent& TransactionEven
 
 void FRigVMClient::OnCollapseNodeRenamed(const URigVMCollapseNode* InCollapseNode)
 {
+	UE::TScopeLock Lock(ControllersLock);
+
 	TMap<FSoftObjectPath, TObjectPtr<URigVMController>> OldControllers = Controllers;
 	for(const TPair<FSoftObjectPath, TObjectPtr<URigVMController>>& Pair : OldControllers)
 	{
@@ -1042,6 +1108,8 @@ void FRigVMClient::OnCollapseNodeRenamed(const URigVMCollapseNode* InCollapseNod
 
 void FRigVMClient::OnCollapseNodeRemoved(const URigVMCollapseNode* InCollapseNode)
 {
+	UE::TScopeLock Lock(ControllersLock);
+	
 	TMap<FSoftObjectPath, TObjectPtr<URigVMController>> OldControllers = Controllers;
 	for(const TPair<FSoftObjectPath, TObjectPtr<URigVMController>>& Pair : OldControllers)
 	{
@@ -1092,6 +1160,7 @@ URigVMPin* FRigVMClient::FindPin(const FString& InPinPath) const
 
 UObject* FRigVMClient::GetOuter() const
 {
+	UE::TScopeLock Lock(OuterClientHostLock);
 	UObject* Outer = OuterClientHost.Get();
 	check(Outer);
 	return Outer;
@@ -1115,25 +1184,46 @@ void FRigVMClient::NotifyOuterOfPropertyChange(EPropertyChangeType::Type ChangeT
 
 URigVMController* FRigVMClient::CreateController(const URigVMGraph* InModel)
 {
-	const FString ModelName = InModel ? InModel->GetName() : TEXT("__NullGraph");
-	const FName SafeControllerName = GetUniqueName(*FString::Printf(TEXT("%s_Controller"), *ModelName));
-	URigVMController* Controller = NewObject<URigVMController>(GetOuter(), ControllerClass, SafeControllerName);
-	Controllers.FindOrAdd(InModel, nullptr) = Controller;
-	if(InModel && InModel->GetSchemaClass() != nullptr)
-	{
-		Controller->SetSchemaClass(InModel->GetSchemaClass());
-	}
-	else
-	{
-		check(DefaultSchemaClass);
-		Controller->SetSchemaClass(DefaultSchemaClass);
-	}
+	const FSoftObjectPath Key(InModel);
+	URigVMController* Controller = nullptr;
 
-	Controller->SetActionStack(GetOrCreateActionStack());
-	if(InModel)
 	{
-		Controller->SetGraph((URigVMGraph*)InModel);
+		UE::TScopeLock Lock(ControllersLock);
+
+		// check one more time if this controller already exists. given we are now deploying locks
+		// this is important to avoid double registration of controllers.
+		if(const TObjectPtr<URigVMController>* ExistingController = Controllers.Find(Key))
+		{
+			if  (!URigVMHost::IsGarbageOrDestroyed(ExistingController->Get()))
+			{
+				checkf(ExistingController->Get()->GetGraph() == InModel, TEXT("Controller %s contains unexpected graph."), *ExistingController->GetPathName());
+				return ExistingController->Get();
+			}
+		}
+
+		const FString ModelName = InModel ? InModel->GetName() : TEXT("__NullGraph");
+		const FName SafeControllerName = GetUniqueName(*FString::Printf(TEXT("%s_Controller"), *ModelName));
+		Controller = NewObject<URigVMController>(GetOuter(), ControllerClass, SafeControllerName);
+		Controllers.FindOrAdd(Key, nullptr) = Controller;
+		if(InModel && InModel->GetSchemaClass() != nullptr)
+		{
+			Controller->SetSchemaClass(InModel->GetSchemaClass());
+		}
+		else
+		{
+			check(DefaultSchemaClass);
+			Controller->SetSchemaClass(DefaultSchemaClass);
+		}
+
+		Controller->SetActionStack(GetOrCreateActionStack());
+		if(InModel)
+		{
+			Controller->SetGraph((URigVMGraph*)InModel);
+		}
 	}
+
+	check(Controller);
+	
 	Controller->OnModified().AddLambda([this](ERigVMGraphNotifType InNotifType, URigVMGraph* InGraph, UObject* InSubject)
 	{
 		HandleGraphModifiedEvent(InNotifType, InGraph, InSubject);
diff --git a/Engine/Plugins/Runtime/RigVM/Source/RigVMDeveloper/Public/RigVMModel/RigVMClient.h b/Engine/Plugins/Runtime/RigVM/Source/RigVMDeveloper/Public/RigVMModel/RigVMClient.h
--- a/Engine/Plugins/Runtime/RigVM/Source/RigVMDeveloper/Public/RigVMModel/RigVMClient.h
+++ b/Engine/Plugins/Runtime/RigVM/Source/RigVMDeveloper/Public/RigVMModel/RigVMClient.h
@@ -7,6 +7,7 @@
 #include "RigVMFunctionLibrary.h"
 #include "RigVMController.h"
 #include "UObject/WeakInterfacePtr.h"
+#include "Misc/TransactionallySafeCriticalSection.h"
 #include "RigVMClient.generated.h"
 
 #define UE_API RIGVMDEVELOPER_API
@@ -236,6 +237,8 @@ public:
 	UE_API void SetFromDeprecatedData(URigVMGraph* InDefaultGraph, URigVMFunctionLibrary* InFunctionLibrary);
 	UE_API void SetExternalModelHost(IRigVMClientExternalModelHost* InExternalModelHost);
 
+	UE_API FRigVMClient& operator =(const FRigVMClient& InOther);
+
 	UE_API void Reset();
 	int32 Num() const { return GetModels().Num(); }
 	UE_DEPRECATED(5.5, "Please use GetDefaultSchema or get a schema per controller/graph.")
@@ -404,9 +407,21 @@ private:
 	FRigVMClientPostGraphModified PostGraphModifiedDelegate;
 #endif // WITH_EDITOR
 
+	mutable FTransactionallySafeCriticalSection OuterClientHostLock;
+	mutable FTransactionallySafeCriticalSection ControllersLock;
+
 	friend class UEngineTestClientHost;
 	friend class URigVMBlueprint;
 	friend class IRigVMAssetInterface;
 };
 
+template<>
+struct TStructOpsTypeTraits<FRigVMClient> : public TStructOpsTypeTraitsBase2<FRigVMClient>
+{
+	enum 
+	{
+		WithCopy = true,
+	};
+};
+
 #undef UE_API
