Fixed functions not getting function compilation data dirtied when local variables are created, deleted or modified.
Also fixes local variables default values getting lost when recompiling the graph.

diff --git a/Engine/Plugins/Runtime/RigVM/Source/RigVMDeveloper/Private/RigVMModel/RigVMClient.cpp b/Engine/Plugins/Runtime/RigVM/Source/RigVMDeveloper/Private/RigVMModel/RigVMClient.cpp
--- a/Engine/Plugins/Runtime/RigVM/Source/RigVMDeveloper/Private/RigVMModel/RigVMClient.cpp
+++ b/Engine/Plugins/Runtime/RigVM/Source/RigVMDeveloper/Private/RigVMModel/RigVMClient.cpp
@@ -1657,6 +1657,18 @@ void FRigVMClient::HandleGraphModifiedEvent(ERigVMGraphNotifType InNotifType, UR
 			}
 			break;	
 		}
+		case ERigVMGraphNotifType::LocalVariableAdded:
+		case ERigVMGraphNotifType::LocalVariableRemoved:
+		case ERigVMGraphNotifType::LocalVariableDefaultValueChanged:
+		case ERigVMGraphNotifType::LocalVariableRenamed:
+		case ERigVMGraphNotifType::LocalVariableTypeChanged:
+		{
+			if(URigVMLibraryNode* LibraryNode = Cast<URigVMLibraryNode>(InSubject))
+			{
+				DirtyGraphFunctionCompilationData(LibraryNode);
+			}
+			break;
+		}
 		case ERigVMGraphNotifType::VariableAdded: // A variable has been added (Subject == URigVMVariableNode)
 		case ERigVMGraphNotifType::VariableRemoved: // A variable has been removed (Subject == URigVMVariableNode)
 		case ERigVMGraphNotifType::VariableRenamed: // A variable has been renamed (Subject == URigVMVariableNode)
diff --git a/Engine/Plugins/Runtime/RigVM/Source/RigVMDeveloper/Private/RigVMModel/RigVMController.cpp b/Engine/Plugins/Runtime/RigVM/Source/RigVMDeveloper/Private/RigVMModel/RigVMController.cpp
--- a/Engine/Plugins/Runtime/RigVM/Source/RigVMDeveloper/Private/RigVMModel/RigVMController.cpp
+++ b/Engine/Plugins/Runtime/RigVM/Source/RigVMDeveloper/Private/RigVMModel/RigVMController.cpp
@@ -15710,6 +15710,11 @@ FRigVMGraphVariableDescription URigVMController::AddLocalVariable(const FName& I
 	if (!bSuspendNotifications)
 	{
 		(void)Graph->MarkPackageDirty();
+
+		if (URigVMLibraryNode* LibraryNode = Cast<URigVMLibraryNode>(Graph->GetOuter()))
+		{
+			Notify(ERigVMGraphNotifType::LocalVariableAdded, LibraryNode);
+		}
 	}
 
 	if (bPrintPythonCommand)
@@ -15865,6 +15870,10 @@ bool URigVMController::RemoveLocalVariable(const FName& InVariableName, bool bSe
 		if (!bSuspendNotifications)
 		{
 			(void)Graph->MarkPackageDirty();
+			if (URigVMLibraryNode* LibraryNode = Cast<URigVMLibraryNode>(Graph->GetOuter()))
+			{
+				Notify(ERigVMGraphNotifType::LocalVariableRemoved, LibraryNode);
+			}
 		}
 
 		if (bSetupUndoRedo)
@@ -15963,12 +15972,16 @@ bool URigVMController::RenameLocalVariable(const FName& InVariableName, const FN
 		}
 	}
 
-	for (URigVMNode* RenamedNode : RenamedNodes)
+	if (!bSuspendNotifications)
 	{
-		Notify(ERigVMGraphNotifType::VariableRenamed, RenamedNode);
-		if (!bSuspendNotifications)
+		(void)Graph->MarkPackageDirty();
+		if (URigVMLibraryNode* LibraryNode = Cast<URigVMLibraryNode>(Graph->GetOuter()))
 		{
-			(void)Graph->MarkPackageDirty();
+			Notify(ERigVMGraphNotifType::LocalVariableDefaultValueChanged, LibraryNode);
+		}
+		for (URigVMNode* RenamedNode : RenamedNodes)
+		{
+			Notify(ERigVMGraphNotifType::VariableRenamed, RenamedNode);
 		}
 	}
 
@@ -16021,6 +16034,10 @@ bool URigVMController::SetLocalVariableType(const FName& InVariableName, const F
 	if (!bSuspendNotifications)
 	{
 		(void)Graph->MarkPackageDirty();
+		if (URigVMLibraryNode* LibraryNode = Cast<URigVMLibraryNode>(Graph->GetOuter()))
+		{
+			Notify(ERigVMGraphNotifType::LocalVariableTypeChanged, LibraryNode);
+		}
 	}
 
 	FRigVMControllerCompileBracketScope CompileScope(this);
@@ -16164,6 +16181,11 @@ bool URigVMController::SetLocalVariableDefaultValue(const FName& InVariableName,
 	if (!bSuspendNotifications)
 	{
 		(void)Graph->MarkPackageDirty();
+
+		if (URigVMLibraryNode* LibraryNode = Cast<URigVMLibraryNode>(Graph->GetOuter()))
+		{
+			Notify(ERigVMGraphNotifType::LocalVariableDefaultValueChanged, LibraryNode);
+		}
 	}
 
 	if (bPrintPythonCommand)
diff --git a/Engine/Plugins/Runtime/RigVM/Source/RigVMDeveloper/Public/RigVMModel/RigVMNotifications.h b/Engine/Plugins/Runtime/RigVM/Source/RigVMDeveloper/Public/RigVMModel/RigVMNotifications.h
--- a/Engine/Plugins/Runtime/RigVM/Source/RigVMDeveloper/Public/RigVMModel/RigVMNotifications.h
+++ b/Engine/Plugins/Runtime/RigVM/Source/RigVMDeveloper/Public/RigVMModel/RigVMNotifications.h
@@ -61,6 +61,11 @@ enum class ERigVMGraphNotifType : uint8
 	PinCategoryExpansionChanged, // The category of a pin expanded / collapsed (Subject == URigVMNode)
 	FunctionVariantGuidChanged, // The guid for a function has changed (Subject == URigVMLibraryNode)
 	NodeEarlyExitChanged, // The node's early exit state has changed
+	LocalVariableDefaultValueChanged, // A local variable default value has been set (Subject == URigVMLibraryNode)
+	LocalVariableAdded, // A local variable has been added to a function (Subject == URigVMLibraryNode)
+	LocalVariableRemoved, // A local variable has been removed from a function (Subject == URigVMLibraryNode)
+	LocalVariableRenamed, // A local variable has been renamed inside a function (Subject == URigVMLibraryNode)
+	LocalVariableTypeChanged, // A local variable has been renamed inside a function (Subject == URigVMLibraryNode)
 	Invalid // The max for this enum (used for guarding)
 };
 
diff --git a/Engine/Plugins/Runtime/RigVM/Source/RigVMEditor/Private/Editor/RigVMEditor.cpp b/Engine/Plugins/Runtime/RigVM/Source/RigVMEditor/Private/Editor/RigVMEditor.cpp
--- a/Engine/Plugins/Runtime/RigVM/Source/RigVMEditor/Private/Editor/RigVMEditor.cpp
+++ b/Engine/Plugins/Runtime/RigVM/Source/RigVMEditor/Private/Editor/RigVMEditor.cpp
@@ -2170,7 +2170,6 @@ void FRigVMEditorBase::OnWrappedPropertyChangedChainEvent(URigVMDetailsViewWrapp
 			}
 			else if (PropertyPath == TEXT("DefaultValue"))
 			{
-				FRigVMControllerNotifGuard NotifGuard(Controller, true);
 				for (FRigVMGraphVariableDescription& Variable : Graph->GetLocalVariables())
 				{
 					if (Variable.Name == VariableDescription.Name)
