Control Rig: Attempt to fix crash when building some projects
#rnx

diff --git a/Engine/Plugins/Runtime/RigVM/Source/RigVMDeveloper/Private/RigVMModel/RigVMController.cpp b/Engine/Plugins/Runtime/RigVM/Source/RigVMDeveloper/Private/RigVMModel/RigVMController.cpp
--- a/Engine/Plugins/Runtime/RigVM/Source/RigVMDeveloper/Private/RigVMModel/RigVMController.cpp
+++ b/Engine/Plugins/Runtime/RigVM/Source/RigVMDeveloper/Private/RigVMModel/RigVMController.cpp
@@ -20042,7 +20042,7 @@ void URigVMController::ApplyPinState(URigVMPin* InPin, const FPinState& InPinSta
 {
 	for (URigVMInjectionInfo* InjectionInfo : InPinState.InjectionInfos)
 	{
-		if (InjectionInfo)
+		if (InjectionInfo && InjectionInfo->Node)
 		{
 			RenameObject(InjectionInfo, nullptr, InPin);
 			InjectionInfo->InputPin = InjectionInfo->InputPin ? InjectionInfo->Node->FindPin(InjectionInfo->InputPin->GetName()) : nullptr;
