Control Rig: Attempt to fix crash when building some projects
#rnx

diff --git a/Engine/Plugins/Runtime/RigVM/Source/RigVMDeveloper/Private/RigVMModel/RigVMController.cpp b/Engine/Plugins/Runtime/RigVM/Source/RigVMDeveloper/Private/RigVMModel/RigVMController.cpp
--- a/Engine/Plugins/Runtime/RigVM/Source/RigVMDeveloper/Private/RigVMModel/RigVMController.cpp
+++ b/Engine/Plugins/Runtime/RigVM/Source/RigVMDeveloper/Private/RigVMModel/RigVMController.cpp
@@ -20042,6 +20042,12 @@ void URigVMController::ApplyPinState(URigVMPin* InPin, const FPinState& InPinSta
 {
 	for (URigVMInjectionInfo* InjectionInfo : InPinState.InjectionInfos)
 	{
+		// Copy over the old unit node if the content is from prior to that property being deprecated.
+		if (InjectionInfo && !InjectionInfo->Node)
+		{
+			InjectionInfo->Node = InjectionInfo->UnitNode_DEPRECATED;
+		}
+
 		if (InjectionInfo && InjectionInfo->Node)
 		{
 			RenameObject(InjectionInfo, nullptr, InPin);
