Control Rig: Fix undo redo remove unsued variables does not restore all removed variables

#rnx

diff --git a/Engine/Plugins/Runtime/RigVM/Source/RigVMDeveloper/Private/RigVMBlueprintLegacy.cpp b/Engine/Plugins/Runtime/RigVM/Source/RigVMDeveloper/Private/RigVMBlueprintLegacy.cpp
--- a/Engine/Plugins/Runtime/RigVM/Source/RigVMDeveloper/Private/RigVMBlueprintLegacy.cpp
+++ b/Engine/Plugins/Runtime/RigVM/Source/RigVMDeveloper/Private/RigVMBlueprintLegacy.cpp
@@ -929,6 +929,21 @@ bool URigVMBlueprint::RemoveMemberVariable(const FName& InName)
 	return true;
 }
 
+bool URigVMBlueprint::BulkRemoveMemberVariables(const TArray<FName>& InNames)
+{
+	for (const FName& Name : InNames)
+	{
+		const int32 VarIndex = FBlueprintEditorUtils::FindNewVariableIndex(this, Name);
+		if (VarIndex == INDEX_NONE)
+		{
+			return false;
+		}
+	}
+
+	FBlueprintEditorUtils::BulkRemoveMemberVariables(this, InNames);
+	return true;
+}
+
 bool URigVMBlueprint::RenameMemberVariable(const FName& InOldName, const FName& InNewName)
 {
 	int32 VarIndex = FBlueprintEditorUtils::FindNewVariableIndex(this, InOldName);
diff --git a/Engine/Plugins/Runtime/RigVM/Source/RigVMDeveloper/Public/RigVMAsset.h b/Engine/Plugins/Runtime/RigVM/Source/RigVMDeveloper/Public/RigVMAsset.h
--- a/Engine/Plugins/Runtime/RigVM/Source/RigVMDeveloper/Public/RigVMAsset.h
+++ b/Engine/Plugins/Runtime/RigVM/Source/RigVMDeveloper/Public/RigVMAsset.h
@@ -211,6 +211,7 @@ public:
 	virtual FName AddHostMemberVariableFromExternal(FRigVMExternalVariable InVariableToCreate, FString InDefaultValue = FString()) = 0;
 	virtual FName AddMemberVariable(const FName& InName, const FString& InCPPType, bool bIsPublic = false, bool bIsReadOnly = false, FString InDefaultValue = TEXT("")) = 0;
 	virtual bool RemoveMemberVariable(const FName& InName) = 0;
+	virtual bool BulkRemoveMemberVariables(const TArray<FName>& InNames) = 0;
 	virtual bool RenameMemberVariable(const FName& InOldName, const FName& InNewName) = 0;
 	virtual bool ChangeMemberVariableType(const FName& InName, const FString& InCPPType, bool bIsPublic = false, bool bIsReadOnly = false, FString InDefaultValue = TEXT("")) = 0;
 	virtual bool ChangeMemberVariableType(const FName& InName, const FEdGraphPinType& InType) = 0;
diff --git a/Engine/Plugins/Runtime/RigVM/Source/RigVMDeveloper/Public/RigVMBlueprintLegacy.h b/Engine/Plugins/Runtime/RigVM/Source/RigVMDeveloper/Public/RigVMBlueprintLegacy.h
--- a/Engine/Plugins/Runtime/RigVM/Source/RigVMDeveloper/Public/RigVMBlueprintLegacy.h
+++ b/Engine/Plugins/Runtime/RigVM/Source/RigVMDeveloper/Public/RigVMBlueprintLegacy.h
@@ -404,6 +404,9 @@ public:
 	UFUNCTION(BlueprintCallable, Category = "Variables")
 	virtual bool RemoveMemberVariable(const FName& InName) override;
 
+	UFUNCTION(BlueprintCallable, Category = "Variables")
+	virtual bool BulkRemoveMemberVariables(const TArray<FName>& InNames) override;
+
 	UFUNCTION(BlueprintCallable, Category = "Variables")
 	virtual bool RenameMemberVariable(const FName& InOldName, const FName& InNewName) override;
 
diff --git a/Engine/Plugins/Runtime/RigVM/Source/RigVMEditor/Private/Editor/RemoveUnusedMembers/RigVMEditorRemoveUnusedMembers.cpp b/Engine/Plugins/Runtime/RigVM/Source/RigVMEditor/Private/Editor/RemoveUnusedMembers/RigVMEditorRemoveUnusedMembers.cpp
--- a/Engine/Plugins/Runtime/RigVM/Source/RigVMEditor/Private/Editor/RemoveUnusedMembers/RigVMEditorRemoveUnusedMembers.cpp
+++ b/Engine/Plugins/Runtime/RigVM/Source/RigVMEditor/Private/Editor/RemoveUnusedMembers/RigVMEditorRemoveUnusedMembers.cpp
@@ -4,6 +4,7 @@
 
 #include "Algo/Transform.h"
 #include "EdGraph/RigVMEdGraphSchema.h"
+#include "Kismet2/BlueprintEditorUtils.h"
 
 #define LOCTEXT_NAMESPACE "RigVMEditorRemoveUnusedMembersDialog"
 
@@ -140,28 +141,27 @@ namespace UE::RigVMEditor::Private
 		const TScriptInterface<IRigVMAssetInterface>& Asset,
 		const TArray<FName>& MemberNamesToRemove)
 	{
-		if (!Asset)
+		UObject* Object = Asset ? Asset->GetObject() : nullptr;
+		URigVMGraph* FocusedGraph = Asset ? Asset->GetFocusedModel() : nullptr;
+		FRigVMClient* RigVMClient = Asset ? Asset->GetRigVMClient() : nullptr;
+		URigVMController* Controller = RigVMClient ? RigVMClient->GetController(FocusedGraph) : nullptr;
+		if (!Asset || 
+			!Object ||
+			!FocusedGraph ||
+			!RigVMClient ||
+			!Controller)
 		{
 			return;
 		}
 
 		// Variables
-		for (const FName& MemberNameToRemove : MemberNamesToRemove)
-		{
-			Asset->RemoveMemberVariable(MemberNameToRemove);
-		}
-
-		// Local variables
-		URigVMGraph* FocusedGraph = Asset->GetFocusedModel();
-		FRigVMClient* RigVMClient = Asset->GetRigVMClient();
-		URigVMController* Controller = RigVMClient ? RigVMClient->GetController(FocusedGraph) : nullptr;
-		if (!FocusedGraph ||
-			!RigVMClient ||
-			!Controller)
+		Object->Modify();
+		if (!Asset->BulkRemoveMemberVariables(MemberNamesToRemove))
 		{
 			return;
 		}
 
+		// Local variables
 		for (const FRigVMGraphVariableDescription& VariableDescription : FocusedGraph->GetLocalVariables())
 		{
 			if (MemberNamesToRemove.Contains(VariableDescription.Name))
