Control Rig: Fix ensure when a function reference requires remapping
#rnx

diff --git a/Engine/Plugins/Runtime/RigVM/Source/RigVMDeveloper/Private/RigVMModel/Nodes/RigVMFunctionReferenceNode.cpp b/Engine/Plugins/Runtime/RigVM/Source/RigVMDeveloper/Private/RigVMModel/Nodes/RigVMFunctionReferenceNode.cpp
--- a/Engine/Plugins/Runtime/RigVM/Source/RigVMDeveloper/Private/RigVMModel/Nodes/RigVMFunctionReferenceNode.cpp
+++ b/Engine/Plugins/Runtime/RigVM/Source/RigVMDeveloper/Private/RigVMModel/Nodes/RigVMFunctionReferenceNode.cpp
@@ -118,11 +118,19 @@ bool URigVMFunctionReferenceNode::RequiresVariableRemappingInternal(TArray<FRigV
 {
 	bool bHostedInDifferencePackage = false;
 	
+	FString ThisPacakgePath = GetPackage()->GetPathName();
+	
+	// avoid function reference related validation for temp assets, a temp asset may get generated during
+	// certain content validation process. It is usually just a simple file-level copy of the source asset
+	// so these references are usually not fixed-up properly. Thus, it is meaningless to validate them.
+	if (ThisPacakgePath.StartsWith(TEXT("/Temp/")))
+	{
+		return false;
+	}
+	
 	FRigVMGraphFunctionIdentifier LibraryPointer = ReferencedFunctionHeader.LibraryPointer;
 	const FString& LibraryPackagePath = LibraryPointer.GetNodeSoftPath().GetLongPackageName();
-	const FString& ThisPacakgePath = GetPackage()->GetPathName();
 	bHostedInDifferencePackage = LibraryPackagePath != ThisPacakgePath;
-		
 	if(bHostedInDifferencePackage)
 	{
 		InnerVariables = GetExternalVariables(false);
