Control Rig: Added information in ensure message
#rnx

diff --git a/Engine/Plugins/Runtime/RigVM/Source/RigVMDeveloper/Private/RigVMCompiler/RigVMCompiler.cpp b/Engine/Plugins/Runtime/RigVM/Source/RigVMDeveloper/Private/RigVMCompiler/RigVMCompiler.cpp
--- a/Engine/Plugins/Runtime/RigVM/Source/RigVMDeveloper/Private/RigVMCompiler/RigVMCompiler.cpp
+++ b/Engine/Plugins/Runtime/RigVM/Source/RigVMDeveloper/Private/RigVMCompiler/RigVMCompiler.cpp
@@ -3021,7 +3021,14 @@ bool URigVMCompiler::TraverseInlineFunction(const FRigVMInlineFunctionExprAST* I
 					}
 					else
 					{
-						ensureMsgf(!FunctionReferenceNode->RequiresVariableRemapping(), TEXT("Could not find variable %s in function reference %s variable map, in package %s\n"), *InnerVariableName.ToString(), *FunctionReferenceNode->GetNodePath(), *GetPackage()->GetPathName());
+						FRigVMGraphFunctionIdentifier LibraryPointer = FunctionReferenceNode->GetFunctionIdentifier();
+						const FString& LibraryPackagePath = LibraryPointer.GetNodeSoftPath().GetLongPackageName();
+						ensureMsgf(!FunctionReferenceNode->RequiresVariableRemapping(),
+							TEXT("Could not find variable %s in function reference %s variable map, in package %s\n. Library package is %s\n"),
+							*InnerVariableName.ToString(),
+							*FunctionReferenceNode->GetNodePath(),
+							*GetPackage()->GetPathName(),
+							*LibraryPackagePath);
 					}
 					const FRigVMOperand& OuterOperand = WorkData.PinPathToOperand->FindChecked(FString::Printf(TEXT("Variable::%s"), *OuterVariableName.ToString()));
 					Operand->RegisterIndex = OuterOperand.RegisterIndex;
