Fix for RigVM memory storage not returning correctly the property value as a string, when the the value is default initialzied to a value that is not the type default value

diff --git a/Engine/Plugins/Runtime/RigVM/Source/RigVM/Private/RigVMCore/RigVMMemoryStorageStruct.cpp b/Engine/Plugins/Runtime/RigVM/Source/RigVM/Private/RigVMCore/RigVMMemoryStorageStruct.cpp
--- a/Engine/Plugins/Runtime/RigVM/Source/RigVM/Private/RigVMCore/RigVMMemoryStorageStruct.cpp
+++ b/Engine/Plugins/Runtime/RigVM/Source/RigVM/Private/RigVMCore/RigVMMemoryStorageStruct.cpp
@@ -229,8 +229,11 @@ uint32 FRigVMMemoryStorageStruct::GetMemoryHash() const
 	{
 		for (const FProperty* Property : LinkedProperties)
 		{
+			// Note that it is using Export Direct and also passing Data as the default value forces the text exporter to generate all the elements
+			// If we pass nullptr as default, values that are default initialized to a value that is different from the type default value will not be correctly exported
 			FString DefaultValue;
-			Property->ExportTextItem_InContainer(DefaultValue, GetContainerPtr(), nullptr, nullptr, PPF_None, nullptr);
+			const uint8* Data = GetData<uint8>(Property);
+			Property->ExportTextItem_Direct(DefaultValue, Data, Data, nullptr, PPF_None, nullptr);
 			CachedMemoryHash = HashCombine(CachedMemoryHash, GetTypeHash(DefaultValue));
 		}
 	}
@@ -246,10 +249,11 @@ const TArray<FRigVMPropertyPath> FRigVMMemoryStorageStruct::EmptyPropertyPaths;
 FString FRigVMMemoryStorageStruct::GetDataAsString(int32 InPropertyIndex, int32 PortFlags)
 {
 	check(IsValidIndex(InPropertyIndex));
-	const uint8* Data = GetData<uint8>(InPropertyIndex);
-
+	// Note that it is using Export Direct and also passing Data as the default value, so it forces the text exporter to generate all the elements
+	// If we pass nullptr, values that are default initialized to a value that is different from the type default value will not be correctly exported
 	FString RetValue;
-	GetProperty(InPropertyIndex)->ExportTextItem_Direct(RetValue, Data, nullptr, nullptr, PortFlags);
+	const uint8* Data = GetData<uint8>(InPropertyIndex);
+	GetProperty(InPropertyIndex)->ExportTextItem_Direct(RetValue, Data, Data, nullptr, PortFlags);
 	return RetValue;
 }
 
diff --git a/Engine/Plugins/Runtime/RigVM/Source/RigVM/Public/RigVMCore/RigVMMemoryStorageStruct.h b/Engine/Plugins/Runtime/RigVM/Source/RigVM/Public/RigVMCore/RigVMMemoryStorageStruct.h
--- a/Engine/Plugins/Runtime/RigVM/Source/RigVM/Public/RigVMCore/RigVMMemoryStorageStruct.h
+++ b/Engine/Plugins/Runtime/RigVM/Source/RigVM/Public/RigVMCore/RigVMMemoryStorageStruct.h
@@ -126,12 +126,19 @@ struct FRigVMMemoryStorageStruct : public FInstancedPropertyBag
 
 	// Returns the memory for a property given its index
 	template<typename T>
-	T* GetData(int32 InPropertyIndex)
+	T* GetData(int32 InPropertyIndex) 
 	{
 		const TArray<const FProperty*>& Properties = GetProperties();
 		check(Properties.IsValidIndex(InPropertyIndex));
 		return Properties[InPropertyIndex]->ContainerPtrToValuePtr<T>(GetMutableValue().GetMemory());
 	}
+	template<typename T>
+	const T* GetData(int32 InPropertyIndex) const
+	{
+		const TArray<const FProperty*>& Properties = GetProperties();
+		check(Properties.IsValidIndex(InPropertyIndex));
+		return Properties[InPropertyIndex]->ContainerPtrToValuePtr<T>(GetValue().GetMemory());
+	}
 
 	// Returns the memory for a property given its name (or nullptr)
 	template<typename T>
@@ -156,6 +163,16 @@ struct FRigVMMemoryStorageStruct : public FInstancedPropertyBag
 
 		return nullptr;
 	}
+	template<typename T>
+	const T* GetData(const FProperty* Property) const
+	{
+		if (const int32 PropertyIndex = GetPropertyIndex(Property); PropertyIndex != INDEX_NONE)
+		{
+			return GetData<T>(PropertyIndex);
+		}
+
+		return nullptr;
+	}
 
 	// Returns the memory for a property given its index and a matching property path
 	template<typename T>
