Fixed a crash in ControlRig when expanding a function reference from a different graph that contains a private function inside. Now the user will receive a localize request to make a local copy of the function

diff --git a/Engine/Plugins/Runtime/RigVM/Source/RigVMDeveloper/Private/RigVMModel/RigVMController.cpp b/Engine/Plugins/Runtime/RigVM/Source/RigVMDeveloper/Private/RigVMModel/RigVMController.cpp
--- a/Engine/Plugins/Runtime/RigVM/Source/RigVMDeveloper/Private/RigVMModel/RigVMController.cpp
+++ b/Engine/Plugins/Runtime/RigVM/Source/RigVMDeveloper/Private/RigVMModel/RigVMController.cpp
@@ -6038,7 +6038,14 @@ TArray<URigVMNode*> URigVMController::ExpandLibraryNode(URigVMLibraryNode* InNod
 		ExpandedNodes.Add(ExpandedNode);
 	}
 
-	check(ExpandedNodeNames.Num() >= NodeNames.Num());
+	if (ExpandedNodeNames.Num() < NodeNames.Num()) // if a private function is found and the user decides not to localize, then the nodes count will not match
+	{
+		if (bSetupUndoRedo)
+		{
+			GetActionStack()->CancelAction(ExpandAction);
+		}
+		return TArray<URigVMNode*>();
+	}
 
 	TMap<FName, FName> NodeNameMap;
 	for (int32 NodeNameIndex = 0, ExpandedNodeNameIndex = 0, InjectedNodeName = 0; ExpandedNodeNameIndex < ExpandedNodeNames.Num(); ExpandedNodeNameIndex++)
diff --git a/Engine/Plugins/Runtime/RigVM/Source/RigVMEditor/Private/RigVMEditorModule.cpp b/Engine/Plugins/Runtime/RigVM/Source/RigVMEditor/Private/RigVMEditorModule.cpp
--- a/Engine/Plugins/Runtime/RigVM/Source/RigVMEditor/Private/RigVMEditorModule.cpp
+++ b/Engine/Plugins/Runtime/RigVM/Source/RigVMEditor/Private/RigVMEditorModule.cpp
@@ -1365,8 +1365,23 @@ void FRigVMEditorModule::GetNodeOrganizationContextMenuActions(TScriptInterface<
 			LOCTEXT("ExpandNode", "Expand Node"),
 			LOCTEXT("ExpandNode_Tooltip", "Expands the contents of the node into this graph"),
 			FSlateIcon(),
-			FUIAction(FExecuteAction::CreateLambda([Controller, LibraryNode]() {
+			FUIAction(FExecuteAction::CreateLambda([Controller, LibraryNode, RigVMBlueprint]() {
 				Controller->OpenUndoBracket(TEXT("Expand node"));
+
+				URigVMFunctionLibrary* FunctionLibrary = LibraryNode->GetGraph()->GetDefaultFunctionLibrary();
+				TGuardValue<FRigVMController_RequestLocalizeFunctionDelegate> RequestLocalizeDelegateGuard(
+					Controller->RequestLocalizeFunctionDelegate,
+					FRigVMController_RequestLocalizeFunctionDelegate::CreateLambda([RigVMBlueprint, FunctionLibrary](FRigVMGraphFunctionIdentifier& InFunctionToLocalize)
+						{
+							if (RigVMBlueprint)
+							{
+								RigVMBlueprint->BroadcastRequestLocalizeFunctionDialog(InFunctionToLocalize, true);
+							}
+							const URigVMLibraryNode* LocalizedFunctionNode = FunctionLibrary->FindPreviouslyLocalizedFunction(InFunctionToLocalize);
+							return LocalizedFunctionNode != nullptr;
+						})
+				);
+
 				TArray<URigVMNode*> ExpandedNodes = Controller->ExpandLibraryNode(LibraryNode->GetFName(), true, true);
 				if (ExpandedNodes.Num() > 0)
 				{
