Fixed RigVM local variable customization not showing value for local variables not used in the graph and losing value after compiling by avoiding the use of the VM memory and using a local instanced struct instead.

diff --git a/Engine/Plugins/Runtime/RigVM/Source/RigVMEditor/Private/Editor/RigVMLocalVariableDetailCustomization.cpp b/Engine/Plugins/Runtime/RigVM/Source/RigVMEditor/Private/Editor/RigVMLocalVariableDetailCustomization.cpp
--- a/Engine/Plugins/Runtime/RigVM/Source/RigVMEditor/Private/Editor/RigVMLocalVariableDetailCustomization.cpp
+++ b/Engine/Plugins/Runtime/RigVM/Source/RigVMEditor/Private/Editor/RigVMLocalVariableDetailCustomization.cpp
@@ -28,7 +28,7 @@ void FRigVMLocalVariableDetailCustomization::CustomizeDetails(IDetailLayoutBuild
 		ObjectsBeingCustomized.Add(WrapperObject);
 	}
 
-	if(ObjectsBeingCustomized[0].IsValid())
+	if (ObjectsBeingCustomized.Num() > 0 && ObjectsBeingCustomized[0].IsValid())
 	{
 		VariableDescription = ObjectsBeingCustomized[0]->GetContent<FRigVMGraphVariableDescription>();
 		if (UObject* Subject = ObjectsBeingCustomized[0]->GetSubject())
@@ -90,34 +90,52 @@ void FRigVMLocalVariableDetailCustomization::CustomizeDetails(IDetailLayoutBuild
 		];
 
 
-	if (BlueprintBeingCustomized)
+	if (GraphBeingCustomized)
 	{
-		URigVM* VM = BlueprintBeingCustomized->GetVM(false);
-		if (VM != nullptr)
+		FString SourcePath = FString::Printf(TEXT("LocalVariableDefault::%s|%s::Const"), *GraphBeingCustomized->GetGraphName(), *VariableDescription.Name.ToString());
+
+		// Here, we create a custom storage for the Details, so when it changes,
+		// we extract the value and then update the default value handle with it,
+		// which in turn will call URigVMDetailsViewWrapperObject::PostEditChangeChainProperty
+		// and then the variable default value will be set through the graph controller
+		// SetLocalVariableDefaultValue function
+		TArray<FRigVMPropertyDescription> PropertyDescriptions;
+		FRigVMPropertyDescription& PropertyDesc = PropertyDescriptions.AddDefaulted_GetRef();
+		PropertyDesc.Name = *SourcePath;
+		PropertyDesc.Property = nullptr;
+		PropertyDesc.CPPType = VariableDescription.CPPType;
+		PropertyDesc.CPPTypeObject = VariableDescription.CPPTypeObject;
+		if (RigVMTypeUtils::IsArrayType(PropertyDesc.CPPType))
 		{
-			FString SourcePath = FString::Printf(TEXT("LocalVariableDefault::%s|%s::Const"), *GraphBeingCustomized->GetGraphName(), *VariableDescription.Name.ToString());
-			FRigVMMemoryStorageStruct* LiteralMemory = VM->GetLiteralMemory();
-			FProperty* Property = LiteralMemory->FindPropertyByName(*SourcePath);
-			if (Property)
-			{
-				IDetailCategoryBuilder& DefaultValueCategory = DetailBuilder.EditCategory(TEXT("DefaultValueCategory"), LOCTEXT("DefaultValueCategoryHeading", "Default Value"));
-				Property->ClearPropertyFlags(CPF_EditConst);
+			PropertyDesc.Containers.Add(EPinContainerType::Array);
+		}
+		PropertyDesc.DefaultValue = VariableDescription.DefaultValue;
+
+		MemoryStorage = MakeShared<FRigVMMemoryStorageStruct>();
+		MemoryStorage->AddProperties(PropertyDescriptions);
+
+
+		FRigVMMemoryStorageStruct* Memory = MemoryStorage.Get();
+		FProperty* Property = Memory->FindPropertyByName(*SourcePath);
+		if (Property)
+		{
+			IDetailCategoryBuilder& DefaultValueCategory = DetailBuilder.EditCategory(TEXT("DefaultValueCategory"), LOCTEXT("DefaultValueCategoryHeading", "Default Value"));
+			Property->ClearPropertyFlags(CPF_EditConst);
 			
-				const FName SanitizedName = FRigVMPropertyDescription::SanitizeName(*SourcePath);
-				IDetailPropertyRow* Row = DefaultValueCategory.AddExternalStructureProperty(MakeShared<FInstancePropertyBagStructureDataProvider>(*LiteralMemory), SanitizedName);
+			const FName SanitizedName = FRigVMPropertyDescription::SanitizeName(*SourcePath);
+			IDetailPropertyRow* Row = DefaultValueCategory.AddExternalStructureProperty(MakeShared<FInstancePropertyBagStructureDataProvider>(*Memory), SanitizedName);
 
-				Row->DisplayName(FText::FromName(VariableDescription.Name));
+			Row->DisplayName(FText::FromName(VariableDescription.Name));
 
-				const FSimpleDelegate OnDefaultValueChanged = FSimpleDelegate::CreateLambda([this, Property, LiteralMemory]()
-				{
-					VariableDescription.DefaultValue = LiteralMemory->GetDataAsString(LiteralMemory->GetPropertyIndex(Property));
-					DefaultValueHandle->SetValue(VariableDescription.DefaultValue);
-				});
+			const FSimpleDelegate OnDefaultValueChanged = FSimpleDelegate::CreateLambda([this, Property]()
+			{
+				VariableDescription.DefaultValue = MemoryStorage->GetDataAsString(MemoryStorage->GetPropertyIndex(Property));
+				DefaultValueHandle->SetValue(VariableDescription.DefaultValue);
+			});
 
-				TSharedPtr<IPropertyHandle> Handle = Row->GetPropertyHandle();
-				Handle->SetOnPropertyValueChanged(OnDefaultValueChanged);
-				Handle->SetOnChildPropertyValueChanged(OnDefaultValueChanged);
-			}
+			TSharedPtr<IPropertyHandle> Handle = Row->GetPropertyHandle();
+			Handle->SetOnPropertyValueChanged(OnDefaultValueChanged);
+			Handle->SetOnChildPropertyValueChanged(OnDefaultValueChanged);
 		}
 	}
 }
diff --git a/Engine/Plugins/Runtime/RigVM/Source/RigVMEditor/Public/Editor/RigVMLocalVariableDetailCustomization.h b/Engine/Plugins/Runtime/RigVM/Source/RigVMEditor/Public/Editor/RigVMLocalVariableDetailCustomization.h
--- a/Engine/Plugins/Runtime/RigVM/Source/RigVMEditor/Public/Editor/RigVMLocalVariableDetailCustomization.h
+++ b/Engine/Plugins/Runtime/RigVM/Source/RigVMEditor/Public/Editor/RigVMLocalVariableDetailCustomization.h
@@ -17,6 +17,7 @@
 #define UE_API RIGVMEDITOR_API
 
 class IPropertyHandle;
+struct FRigVMMemoryStorageStruct;
 
 class FRigVMLocalVariableDetailCustomization : public IDetailCustomization
 {
@@ -49,6 +50,7 @@ private:
 	TSharedPtr<IPropertyHandle> TypeHandle;
 	TSharedPtr<IPropertyHandle> TypeObjectHandle;
 	TSharedPtr<IPropertyHandle> DefaultValueHandle;
+	TSharedPtr<FRigVMMemoryStorageStruct> MemoryStorage;
 
 	FRigVMLocalVariableNameValidator NameValidator;
 	TArray<TSharedPtr<FString>> EnumOptions;
