Avoid a crash in Control Rig editor when deleting the current graph and trying to add a new item using the add button, as there is no current focused model

#rnx

diff --git a/Engine/Plugins/Runtime/RigVM/Source/RigVMEditor/Private/Editor/RigVMEditor.cpp b/Engine/Plugins/Runtime/RigVM/Source/RigVMEditor/Private/Editor/RigVMEditor.cpp
--- a/Engine/Plugins/Runtime/RigVM/Source/RigVMEditor/Private/Editor/RigVMEditor.cpp
+++ b/Engine/Plugins/Runtime/RigVM/Source/RigVMEditor/Private/Editor/RigVMEditor.cpp
@@ -682,11 +682,13 @@ bool FRigVMEditorBase::TransactionObjectAffectsBlueprintImpl(UObject* InTransact
 
 bool FRigVMEditorBase::CanAddNewLocalVariable() const
 {
-	const URigVMGraph* Graph = GetFocusedModel();
-	const URigVMGraph* ParentGraph = Graph->GetParentGraph();				
-	if (ParentGraph && ParentGraph->IsA<URigVMFunctionLibrary>())
+	if (const URigVMGraph* Graph = GetFocusedModel())
 	{
-		return true;
+		const URigVMGraph* ParentGraph = Graph->GetParentGraph();
+		if (ParentGraph && ParentGraph->IsA<URigVMFunctionLibrary>())
+		{
+			return true;
+		}
 	}
 	return false;
 }
