RigVMController: Guard against garbage collected injection nodes

#rnx

diff --git a/Engine/Plugins/Runtime/RigVM/Source/RigVMDeveloper/Private/RigVMModel/RigVMController.cpp b/Engine/Plugins/Runtime/RigVM/Source/RigVMDeveloper/Private/RigVMModel/RigVMController.cpp
--- a/Engine/Plugins/Runtime/RigVM/Source/RigVMDeveloper/Private/RigVMModel/RigVMController.cpp
+++ b/Engine/Plugins/Runtime/RigVM/Source/RigVMDeveloper/Private/RigVMModel/RigVMController.cpp
@@ -17,6 +17,7 @@
 #include "RigVMFunctions/Execution/RigVMFunction_Sequence.h"
 #include "RigVMCompiler/RigVMCompiler.h"
 #include "RigVMDeveloperModule.h"
+#include "RigVMHost.h"
 #include "UObject/PropertyPortFlags.h"
 #include "UObject/Package.h"
 #include "Misc/CoreMisc.h"
@@ -20055,6 +20056,12 @@ void URigVMController::ApplyPinState(URigVMPin* InPin, const FPinState& InPinSta
 			InjectionInfo->Node = InjectionInfo->UnitNode_DEPRECATED;
 		}
 
+		// If this node has been destroyed already or is marked as garbage
+		if (InjectionInfo && URigVMHost::IsGarbageOrDestroyed(InjectionInfo->Node))
+		{
+			continue;
+		}
+
 		if (InjectionInfo && InjectionInfo->Node)
 		{
 			RenameObject(InjectionInfo, nullptr, InPin);
