RigVM: Relax tests for existing type indices when loading back content into the engine

#rnx

diff --git a/Engine/Plugins/Runtime/RigVM/Source/RigVM/Private/RigVMCore/RigVMRegistry.cpp b/Engine/Plugins/Runtime/RigVM/Source/RigVM/Private/RigVMCore/RigVMRegistry.cpp
--- a/Engine/Plugins/Runtime/RigVM/Source/RigVM/Private/RigVMCore/RigVMRegistry.cpp
+++ b/Engine/Plugins/Runtime/RigVM/Source/RigVM/Private/RigVMCore/RigVMRegistry.cpp
@@ -1189,9 +1189,11 @@ TRigVMTypeIndex FRigVMRegistry_NoLock::FindOrAddType_NoLock(const FRigVMTemplate
 			const FSoftObjectPath SoftObjectPath(ElementType.CPPTypeObject);
 			if (const TRigVMTypeIndexTriple* ExistingIndices = SoftObjectPathToTypeIndex.Find(SoftObjectPath))
 			{
-				check(Types[ExistingIndices->Get<0>()].IsUnknownType());
-				check(Types[ExistingIndices->Get<1>()].IsUnknownType());
-				check(Types[ExistingIndices->Get<2>()].IsUnknownType());
+				// this can potentially happen multiple times - it's ok
+				// if the types at the existing locations are valid.
+				check(Types.IsValidIndex(ExistingIndices->Get<0>()));
+				check(Types.IsValidIndex(ExistingIndices->Get<1>()));
+				check(Types.IsValidIndex(ExistingIndices->Get<2>()));
 				Indices.Add(ExistingIndices->Get<0>());
 				Indices.Add(ExistingIndices->Get<1>());
 				Indices.Add(ExistingIndices->Get<2>());
