Control Rig: Fix crash when removing array element pins of type struct
#rnx

diff --git a/Engine/Plugins/Runtime/RigVM/Source/RigVMDeveloper/Private/RigVMModel/RigVMController.cpp b/Engine/Plugins/Runtime/RigVM/Source/RigVMDeveloper/Private/RigVMModel/RigVMController.cpp
--- a/Engine/Plugins/Runtime/RigVM/Source/RigVMDeveloper/Private/RigVMModel/RigVMController.cpp
+++ b/Engine/Plugins/Runtime/RigVM/Source/RigVMDeveloper/Private/RigVMModel/RigVMController.cpp
@@ -12081,13 +12081,23 @@ bool URigVMController::RemoveArrayPin(const FString& InArrayElementPinPath, bool
 	{
 		GetActionStack()->AddAction(RemovePinAction);
 	}
+	
+	TFunction<void(URigVMPin*)> ClearPinPathCacheRecursively;
+	ClearPinPathCacheRecursively = [&](URigVMPin* InPin)
+		{
+			InPin->CachedPinPath.ResetCachedValue();
+			InPin->CachedPinPathWithNodePath.ResetCachedValue();
+			for (URigVMPin* SubPin : InPin->GetSubPins())
+			{
+				ClearPinPathCacheRecursively(SubPin);
+			}
+		};
 
 	for (int32 ExistingIndex = IndexToRemove; ExistingIndex < ArrayPin->GetArraySize(); ExistingIndex++)
 	{
 		URigVMPin* ExistingPin = ArrayPin->GetSubPins()[ExistingIndex];
 		ExistingPin->SetNameFromIndex();
-		ExistingPin->CachedPinPath.ResetCachedValue();
-		ExistingPin->CachedPinPathWithNodePath.ResetCachedValue();
+		ClearPinPathCacheRecursively(ExistingPin);
 		Notify(ERigVMGraphNotifType::PinRenamed, ExistingPin);
 	}
 
@@ -13916,9 +13926,19 @@ bool URigVMController::RenameExposedPin(const FName& InOldPinName, const FName&
 			{
 				return false;
 			}
+			
+			TFunction<void(URigVMPin*)> ClearPinPathCacheRecursively;
+			ClearPinPathCacheRecursively = [&](URigVMPin* InPin)
+				{
+					InPin->CachedPinPath.ResetCachedValue();
+					InPin->CachedPinPathWithNodePath.ResetCachedValue();
+					for (URigVMPin* SubPin : InPin->GetSubPins())
+					{
+						ClearPinPathCacheRecursively(SubPin);
+					}
+				};
 
-			InPin->CachedPinPath.ResetCachedValue();
-			InPin->CachedPinPathWithNodePath.ResetCachedValue();
+			ClearPinPathCacheRecursively(InPin);
 			InPin->SetDisplayName(InNewName);
 			InPin->IncrementVersion();
 			
