Control Rig: Only force compile public functions during cooking
#rnx

diff --git a/Engine/Plugins/Runtime/RigVM/Source/RigVMDeveloper/Private/RigVMCompiler/RigVMCompiler.cpp b/Engine/Plugins/Runtime/RigVM/Source/RigVMDeveloper/Private/RigVMCompiler/RigVMCompiler.cpp
--- a/Engine/Plugins/Runtime/RigVM/Source/RigVMDeveloper/Private/RigVMCompiler/RigVMCompiler.cpp
+++ b/Engine/Plugins/Runtime/RigVM/Source/RigVMDeveloper/Private/RigVMCompiler/RigVMCompiler.cpp
@@ -480,18 +480,22 @@ bool URigVMCompiler::Compile(const FRigVMCompileSettings& InSettings, TArray<URi
 	TMap<FString, const FRigVMFunctionCompilationData*> CurrentCompiledFunctions;
 
 #if WITH_EDITOR
-	if (!CurrentCompilationFunction)
+	// During cooking, compile all public functions
+	if (IsRunningCookCommandlet())
 	{
-		if (FunctionLibrary)
+		if (!CurrentCompilationFunction)
 		{
-			TArray<URigVMLibraryNode*> Functions = FunctionLibrary->GetFunctions();
-			for (URigVMLibraryNode* Function : Functions)
+			if (FunctionLibrary)
 			{
-				if (FunctionLibrary->IsFunctionPublic(Function->GetFName()))
+				TArray<URigVMLibraryNode*> Functions = FunctionLibrary->GetFunctions();
+				for (URigVMLibraryNode* Function : Functions)
 				{
-					if (FRigVMGraphFunctionData* FunctionData = Function->GetFunctionHeader().GetFunctionData())
+					if (FunctionLibrary->IsFunctionPublic(Function->GetFName()))
 					{
-						CompileFunction(WorkData.Settings, Function, InController, InExternalVariables, &FunctionData->CompilationData, OutVMContext);
+						if (FRigVMGraphFunctionData* FunctionData = Function->GetFunctionHeader().GetFunctionData())
+						{
+							CompileFunction(WorkData.Settings, Function, InController, InExternalVariables, &FunctionData->CompilationData, OutVMContext);
+						}
 					}
 				}
 			}
