Mark RegisterCompiledInStruct_NoLock as override.
Attempting to fix PVS-Studio warning about virtual method declared in final class.

diff --git a/Engine/Plugins/Runtime/RigVM/Source/RigVM/Public/RigVMCore/RigVMRegistry.h b/Engine/Plugins/Runtime/RigVM/Source/RigVM/Public/RigVMCore/RigVMRegistry.h
--- a/Engine/Plugins/Runtime/RigVM/Source/RigVM/Public/RigVMCore/RigVMRegistry.h
+++ b/Engine/Plugins/Runtime/RigVM/Source/RigVM/Public/RigVMCore/RigVMRegistry.h
@@ -1116,7 +1116,7 @@ protected:
 	}
 
 	// As above but registers a collection of functions for a struct
-	virtual void RegisterCompiledInStruct_NoLock(UScriptStruct* InStruct, TConstArrayView<FRigVMCompiledInFunction> InFunctions)
+	virtual void RegisterCompiledInStruct_NoLock(UScriptStruct* InStruct, TConstArrayView<FRigVMCompiledInFunction> InFunctions) override
 	{
 		Super::RegisterCompiledInStruct_NoLock(InStruct, InFunctions);
 	}
