Control Rig: Fix crash when duplicating array elements
#rnx

diff --git a/Engine/Plugins/Runtime/RigVM/Source/RigVMDeveloper/Private/RigVMModel/RigVMController.cpp b/Engine/Plugins/Runtime/RigVM/Source/RigVMDeveloper/Private/RigVMModel/RigVMController.cpp
--- a/Engine/Plugins/Runtime/RigVM/Source/RigVMDeveloper/Private/RigVMModel/RigVMController.cpp
+++ b/Engine/Plugins/Runtime/RigVM/Source/RigVMDeveloper/Private/RigVMModel/RigVMController.cpp
@@ -11881,8 +11881,16 @@ URigVMPin* URigVMController::InsertArrayPin(URigVMPin* ArrayPin, int32 InIndex,
 	{
 		URigVMPin* ExistingPin = ArrayPin->GetSubPins()[ExistingIndex];
 		RenameObject(ExistingPin, *FString::FormatAsNumber(ExistingIndex + 1));
-		ExistingPin->CachedPinPath.ResetCachedValue();
-		ExistingPin->CachedPinPathWithNodePath.ResetCachedValue();
+
+		TArray<URigVMPin*> PinsToClearCache;
+		PinsToClearCache.Add(ExistingPin);
+		for (int32 i=0; i<PinsToClearCache.Num(); ++i)
+		{
+			URigVMPin* PinToClearCache = PinsToClearCache[i];
+			PinToClearCache->CachedPinPath.ResetCachedValue();
+			PinToClearCache->CachedPinPathWithNodePath.ResetCachedValue();
+			PinsToClearCache.Append(PinToClearCache->GetSubPins());
+		}
 
 		// temporarily force a the index based display name
 		// so that the ed graph node can pick that up during handling the notify
