Control Rig: Fix cannot empty array pin
#rnx

diff --git a/Engine/Plugins/Runtime/RigVM/Source/RigVMDeveloper/Private/RigVMModel/RigVMController.cpp b/Engine/Plugins/Runtime/RigVM/Source/RigVMDeveloper/Private/RigVMModel/RigVMController.cpp
--- a/Engine/Plugins/Runtime/RigVM/Source/RigVMDeveloper/Private/RigVMModel/RigVMController.cpp
+++ b/Engine/Plugins/Runtime/RigVM/Source/RigVMDeveloper/Private/RigVMModel/RigVMController.cpp
@@ -12202,7 +12202,13 @@ bool URigVMController::SetArrayPinSize(const FString& InArrayPinPath, int32 InSi
 		GetActionStack()->BeginAction(Action);
 	}
 
-	InSize = FMath::Max<int32>(InSize, 0);
+	int32 MinSize = 0;
+	if(Pin->IsExecuteContext() || Pin->IsFixedSizeArray())
+	{
+		MinSize = 1;
+	}
+
+	InSize = FMath::Max<int32>(InSize, MinSize);
 	int32 AddedPins = 0;
 	int32 RemovedPins = 0;
 
@@ -13736,6 +13742,34 @@ bool URigVMController::RemoveExposedPin(const FName& InPinName, bool bSetupUndoR
 		return false;
 	}
 
+	if (URigVMAggregateNode* AggregateNode = Cast<URigVMAggregateNode>(LibraryNode))
+	{
+		if (AggregateNode->IsInputAggregate())
+		{
+			if (Pin->GetDirection() != ERigVMPinDirection::Input)
+			{
+				return false;
+			}
+
+			if (AggregateNode->GetAggregateInputs().Num() <= 2)
+			{
+				return false;
+			}
+		}
+		else
+		{
+			if (Pin->GetDirection() != ERigVMPinDirection::Output)
+			{
+				return false;
+			}
+
+			if (AggregateNode->GetAggregateOutputs().Num() <= 2)
+			{
+				return false;
+			}
+		}
+	}
+
 	if(bSetupUndoRedo)
 	{
 		if(RequestBulkEditDialogDelegate.IsBound())
diff --git a/Engine/Plugins/Runtime/RigVM/Source/RigVMDeveloper/Public/RigVMModel/RigVMGraph.h b/Engine/Plugins/Runtime/RigVM/Source/RigVMDeveloper/Public/RigVMModel/RigVMGraph.h
--- a/Engine/Plugins/Runtime/RigVM/Source/RigVMDeveloper/Public/RigVMModel/RigVMGraph.h
+++ b/Engine/Plugins/Runtime/RigVM/Source/RigVMDeveloper/Public/RigVMModel/RigVMGraph.h
@@ -253,6 +253,9 @@ private:
 	friend class URigVMNode;
 	friend struct FRigVMClient;
 	friend struct FRigVMControllerObjectFactory;
+	friend class FRigVMFunctionArgumentLayout;
+	friend class FRigVMFunctionArgumentReorderDragDropHandler;
+	friend class FRigVMGraphDetailCustomizationImpl;
 };
 
 #undef UE_API
diff --git a/Engine/Plugins/Runtime/RigVM/Source/RigVMEditor/Private/Editor/RigVMGraphDetailCustomization.cpp b/Engine/Plugins/Runtime/RigVM/Source/RigVMEditor/Private/Editor/RigVMGraphDetailCustomization.cpp
--- a/Engine/Plugins/Runtime/RigVM/Source/RigVMEditor/Private/Editor/RigVMGraphDetailCustomization.cpp
+++ b/Engine/Plugins/Runtime/RigVM/Source/RigVMEditor/Private/Editor/RigVMGraphDetailCustomization.cpp
@@ -157,6 +157,7 @@ bool FRigVMFunctionArgumentReorderDragDropHandler::AcceptDrop(const FDragDropEve
 		}
 	}
 
+	TGuardValue<bool> GuardEditGraph(LibraryNode->GetContainedGraph()->bEditable, true);
 	Controller->SetExposedPinIndex(DraggedPin->GetFName(), NewPinIndex);
 	return true;
 }
@@ -486,6 +487,7 @@ void FRigVMFunctionArgumentLayout::OnRemoveClicked()
 		{
 			if (URigVMController* Controller = RigVMClientHost->GetController(LibraryNode->GetContainedGraph()))
 			{
+				TGuardValue<bool> GuardEditGraph(LibraryNode->GetContainedGraph()->bEditable, true);
 				Controller->RemoveExposedPin(Pin->GetFName(), true, true);
 			}
 		}
@@ -557,6 +559,7 @@ void FRigVMFunctionArgumentLayout::OnArgNameTextCommitted(const FText& NewText,
 				if (URigVMController* Controller = RigVMClientHost->GetController(LibraryNode->GetContainedGraph()))
 				{
 					const FString& NewName = NewText.ToString();
+					TGuardValue<bool> GuardEditGraph(LibraryNode->GetContainedGraph()->bEditable, true);
 					Controller->RenameExposedPin(Pin->GetFName(), *NewName, true, true);
 				}
 			}
@@ -588,6 +591,7 @@ void FRigVMFunctionArgumentLayout::PinInfoChanged(const FEdGraphPinType& PinType
 				RigVMTypeUtils::CPPTypeFromPinType(PinType, CPPType, CPPTypeObjectName);
 				
 				bool bSetupUndoRedo = true;
+				TGuardValue<bool> GuardEditGraph(LibraryNode->GetContainedGraph()->bEditable, true);
 				Controller->ChangeExposedPinType(Pin->GetFName(), CPPType, CPPTypeObjectName, bSetupUndoRedo, false, true);
 
 				// If the controller has identified this as a bulk change, it has not added the actions to the action stack
@@ -1203,9 +1207,10 @@ FReply FRigVMGraphDetailCustomizationImpl::OnAddNewInputClicked()
 						DefaultValue = LastPin->GetDefaultValue();
 					}
 				}
-			}
 
-			Controller->AddExposedPin(ArgumentName, ERigVMPinDirection::Input, CPPType, CPPTypeObjectPath, DefaultValue, true, true);
+				TGuardValue<bool> GuardEditGraph(LibraryNode->GetContainedGraph()->bEditable, true);
+				Controller->AddExposedPin(ArgumentName, ERigVMPinDirection::Input, CPPType, CPPTypeObjectPath, DefaultValue, true, true);
+			}
 		}
 	}
 	return FReply::Unhandled();
@@ -1223,7 +1228,11 @@ FReply FRigVMGraphDetailCustomizationImpl::OnAddNewOutputClicked()
 			FString DefaultValue = TEXT("False");
 			// todo: base decisions on types on last argument
 
-			Controller->AddExposedPin(ArgumentName, ERigVMPinDirection::Output, CPPType, CPPTypeObjectPath, DefaultValue, true, true);
+			if (URigVMGraph* Model = WeakModel.Get())
+			{
+				TGuardValue<bool> GuardEditGraph(Model->bEditable, true);
+				Controller->AddExposedPin(ArgumentName, ERigVMPinDirection::Output, CPPType, CPPTypeObjectPath, DefaultValue, true, true);
+			}
 		}
 	}
 	return FReply::Unhandled();
