Allow for setting category name with different casing

diff --git a/Engine/Plugins/Runtime/RigVM/Source/RigVMDeveloper/Private/RigVMModel/RigVMController.cpp b/Engine/Plugins/Runtime/RigVM/Source/RigVMDeveloper/Private/RigVMModel/RigVMController.cpp
--- a/Engine/Plugins/Runtime/RigVM/Source/RigVMDeveloper/Private/RigVMModel/RigVMController.cpp
+++ b/Engine/Plugins/Runtime/RigVM/Source/RigVMDeveloper/Private/RigVMModel/RigVMController.cpp
@@ -8130,7 +8130,7 @@ bool URigVMController::SetNodeCategory(URigVMCollapseNode* InNode, const FString
 		return false;
 	}
 
-	if (InNode->GetNodeCategory() == InCategory)
+	if (InNode->GetNodeCategory().Equals(InCategory))
 	{
 		return false;
 	}
