Control Rig: Fix compilation issue when bAddBlueprintEditorDependency is set to false
#rnx

diff --git a/Engine/Plugins/Runtime/RigVM/Source/RigVMEditor/Private/Editor/Kismet/RigVMFindInBlueprints.cpp b/Engine/Plugins/Runtime/RigVM/Source/RigVMEditor/Private/Editor/Kismet/RigVMFindInBlueprints.cpp
--- a/Engine/Plugins/Runtime/RigVM/Source/RigVMEditor/Private/Editor/Kismet/RigVMFindInBlueprints.cpp
+++ b/Engine/Plugins/Runtime/RigVM/Source/RigVMEditor/Private/Editor/Kismet/RigVMFindInBlueprints.cpp
@@ -1236,12 +1236,19 @@ void SRigVMFindInBlueprints::MakeSearchQuery(FString InSearchString, bool bInIsF
 				StreamSearch.Reset();
 			}
 
-			UBlueprint* Blueprint =
+			UBlueprint* Blueprint = nullptr;
 #if WITH_EDITOR
-				EditorPtr.Pin()->GetRigVMBlueprint();
-#else
-				nullptr;
+			if (TSharedPtr<FRigVMEditorBase> StrongEditor = EditorPtr.Pin())
+			{
+				Blueprint = Cast<UBlueprint>(EditorPtr.Pin()->GetRigVMAssetInterface().GetObject());
+			}
 #endif
+			
+			if (!Blueprint)
+			{
+				return;
+			}
+			
 			FString ParentClass;
 			if (FProperty* ParentClassProp = Blueprint->GetClass()->FindPropertyByName(GET_MEMBER_NAME_CHECKED(UBlueprint, ParentClass)))
 			{
diff --git a/Engine/Plugins/Runtime/RigVM/Source/RigVMEditor/Private/Editor/RigVMFindReferences.cpp b/Engine/Plugins/Runtime/RigVM/Source/RigVMEditor/Private/Editor/RigVMFindReferences.cpp
--- a/Engine/Plugins/Runtime/RigVM/Source/RigVMEditor/Private/Editor/RigVMFindReferences.cpp
+++ b/Engine/Plugins/Runtime/RigVM/Source/RigVMEditor/Private/Editor/RigVMFindReferences.cpp
@@ -61,7 +61,7 @@ bool RigVMFindReferencesHelpers::ParsePinType(FText InKey, FText InValue, FEdGra
 	return bParsed;
 }
 
-void RigVMFindReferencesHelpers::ExpandAllChildren(FRigVMSearchResult InTreeNode, TSharedPtr<STreeView<FRigVMSearchResult>> InTreeView)
+void RigVMFindReferencesHelpers::ExpandAllChildren(FRigVMFindResultPtr InTreeNode, TSharedPtr<STreeView<FRigVMFindResultPtr>> InTreeView)
 {
 	if (InTreeNode->Children.Num())
 	{
@@ -547,7 +547,7 @@ void SRigVMFindReferences::Construct( const FArguments& InArgs, TSharedPtr<FRigV
 				.BorderImage(FAppStyle::Get().GetBrush("Brushes.Recessed"))
 				.Padding(FMargin(8.f, 8.f, 4.f, 0.f))
 				[
-					SAssignNew(TreeView, SRigVMTreeViewType)
+					SAssignNew(TreeView, SRigVMFindResultTreeViewType)
 						.TreeItemsSource( &ItemsFound )
 						.OnGenerateRow( this, &SRigVMFindReferences::OnGenerateRow )
 						.OnGetChildren( this, &SRigVMFindReferences::OnGetChildren )
@@ -625,14 +625,14 @@ void SRigVMFindReferences::FindReferences(const FString& SearchTerms)
 
 	auto FindOrAddGraphResult = [&](const UEdGraph* EdGraph)
 	{
-		FRigVMSearchResult GraphResult;
-		if (FRigVMSearchResult* Found = ElementHashToResult.Find(GetTypeHash(EdGraph)))
+		FRigVMFindResultPtr GraphResult;
+		if (FRigVMFindResultPtr* Found = ElementHashToResult.Find(GetTypeHash(EdGraph)))
 		{
 			GraphResult = *Found;
 		}
 		else
 		{
-			GraphResult = ItemsFound.Add_GetRef(FRigVMSearchResult(new FRigVMFindReferencesGraph(Blueprint, Schema->GetGraphType(EdGraph))));
+			GraphResult = ItemsFound.Add_GetRef(FRigVMFindResultPtr(new FRigVMFindReferencesGraph(Blueprint, Schema->GetGraphType(EdGraph))));
 			ElementHashToResult.Add(GetTypeHash(EdGraph), GraphResult);
 		}
 
@@ -645,15 +645,15 @@ void SRigVMFindReferences::FindReferences(const FString& SearchTerms)
 
 	auto FindOrAddNodeResult = [&](const UEdGraphNode* EdNode, const TArray<struct FSearchTagDataPair>& NodeMetaData)
 	{
-		FRigVMSearchResult NodeResult;
-		if (FRigVMSearchResult* Found = ElementHashToResult.Find(GetTypeHash(EdNode)))
+		FRigVMFindResultPtr NodeResult;
+		if (FRigVMFindResultPtr* Found = ElementHashToResult.Find(GetTypeHash(EdNode)))
 		{
 			NodeResult = *Found;
 		}
 		else
 		{
-			FRigVMSearchResult GraphResult = FindOrAddGraphResult(EdNode->GetGraph());
-			NodeResult = FRigVMSearchResult(new FRigVMFindReferencesGraphNode(Blueprint));
+			FRigVMFindResultPtr GraphResult = FindOrAddGraphResult(EdNode->GetGraph());
+			NodeResult = FRigVMFindResultPtr(new FRigVMFindReferencesGraphNode(Blueprint));
 			ElementHashToResult.Add(GetTypeHash(EdNode), NodeResult);
 			GraphResult->Children.Add(NodeResult);
 			NodeResult->Parent = GraphResult;
@@ -668,8 +668,8 @@ void SRigVMFindReferences::FindReferences(const FString& SearchTerms)
 
 	auto FindOrAddPinResult = [&](const UEdGraphPin* EdPin, const TArray<struct FSearchTagDataPair>& PinMetaData, const TArray<struct FSearchTagDataPair>& NodeMetaData)
 	{
-		FRigVMSearchResult PinResult;
-		if (FRigVMSearchResult* Found = ElementHashToResult.Find(GetTypeHash(EdPin)))
+		FRigVMFindResultPtr PinResult;
+		if (FRigVMFindResultPtr* Found = ElementHashToResult.Find(GetTypeHash(EdPin)))
 		{
 			PinResult = *Found;
 		}
@@ -678,8 +678,8 @@ void SRigVMFindReferences::FindReferences(const FString& SearchTerms)
 			UEdGraphNode* EdNode = EdPin->GetOwningNode();
 			
 			const FString SchemaClassName = Blueprint->GetRigVMClientHost()->GetRigVMEdGraphSchemaClass()->GetName();
-			FRigVMSearchResult NodeResult = FindOrAddNodeResult(EdPin->GetOwningNode(), NodeMetaData);
-			PinResult = FRigVMSearchResult(new FRigVMFindReferencesPin(Blueprint, SchemaClassName));
+			FRigVMFindResultPtr NodeResult = FindOrAddNodeResult(EdPin->GetOwningNode(), NodeMetaData);
+			PinResult = FRigVMFindResultPtr(new FRigVMFindReferencesPin(Blueprint, SchemaClassName));
 			ElementHashToResult.Add(GetTypeHash(EdPin), PinResult);
 			NodeResult->Children.Add(PinResult);
 			PinResult->Parent = NodeResult;
@@ -694,22 +694,22 @@ void SRigVMFindReferences::FindReferences(const FString& SearchTerms)
 
 	auto FindOrAddVariableResult = [&](const FName& VariableName, const UEdGraph* Graph, const TArray<UBlueprintExtension::FSearchTagDataPair>& VariableMetaData)
 	{
-		FRigVMSearchResult VariableResult;
+		FRigVMFindResultPtr VariableResult;
 
 		const uint32 Hash = HashCombine(GetTypeHash(VariableName), GetTypeHash(Graph));
-		if (FRigVMSearchResult* Found = ElementHashToResult.Find(Hash))
+		if (FRigVMFindResultPtr* Found = ElementHashToResult.Find(Hash))
 		{
 			VariableResult = *Found;
 		}
 		else
 		{
-			FRigVMSearchResult GraphResult;
+			FRigVMFindResultPtr GraphResult;
 			if (Graph)
 			{
 				GraphResult = FindOrAddGraphResult(Graph);
 			}
 
-			VariableResult = FRigVMSearchResult(new FRigVMFindReferencesVariable(Blueprint));
+			VariableResult = FRigVMFindResultPtr(new FRigVMFindReferencesVariable(Blueprint));
 			ElementHashToResult.Add(Hash, VariableResult);
 			if (GraphResult)
 			{
@@ -750,10 +750,10 @@ void SRigVMFindReferences::FindReferences(const FString& SearchTerms)
 				}
 				if (MatchedPair)
 				{
-					FRigVMSearchResult Result = FindOrAddNodeResult(EdNode, NodeMetadata);
+					FRigVMFindResultPtr Result = FindOrAddNodeResult(EdNode, NodeMetadata);
 
 					const FText DisplayText = FText::Format(LOCTEXT("RigVMFindReferencesValues", "{0}: {1}"), MatchedPair->Key, MatchedPair->Value);
-					FRigVMSearchResult DetailResult = FRigVMSearchResult(new FRigVMFindResult(Blueprint, DisplayText));
+					FRigVMFindResultPtr DetailResult = FRigVMFindResultPtr(new FRigVMFindResult(Blueprint, DisplayText));
 
 					Result->Children.Add(DetailResult);
 					DetailResult->Parent = Result;
@@ -776,10 +776,10 @@ void SRigVMFindReferences::FindReferences(const FString& SearchTerms)
 
 					if (MatchedPinPair)
 					{
-						FRigVMSearchResult Result = FindOrAddPinResult(Pin, PinMetaData, NodeMetadata);
+						FRigVMFindResultPtr Result = FindOrAddPinResult(Pin, PinMetaData, NodeMetadata);
 
 						const FText DisplayText = FText::Format(LOCTEXT("RigVMFindReferencesValues", "{0}: {1}"), MatchedPinPair->Key, MatchedPinPair->Value);
-						FRigVMSearchResult DetailResult = FRigVMSearchResult(new FRigVMFindResult(Blueprint, DisplayText));
+						FRigVMFindResultPtr DetailResult = FRigVMFindResultPtr(new FRigVMFindResult(Blueprint, DisplayText));
 
 						Result->Children.Add(DetailResult);
 						DetailResult->Parent = Result;
@@ -805,10 +805,10 @@ void SRigVMFindReferences::FindReferences(const FString& SearchTerms)
 
 				if (MatchedVariablePair)
 				{
-					FRigVMSearchResult Result = FindOrAddVariableResult(Variable.Name, EdGraph, LocalVariableMetadata);
+					FRigVMFindResultPtr Result = FindOrAddVariableResult(Variable.Name, EdGraph, LocalVariableMetadata);
 
 					const FText DisplayText = FText::Format(LOCTEXT("RigVMFindReferencesValues", "{0}: {1}"), MatchedVariablePair->Key, MatchedVariablePair->Value);
-					FRigVMSearchResult DetailResult = FRigVMSearchResult(new FRigVMFindResult(Blueprint, DisplayText));
+					FRigVMFindResultPtr DetailResult = FRigVMFindResultPtr(new FRigVMFindResult(Blueprint, DisplayText));
 
 					Result->Children.Add(DetailResult);
 					DetailResult->Parent = Result;
@@ -837,10 +837,10 @@ void SRigVMFindReferences::FindReferences(const FString& SearchTerms)
 
 			if (MatchedVariablePair)
 			{
-				FRigVMSearchResult Result = FindOrAddVariableResult(Variable.Name, nullptr, BPVariableMetadata);
+				FRigVMFindResultPtr Result = FindOrAddVariableResult(Variable.Name, nullptr, BPVariableMetadata);
 
 				const FText DisplayText = FText::Format(LOCTEXT("RigVMFindReferencesValues", "{0}: {1}"), MatchedVariablePair->Key, MatchedVariablePair->Value);
-				FRigVMSearchResult DetailResult = FRigVMSearchResult(new FRigVMFindResult(Blueprint, DisplayText));
+				FRigVMFindResultPtr DetailResult = FRigVMFindResultPtr(new FRigVMFindResult(Blueprint, DisplayText));
 
 				Result->Children.Add(DetailResult);
 				DetailResult->Parent = Result;
@@ -849,7 +849,7 @@ void SRigVMFindReferences::FindReferences(const FString& SearchTerms)
 	}
 
 	TreeView->RequestTreeRefresh();
-	for (FRigVMSearchResult Item : ItemsFound)
+	for (FRigVMFindResultPtr Item : ItemsFound)
 	{
 		RigVMFindReferencesHelpers::ExpandAllChildren(Item, TreeView);
 	}
@@ -868,7 +868,7 @@ void SRigVMFindReferences::OnSearchTextCommitted( const FText& Text, ETextCommit
 	}
 }
 
-TSharedRef<ITableRow> SRigVMFindReferences::OnGenerateRow( FRigVMSearchResult InItem, const TSharedRef<STableViewBase>& OwnerTable )
+TSharedRef<ITableRow> SRigVMFindReferences::OnGenerateRow( FRigVMFindResultPtr InItem, const TSharedRef<STableViewBase>& OwnerTable )
 {
 	// Finalize the search data, this does some non-thread safe actions that could not be done on the separate thread.
 	InItem->FinalizeSearchData();
@@ -902,12 +902,12 @@ TSharedRef<ITableRow> SRigVMFindReferences::OnGenerateRow( FRigVMSearchResult In
 		];
 }
 
-void SRigVMFindReferences::OnGetChildren( FRigVMSearchResult InItem, TArray< FRigVMSearchResult >& OutChildren )
+void SRigVMFindReferences::OnGetChildren( FRigVMFindResultPtr InItem, TArray< FRigVMFindResultPtr >& OutChildren )
 {
 	OutChildren += InItem->Children;
 }
 
-void SRigVMFindReferences::OnTreeSelectionDoubleClicked( FRigVMSearchResult Item )
+void SRigVMFindReferences::OnTreeSelectionDoubleClicked( FRigVMFindResultPtr Item )
 {
 	if(Item.IsValid())
 	{
@@ -930,7 +930,7 @@ TSharedPtr<SWidget> SRigVMFindReferences::OnContextMenuOpening()
 	return MenuBuilder.MakeWidget();
 }
 
-void SRigVMFindReferences::SelectAllItemsHelper(FRigVMSearchResult InItemToSelect)
+void SRigVMFindReferences::SelectAllItemsHelper(FRigVMFindResultPtr InItemToSelect)
 {
 	// Iterates over all children and recursively selects all items in the results
 	TreeView->SetItemSelection(InItemToSelect, true);
diff --git a/Engine/Plugins/Runtime/RigVM/Source/RigVMEditor/Public/Editor/RigVMFindReferences.h b/Engine/Plugins/Runtime/RigVM/Source/RigVMEditor/Public/Editor/RigVMFindReferences.h
--- a/Engine/Plugins/Runtime/RigVM/Source/RigVMEditor/Public/Editor/RigVMFindReferences.h
+++ b/Engine/Plugins/Runtime/RigVM/Source/RigVMEditor/Public/Editor/RigVMFindReferences.h
@@ -80,8 +80,8 @@ public:
 	FText DisplayText;
 };
 
-typedef TSharedPtr<FRigVMFindResult> FRigVMSearchResult;
-typedef STreeView<FRigVMSearchResult>  SRigVMTreeViewType;
+typedef TSharedPtr<FRigVMFindResult> FRigVMFindResultPtr;
+typedef STreeView<FRigVMFindResultPtr>  SRigVMFindResultTreeViewType;
 
 /** Some utility functions to help with Find-in-Blueprint functionality */
 namespace RigVMFindReferencesHelpers
@@ -108,7 +108,7 @@ namespace RigVMFindReferencesHelpers
 	/**
 	* Iterates through all the given tree node's children and tells the tree view to expand them
 	*/
-	void ExpandAllChildren(FRigVMSearchResult InTreeNode, TSharedPtr<STreeView<FRigVMSearchResult>> InTreeView);
+	void ExpandAllChildren(FRigVMFindResultPtr InTreeNode, TSharedPtr<STreeView<FRigVMFindResultPtr>> InTreeView);
 }
 
 /** Graph nodes use this class to store their data */
@@ -249,19 +249,19 @@ private:
 	UE_API void OnSearchTextCommitted(const FText& Text, ETextCommit::Type CommitType);
 
 	/* Get the children of a row */
-	UE_API void OnGetChildren( FRigVMSearchResult InItem, TArray< FRigVMSearchResult >& OutChildren );
+	UE_API void OnGetChildren( FRigVMFindResultPtr InItem, TArray< FRigVMFindResultPtr >& OutChildren );
 
 	/* Called when user double clicks on a new result */
-	UE_API void OnTreeSelectionDoubleClicked( FRigVMSearchResult Item );
+	UE_API void OnTreeSelectionDoubleClicked( FRigVMFindResultPtr Item );
 
 	/* Called when a new row is being generated */
-	UE_API TSharedRef<ITableRow> OnGenerateRow(FRigVMSearchResult InItem, const TSharedRef<STableViewBase>& OwnerTable);
+	UE_API TSharedRef<ITableRow> OnGenerateRow(FRigVMFindResultPtr InItem, const TSharedRef<STableViewBase>& OwnerTable);
 
 	/** Callback to build the context menu when right clicking in the tree */
 	UE_API TSharedPtr<SWidget> OnContextMenuOpening();
 
 	/** Helper function to select all items */
-	UE_API void SelectAllItemsHelper(FRigVMSearchResult InItemToSelect);
+	UE_API void SelectAllItemsHelper(FRigVMFindResultPtr InItemToSelect);
 
 	/** Callback when user attempts to select all items in the search results */
 	UE_API void OnSelectAllAction();
@@ -274,16 +274,16 @@ private:
 	TWeakPtr<class FRigVMEditorBase> EditorPtr;
 	
 	/* The tree view displays the results */
-	TSharedPtr<SRigVMTreeViewType> TreeView;
+	TSharedPtr<SRigVMFindResultTreeViewType> TreeView;
 
 	/** The search text box */
 	TSharedPtr<class SSearchBox> SearchTextField;
 	
 	/* This buffer stores the currently displayed results */
-	TArray<FRigVMSearchResult> ItemsFound;
+	TArray<FRigVMFindResultPtr> ItemsFound;
 
 	/* Map relationship between element hash and its result  */ 
-	TMap<uint32, FRigVMSearchResult> ElementHashToResult;
+	TMap<uint32, FRigVMFindResultPtr> ElementHashToResult;
 
 	/* The string to highlight in the results */
 	FText HighlightText;
