RigVMTypeUtils move a bunch of stuff from inline to .cpp

Also, change constants to be extern instead of defined in header

diff --git a/Engine/Plugins/Runtime/RigVM/Source/RigVM/Private/RigVMTypeUtils.cpp b/Engine/Plugins/Runtime/RigVM/Source/RigVM/Private/RigVMTypeUtils.cpp
--- a/Engine/Plugins/Runtime/RigVM/Source/RigVM/Private/RigVMTypeUtils.cpp
+++ b/Engine/Plugins/Runtime/RigVM/Source/RigVM/Private/RigVMTypeUtils.cpp
@@ -13,6 +13,49 @@
 
 namespace RigVMTypeUtils
 {
+	const TCHAR BoolType[] = TEXT("bool");
+	const TCHAR FloatType[] = TEXT("float");
+	const TCHAR DoubleType[] = TEXT("double");
+	const TCHAR IntType[] = TEXT("int");
+	const TCHAR Int32Type[] = TEXT("int32");
+	const TCHAR Int64Type[] = TEXT("int64");
+	const TCHAR UInt8Type[] = TEXT("uint8");
+	const TCHAR UInt32Type[] = TEXT("uint32");
+	const TCHAR UInt64Type[] = TEXT("uint64");
+	const TCHAR FNameType[] = TEXT("FName");
+	const TCHAR FStringType[] = TEXT("FString");
+	const TCHAR FTextType[] = TEXT("FText");
+	const TCHAR BoolArrayType[] = TEXT("TArray<bool>");
+	const TCHAR FloatArrayType[] = TEXT("TArray<float>");
+	const TCHAR DoubleArrayType[] = TEXT("TArray<double>");
+	const TCHAR Int32ArrayType[] = TEXT("TArray<int32>");
+	const TCHAR UInt32ArrayType[] = TEXT("TArray<uint32>");
+	const TCHAR UInt8ArrayType[] = TEXT("TArray<uint8>");
+	const TCHAR FNameArrayType[] = TEXT("TArray<FName>");
+	const TCHAR FStringArrayType[] = TEXT("TArray<FString>");
+	const TCHAR FTextArrayType[] = TEXT("TArray<FText>");
+
+	const FLazyName BoolTypeName(BoolType);
+	const FLazyName FloatTypeName(FloatType);
+	const FLazyName DoubleTypeName(DoubleType);
+	const FLazyName IntTypeName(IntType);
+	const FLazyName Int32TypeName(Int32Type);
+	const FLazyName Int64TypeName(Int64Type);
+	const FLazyName UInt8TypeName(UInt8Type);
+	const FLazyName UInt32TypeName(UInt32Type);
+	const FLazyName UInt64TypeName(UInt64Type);
+	const FLazyName FNameTypeName(FNameType);
+	const FLazyName FStringTypeName(FStringType);
+	const FLazyName FTextTypeName(FTextType);
+	const FLazyName BoolArrayTypeName(BoolArrayType);
+	const FLazyName FloatArrayTypeName(FloatArrayType);
+	const FLazyName DoubleArrayTypeName(DoubleArrayType);
+	const FLazyName Int32ArrayTypeName(Int32ArrayType);
+	const FLazyName UInt32ArrayTypeName(UInt32ArrayType);
+	const FLazyName UInt8ArrayTypeName(UInt8ArrayType);
+	const FLazyName FNameArrayTypeName(FNameArrayType);
+	const FLazyName FStringArrayTypeName(FStringArrayType);
+
 	TRigVMTypeIndex TypeIndex::Execute = INDEX_NONE;	
 	TRigVMTypeIndex TypeIndex::ExecuteArray = INDEX_NONE;	
 	TRigVMTypeIndex TypeIndex::Bool = INDEX_NONE;	
@@ -35,6 +78,160 @@ namespace RigVMTypeUtils
 	TRigVMTypeIndex TypeIndex::WildCardArray = INDEX_NONE;	
 }
 
+bool RigVMTypeUtils::RequiresCPPTypeObject(const FString& InCPPType)
+{
+	static const TCHAR* PrefixesRequiringCPPTypeObject[] = {
+		TEXT("F"), 
+		TEXT("TArray<F"), 
+		TEXT("TArray<TArray<F"), 
+		TEXT("E"), 
+		TEXT("TArray<E"), 
+		TEXT("TArray<TArray<E"), 
+		TEXT("U"), 
+		TEXT("TArray<U"),
+		TEXT("TArray<TArray<U"),
+		TEXT("TObjectPtr<"), 
+		TEXT("TArray<TObjectPtr<"),
+		TEXT("TArray<TArray<TObjectPtr<"),
+		TEXT("TScriptInterface<")
+		TEXT("TSubclassOf<")
+	};
+	static const TCHAR* PrefixesNotRequiringCPPTypeObject[] = {
+		TEXT("UInt"), 
+		TEXT("TArray<UInt"),
+		TEXT("TArray<TArray<UInt"),
+		TEXT("uint"), 
+		TEXT("TArray<uint"),
+		TEXT("TArray<TArray<uint")
+	};
+	static const TCHAR* CPPTypesNotRequiringCPPTypeObject[] = {
+		TEXT("FString"), 
+		TEXT("TArray<FString>"), 
+		TEXT("TArray<TArray<FString>>"), 
+		TEXT("FName"), 
+		TEXT("TArray<FName>"),
+		TEXT("TArray<TArray<FName>>") 
+	};
+
+	for(const TCHAR* Type : CPPTypesNotRequiringCPPTypeObject)
+	{
+		if(InCPPType == Type)
+		{
+			return false;
+		}
+	}
+
+	for(const TCHAR* Prefix : PrefixesNotRequiringCPPTypeObject)
+	{
+		if(InCPPType.StartsWith(Prefix, ESearchCase::CaseSensitive))
+		{
+			return false;
+		}
+	}
+
+	for(const TCHAR* Prefix : PrefixesRequiringCPPTypeObject)
+	{
+		if(InCPPType.StartsWith(Prefix, ESearchCase::CaseSensitive))
+		{
+			return true;
+		}
+	}
+
+	return false;
+}
+
+UObject* RigVMTypeUtils::FindObjectGlobally(const TCHAR* InObjectName, bool bUseRedirector)
+{
+	// Do a global search for the CPP type. Note that searching with ANY_PACKAGE _does not_
+	// apply redirectors. So only if this fails do we apply them manually below.
+	UObject* Object = FindFirstObject<UField>(InObjectName, EFindFirstObjectOptions::NativeFirst);
+	if(Object != nullptr)
+	{
+		return Object;
+	}
+
+	// If its an enum, it might be defined as a namespace with the actual enum inside (see ERigVMClampSpatialMode). 
+	FString ObjectNameStr(InObjectName);
+	if (!ObjectNameStr.IsEmpty() && ObjectNameStr[0] == 'E')
+	{
+		FString Left, Right;
+		if (ObjectNameStr.Split(TEXT("::"), &Left, &Right))
+		{
+			ObjectNameStr = Left;
+			Object = FindFirstObject<UField>(*ObjectNameStr, EFindFirstObjectOptions::NativeFirst);
+			if(Object != nullptr)
+			{
+				return Object;
+			}
+		}
+	}
+
+	if (!bUseRedirector)
+	{
+		return nullptr;
+	}
+
+	FCoreRedirectObjectName OldObjectName (ObjectNameStr);
+	FCoreRedirectObjectName NewObjectName;
+	const bool bFoundRedirect = FCoreRedirects::RedirectNameAndValues(
+		ECoreRedirectFlags::Type_Class | ECoreRedirectFlags::Type_Struct | ECoreRedirectFlags::Type_Enum,
+		OldObjectName,
+		NewObjectName,
+		nullptr,
+		ECoreRedirectMatchFlags::AllowPartialMatch); // AllowPartialMatch to allow redirects from one package to another
+
+	if (!bFoundRedirect)
+	{
+		return nullptr;
+	}
+
+	const FString RedirectedObjectName = NewObjectName.ObjectName.ToString();
+	UPackage *Package = nullptr;
+	if (!NewObjectName.PackageName.IsNone())
+	{
+		Package = FindPackage(nullptr, *NewObjectName.PackageName.ToString());
+	}
+	if (Package != nullptr)
+	{
+		Object = FindObject<UField>(Package, *RedirectedObjectName);
+	}
+	if (Package == nullptr || Object == nullptr)
+	{
+		// Hail Mary pass.
+		Object = FindFirstObject<UField>(*RedirectedObjectName, EFindFirstObjectOptions::NativeFirst);
+	}
+	return Object;
+}
+
+FString RigVMTypeUtils::CPPTypeFromObject(const UObject* InCPPTypeObject, EClassArgType InClassArgType)
+{
+	if (const UClass* Class = Cast<UClass>(InCPPTypeObject))
+	{
+		if (InClassArgType == EClassArgType::AsClass)
+		{
+			return FString::Printf(RigVMTypeUtils::TSubclassOfTemplate, Class->GetPrefixCPP(), *Class->GetName());
+		}
+		else if (Class->IsChildOf(UInterface::StaticClass()))
+		{
+			return FString::Printf(RigVMTypeUtils::TScriptInterfaceTemplate, TEXT("I"), *Class->GetName());
+		}
+		else
+		{
+			return FString::Printf(RigVMTypeUtils::TObjectPtrTemplate, Class->GetPrefixCPP(), *Class->GetName());
+		}
+	}
+	else if (const UScriptStruct* ScriptStruct = Cast<UScriptStruct>(InCPPTypeObject))
+	{
+		return GetUniqueStructTypeName(ScriptStruct);
+	}
+	else if (const UEnum* Enum = Cast<UEnum>(InCPPTypeObject))
+	{
+		return RigVMTypeUtils::CPPTypeFromEnum(Enum);
+	}
+
+	return FString();
+}
+
 UObject* RigVMTypeUtils::UserDefinedTypeFromCPPType(FString& InOutCPPType, const FRigVMUserDefinedTypeResolver* InTypeResolver)
 {
 	const FString OriginalTypeName = InOutCPPType;
@@ -166,3 +363,213 @@ UObject* RigVMTypeUtils::UserDefinedTypeFromCPPType(FString& InOutCPPType, const
 
 	return CPPTypeObject;
 }
+
+UObject* RigVMTypeUtils::ObjectFromCPPType(FString& InOutCPPType, bool bUseRedirector, const FRigVMUserDefinedTypeResolver* InTypeResolver)
+{
+	if (!RequiresCPPTypeObject(InOutCPPType))
+	{
+		return nullptr;
+	}
+
+	// try to find the CPPTypeObject by name
+	FString BaseCPPType = InOutCPPType;
+	while (IsArrayType(BaseCPPType))
+	{
+		BaseCPPType = BaseTypeFromArrayType(BaseCPPType);
+	}
+	FString CPPType = BaseCPPType;
+	const bool bIsClass = CPPType.StartsWith(TSubclassOfPrefix);
+
+	static const FString PrefixObjectPtr = TObjectPtrPrefix;
+	static const FString PrefixSubclassOf = TSubclassOfPrefix;
+	static const FString PrefixScriptInterface = TScriptInterfacePrefix;
+
+	if (CPPType.StartsWith(PrefixObjectPtr))
+	{
+		// Chop the prefix + the U indicating object class
+		CPPType = CPPType.Mid(PrefixObjectPtr.Len() + 1, CPPType.Len() - (PrefixObjectPtr.Len() + 2));
+	}
+	else if (CPPType.StartsWith(PrefixSubclassOf))
+	{
+		// Chop the prefix + the U indicating object class
+		CPPType = CPPType.Mid(PrefixSubclassOf.Len() + 1, CPPType.Len() - (PrefixSubclassOf.Len() + 2));
+	}
+	else if (CPPType.StartsWith(TScriptInterfacePrefix))
+	{
+		// Chop the prefix + the I indicating interface class
+		CPPType = CPPType.Mid(PrefixScriptInterface.Len() + 1, CPPType.Len() - (PrefixScriptInterface.Len() + 2));
+	}
+
+	UObject* CPPTypeObject = FindObjectGlobally(*CPPType, bUseRedirector);
+	if (CPPTypeObject == nullptr)
+	{
+		// If we've mistakenly stored the struct type with the 'F', 'U', or 'A' prefixes, we need to strip them
+		// off first. Enums are always named with their prefix intact.
+		if (!CPPType.IsEmpty() && (CPPType[0] == TEXT('F') || CPPType[0] == TEXT('U') || CPPType[0] == TEXT('A')))
+		{
+			CPPType = CPPType.Mid(1);
+		}
+		CPPTypeObject = RigVMTypeUtils::FindObjectGlobally(*CPPType, bUseRedirector);
+	}
+
+	if(CPPTypeObject == nullptr)
+	{
+		CPPType = BaseCPPType;
+		CPPTypeObject = UserDefinedTypeFromCPPType(CPPType, InTypeResolver);
+	}
+
+	if(CPPTypeObject == nullptr)
+	{
+		InOutCPPType.Reset();
+		return nullptr;
+	}
+
+	CPPType = CPPTypeFromObject(CPPTypeObject, bIsClass ? EClassArgType::AsClass : EClassArgType::AsObject);
+	InOutCPPType.ReplaceInline(*BaseCPPType, *CPPType);
+	return CPPTypeObject;
+}
+
+FString RigVMTypeUtils::PostProcessCPPType(const FString& InCPPType, UObject* InCPPTypeObject, const FRigVMUserDefinedTypeResolver* InResolvalInfo)
+{
+	FString CPPType = InCPPType;
+	CleanupCPPType(CPPType);
+	
+	if (InCPPTypeObject)
+	{
+		const bool bIsClass = CPPType.StartsWith(TSubclassOfPrefix);
+		CPPType = CPPTypeFromObject(InCPPTypeObject, bIsClass ? EClassArgType::AsClass : EClassArgType::AsObject);
+		if(CPPType != InCPPType)
+		{
+			FString TemplateType = InCPPType;
+			while (RigVMTypeUtils::IsArrayType(TemplateType))
+			{
+				CPPType = RigVMTypeUtils::ArrayTypeFromBaseType(CPPType);
+				TemplateType = RigVMTypeUtils::BaseTypeFromArrayType(TemplateType);
+			}		
+		}
+	}
+	else if (RequiresCPPTypeObject(CPPType))
+	{
+		// Uses redirectors and updates the CPPType if necessary
+		ObjectFromCPPType(CPPType, true, InResolvalInfo);
+	}
+
+	return CPPType;
+}
+
+bool RigVMTypeUtils::FixCPPTypeAndObject(FString& InOutCPPType, TObjectPtr<UObject>& InOutCPPTypeObject, const FRigVMUserDefinedTypeResolver* InResolvalInfo)
+{
+	const FString OldCPPType = InOutCPPType;
+	UObject* OldCPPTypeObject = InOutCPPTypeObject;
+	InOutCPPType = PostProcessCPPType(InOutCPPType, InOutCPPTypeObject, InResolvalInfo);
+
+	// Type object might have changed through a redirect
+	bool bFindObjectAgain = true;
+
+	// If its a user defined enum, do not attempt to find the object from the cpp type
+	if (const UEnum* Enum = Cast<UEnum>(InOutCPPTypeObject))
+	{
+		if (Enum->CppType.IsEmpty())
+		{
+			bFindObjectAgain = false;
+		}
+	}
+
+	if (bFindObjectAgain)
+	{
+		InOutCPPTypeObject = ObjectFromCPPType(InOutCPPType);
+	}
+	
+	return !InOutCPPType.Equals(OldCPPType) || InOutCPPTypeObject != OldCPPTypeObject;
+}
+
+UObject* RigVMTypeUtils::FindObjectFromCPPTypeObjectPath(const FString& InObjectPath)
+{
+	UObject* Result = nullptr;
+	if (InObjectPath.IsEmpty())
+	{
+		return Result;
+	}
+
+	if (InObjectPath == FName(NAME_None).ToString())
+	{
+		return Result;
+	}
+
+	// we do this to avoid ambiguous searches for 
+	// common names such as "transform" or "vector"
+	UPackage* Package = nullptr;
+	FString PackageName;
+	FString CPPTypeObjectName = InObjectPath;
+	if (InObjectPath.Split(TEXT("."), &PackageName, &CPPTypeObjectName))
+	{
+		Package = FindPackage(nullptr, *PackageName);
+	}
+
+	if (UObject* ObjectWithinPackage = FindObject<UObject>(Package, *CPPTypeObjectName))
+	{
+		Result = ObjectWithinPackage;
+	}
+
+	if (!Result)
+	{
+		Result = FindFirstObject<UObject>(*InObjectPath, EFindFirstObjectOptions::NativeFirst | EFindFirstObjectOptions::EnsureIfAmbiguous);
+	}
+	if (!Result)
+	{
+		const FCoreRedirectObjectName OldObjectName(InObjectPath);
+		const FCoreRedirectObjectName NewObjectName = FCoreRedirects::GetRedirectedName(ECoreRedirectFlags::Type_Struct, OldObjectName);
+		if (OldObjectName != NewObjectName)
+		{
+			Result = FindObjectFromCPPTypeObjectPath(NewObjectName.ToString());
+		}
+	}
+	return Result;
+}
+
+bool RigVMTypeUtils::AreCompatible(const FProperty* InSourceProperty, const FProperty* InTargetProperty)
+{
+	bool bCompatible = InSourceProperty->SameType(InTargetProperty);
+	if (!bCompatible)
+	{
+		if(const FFloatProperty* TargetFloatProperty = CastField<FFloatProperty>(InTargetProperty))
+		{
+			bCompatible = InSourceProperty->IsA<FDoubleProperty>();
+		}
+		else if(const FDoubleProperty* TargetDoubleProperty = CastField<FDoubleProperty>(InTargetProperty))
+		{
+			bCompatible = InSourceProperty->IsA<FFloatProperty>();
+		}
+		else if (const FByteProperty* TargetByteProperty = CastField<FByteProperty>(InTargetProperty))
+		{
+			bCompatible = InSourceProperty->IsA<FEnumProperty>();
+		}
+		else if (const FEnumProperty* TargetEnumProperty = CastField<FEnumProperty>(InTargetProperty))
+		{
+			bCompatible = InSourceProperty->IsA<FByteProperty>();
+		}
+		else if(const FArrayProperty* TargetArrayProperty = CastField<FArrayProperty>(InTargetProperty))
+		{
+			if(const FArrayProperty* SourceArrayProperty = CastField<FArrayProperty>(InSourceProperty))
+			{
+				if(TargetArrayProperty->Inner->IsA<FFloatProperty>())
+				{
+					bCompatible = SourceArrayProperty->Inner->IsA<FDoubleProperty>();
+				}
+				else if(TargetArrayProperty->Inner->IsA<FDoubleProperty>())
+				{
+					bCompatible = SourceArrayProperty->Inner->IsA<FFloatProperty>();
+				}
+				else if(FByteProperty* TargetArrayInnerByteProperty = CastField<FByteProperty>(TargetArrayProperty->Inner))
+				{
+					bCompatible = SourceArrayProperty->Inner->IsA<FEnumProperty>();
+				}
+				else if(FEnumProperty* TargetArrayInnerEnumProperty = CastField<FEnumProperty>(TargetArrayProperty->Inner))
+				{
+					bCompatible = SourceArrayProperty->Inner->IsA<FByteProperty>();
+				}
+			}
+		}
+	}
+	return bCompatible;
+}
diff --git a/Engine/Plugins/Runtime/RigVM/Source/RigVM/Public/RigVMTypeUtils.h b/Engine/Plugins/Runtime/RigVM/Source/RigVM/Public/RigVMTypeUtils.h
--- a/Engine/Plugins/Runtime/RigVM/Source/RigVM/Public/RigVMTypeUtils.h
+++ b/Engine/Plugins/Runtime/RigVM/Source/RigVM/Public/RigVMTypeUtils.h
@@ -58,48 +58,48 @@ namespace RigVMTypeUtils
 	constexpr TCHAR TSubclassOfTemplate[] = TEXT("TSubclassOf<%s%s>");
 	constexpr TCHAR TScriptInterfaceTemplate[] = TEXT("TScriptInterface<%s%s>");
 
-	const inline TCHAR* BoolType = TEXT("bool");
-	const inline TCHAR* FloatType = TEXT("float");
-	const inline TCHAR* DoubleType = TEXT("double");
-	const inline TCHAR* IntType = TEXT("int");
-	const inline TCHAR* Int32Type = TEXT("int32");
-	const inline TCHAR* Int64Type = TEXT("int64");
-	const inline TCHAR* UInt8Type = TEXT("uint8");
-	const inline TCHAR* UInt32Type = TEXT("uint32");
-	const inline TCHAR* UInt64Type = TEXT("uint64");
-	const inline TCHAR* FNameType = TEXT("FName");
-	const inline TCHAR* FStringType = TEXT("FString");
-	const inline TCHAR* FTextType = TEXT("FText");
-	const inline TCHAR* BoolArrayType = TEXT("TArray<bool>");
-	const inline TCHAR* FloatArrayType = TEXT("TArray<float>");
-	const inline TCHAR* DoubleArrayType = TEXT("TArray<double>");
-	const inline TCHAR* Int32ArrayType = TEXT("TArray<int32>");
-	const inline TCHAR* UInt32ArrayType = TEXT("TArray<uint32>");
-	const inline TCHAR* UInt8ArrayType = TEXT("TArray<uint8>");
-	const inline TCHAR* FNameArrayType = TEXT("TArray<FName>");
-	const inline TCHAR* FStringArrayType = TEXT("TArray<FString>");
-	const inline TCHAR* FTextArrayType = TEXT("TArray<FText>");
-
-	const FLazyName BoolTypeName(BoolType);
-	const FLazyName FloatTypeName(FloatType);
-	const FLazyName DoubleTypeName(DoubleType);
-	const FLazyName IntTypeName(IntType);
-	const FLazyName Int32TypeName(Int32Type);
-	const FLazyName Int64TypeName(Int64Type);
-	const FLazyName UInt8TypeName(UInt8Type);
-	const FLazyName UInt32TypeName(UInt32Type);
-	const FLazyName UInt64TypeName(UInt64Type);
-	const FLazyName FNameTypeName(FNameType);
-	const FLazyName FStringTypeName(FStringType);
-	const FLazyName FTextTypeName(FTextType);
-	const FLazyName BoolArrayTypeName(BoolArrayType);
-	const FLazyName FloatArrayTypeName(FloatArrayType);
-	const FLazyName DoubleArrayTypeName(DoubleArrayType);
-	const FLazyName Int32ArrayTypeName(Int32ArrayType);
-	const FLazyName UInt32ArrayTypeName(UInt32ArrayType);
-	const FLazyName UInt8ArrayTypeName(UInt8ArrayType);
-	const FLazyName FNameArrayTypeName(FNameArrayType);
-	const FLazyName FStringArrayTypeName(FStringArrayType);
+	extern UE_API const TCHAR BoolType[];
+	extern UE_API const TCHAR FloatType[];
+	extern UE_API const TCHAR DoubleType[];
+	extern UE_API const TCHAR IntType[];
+	extern UE_API const TCHAR Int32Type[];
+	extern UE_API const TCHAR Int64Type[];
+	extern UE_API const TCHAR UInt8Type[];
+	extern UE_API const TCHAR UInt32Type[];
+	extern UE_API const TCHAR UInt64Type[];
+	extern UE_API const TCHAR FNameType[];
+	extern UE_API const TCHAR FStringType[];
+	extern UE_API const TCHAR FTextType[];
+	extern UE_API const TCHAR BoolArrayType[];
+	extern UE_API const TCHAR FloatArrayType[];
+	extern UE_API const TCHAR DoubleArrayType[];
+	extern UE_API const TCHAR Int32ArrayType[];
+	extern UE_API const TCHAR UInt32ArrayType[];
+	extern UE_API const TCHAR UInt8ArrayType[];
+	extern UE_API const TCHAR FNameArrayType[];
+	extern UE_API const TCHAR FStringArrayType[];
+	extern UE_API const TCHAR FTextArrayType[];
+
+	extern UE_API const FLazyName BoolTypeName;
+	extern UE_API const FLazyName FloatTypeName;
+	extern UE_API const FLazyName DoubleTypeName;
+	extern UE_API const FLazyName IntTypeName;
+	extern UE_API const FLazyName Int32TypeName;
+	extern UE_API const FLazyName Int64TypeName;
+	extern UE_API const FLazyName UInt8TypeName;
+	extern UE_API const FLazyName UInt32TypeName;
+	extern UE_API const FLazyName UInt64TypeName;
+	extern UE_API const FLazyName FNameTypeName;
+	extern UE_API const FLazyName FStringTypeName;
+	extern UE_API const FLazyName FTextTypeName;
+	extern UE_API const FLazyName BoolArrayTypeName;
+	extern UE_API const FLazyName FloatArrayTypeName;
+	extern UE_API const FLazyName DoubleArrayTypeName;
+	extern UE_API const FLazyName Int32ArrayTypeName;
+	extern UE_API const FLazyName UInt32ArrayTypeName;
+	extern UE_API const FLazyName UInt8ArrayTypeName;
+	extern UE_API const FLazyName FNameArrayTypeName;
+	extern UE_API const FLazyName FStringArrayTypeName;
 
 	class TypeIndex
 	{
@@ -213,124 +213,9 @@ namespace RigVMTypeUtils
 		return WildCardArrayCPPTypeName;
 	}
 
-	static bool RequiresCPPTypeObject(const FString& InCPPType)
-	{
-		static const TArray<FString> PrefixesRequiringCPPTypeObject = {
-			TEXT("F"), 
-			TEXT("TArray<F"), 
-			TEXT("TArray<TArray<F"), 
-			TEXT("E"), 
-			TEXT("TArray<E"), 
-			TEXT("TArray<TArray<E"), 
-			TEXT("U"), 
-			TEXT("TArray<U"),
-			TEXT("TArray<TArray<U"),
-			TEXT("TObjectPtr<"), 
-			TEXT("TArray<TObjectPtr<"),
-			TEXT("TArray<TArray<TObjectPtr<"),
-			TEXT("TScriptInterface<")
-			TEXT("TSubclassOf<")
-		};
-		static const TArray<FString> PrefixesNotRequiringCPPTypeObject = {
-			TEXT("UInt"), 
-			TEXT("TArray<UInt"),
-			TEXT("TArray<TArray<UInt"),
-			TEXT("uint"), 
-			TEXT("TArray<uint"),
-			TEXT("TArray<TArray<uint")
-		};
-		static const TArray<FString> CPPTypesNotRequiringCPPTypeObject = {
-			TEXT("FString"), 
-			TEXT("TArray<FString>"), 
-			TEXT("TArray<TArray<FString>>"), 
-			TEXT("FName"), 
-			TEXT("TArray<FName>"),
-			TEXT("TArray<TArray<FName>>") 
-		};
-
-		if(!CPPTypesNotRequiringCPPTypeObject.Contains(InCPPType))
-		{
-			for(const FString& Prefix : PrefixesNotRequiringCPPTypeObject)
-			{
-				if(InCPPType.StartsWith(Prefix, ESearchCase::CaseSensitive))
-				{
-					return false;
-				}
-			}
-			for(const FString& Prefix : PrefixesRequiringCPPTypeObject)
-			{
-				if(InCPPType.StartsWith(Prefix, ESearchCase::CaseSensitive))
-				{
-					return true;
-				}
-			}
-		}
-
-		return false;
-	}
-
-	static UObject* FindObjectGlobally(const TCHAR* InObjectName, bool bUseRedirector)
-	{
-		// Do a global search for the CPP type. Note that searching with ANY_PACKAGE _does not_
-		// apply redirectors. So only if this fails do we apply them manually below.
-		UObject* Object = FindFirstObject<UField>(InObjectName, EFindFirstObjectOptions::NativeFirst);
-		if(Object != nullptr)
-		{
-			return Object;
-		}
-
-		// If its an enum, it might be defined as a namespace with the actual enum inside (see ERigVMClampSpatialMode). 
-		FString ObjectNameStr(InObjectName);
-		if (!ObjectNameStr.IsEmpty() && ObjectNameStr[0] == 'E')
-		{
-			FString Left, Right;
-			if (ObjectNameStr.Split(TEXT("::"), &Left, &Right))
-			{
-				ObjectNameStr = Left;
-				Object = FindFirstObject<UField>(*ObjectNameStr, EFindFirstObjectOptions::NativeFirst);
-				if(Object != nullptr)
-				{
-					return Object;
-				}
-			}
-		}
-
-		if (!bUseRedirector)
-		{
-			return nullptr;
-		}
+	RIGVM_API bool RequiresCPPTypeObject(const FString& InCPPType);
 
-		FCoreRedirectObjectName OldObjectName (ObjectNameStr);
-		FCoreRedirectObjectName NewObjectName;
-		const bool bFoundRedirect = FCoreRedirects::RedirectNameAndValues(
-			ECoreRedirectFlags::Type_Class | ECoreRedirectFlags::Type_Struct | ECoreRedirectFlags::Type_Enum,
-			OldObjectName,
-			NewObjectName,
-			nullptr,
-			ECoreRedirectMatchFlags::AllowPartialMatch); // AllowPartialMatch to allow redirects from one package to another
-
-		if (!bFoundRedirect)
-		{
-			return nullptr;
-		}
-
-		const FString RedirectedObjectName = NewObjectName.ObjectName.ToString();
-		UPackage *Package = nullptr;
-		if (!NewObjectName.PackageName.IsNone())
-		{
-			Package = FindPackage(nullptr, *NewObjectName.PackageName.ToString());
-		}
-		if (Package != nullptr)
-		{
-			Object = FindObject<UField>(Package, *RedirectedObjectName);
-		}
-		if (Package == nullptr || Object == nullptr)
-		{
-			// Hail Mary pass.
-			Object = FindFirstObject<UField>(*RedirectedObjectName, EFindFirstObjectOptions::NativeFirst);
-		}
-		return Object;
-	}
+	RIGVM_API UObject* FindObjectGlobally(const TCHAR* InObjectName, bool bUseRedirector);
 
 	// A UClass argument is used to signify both the object type and its class type.
 	// This argument differentiates between the two
@@ -343,103 +228,13 @@ namespace RigVMTypeUtils
 		AsObject
 	};
 
-	static FString CPPTypeFromObject(const UObject* InCPPTypeObject, EClassArgType InClassArgType = EClassArgType::AsObject)
-	{
-		if (const UClass* Class = Cast<UClass>(InCPPTypeObject))
-		{
-			if (InClassArgType == EClassArgType::AsClass)
-			{
-				return FString::Printf(RigVMTypeUtils::TSubclassOfTemplate, Class->GetPrefixCPP(), *Class->GetName());
-			}
-			else if (Class->IsChildOf(UInterface::StaticClass()))
-			{
-				return FString::Printf(RigVMTypeUtils::TScriptInterfaceTemplate, TEXT("I"), *Class->GetName());
-			}
-			else
-			{
-				return FString::Printf(RigVMTypeUtils::TObjectPtrTemplate, Class->GetPrefixCPP(), *Class->GetName());
-			}
-		}
-		else if (const UScriptStruct* ScriptStruct = Cast<UScriptStruct>(InCPPTypeObject))
-		{
-			return GetUniqueStructTypeName(ScriptStruct);
-		}
-		else if (const UEnum* Enum = Cast<UEnum>(InCPPTypeObject))
-		{
-			return RigVMTypeUtils::CPPTypeFromEnum(Enum);
-		}
-
-		return FString();
-	}
+	RIGVM_API FString CPPTypeFromObject(const UObject* InCPPTypeObject, EClassArgType InClassArgType = EClassArgType::AsObject);
 
 	// Finds the CPPTypeObject from a CPP type of a potentially missing / unloaded user defined struct or enum
 	RIGVM_API UObject* UserDefinedTypeFromCPPType(FString& InOutCPPType, const FRigVMUserDefinedTypeResolver* InTypeResolver = nullptr);
 
 	// Finds the CPPTypeObject from the CPPType. If not found, tries to use redirectors and modifies the InOutCPPType.
-	static UObject* ObjectFromCPPType(FString& InOutCPPType, bool bUseRedirector = true, const FRigVMUserDefinedTypeResolver* InTypeResolver = nullptr)
-	{
-		if (!RequiresCPPTypeObject(InOutCPPType))
-		{
-			return nullptr;
-		}
-
-		// try to find the CPPTypeObject by name
-		FString BaseCPPType = InOutCPPType;
-		while (IsArrayType(BaseCPPType))
-		{
-			BaseCPPType = BaseTypeFromArrayType(BaseCPPType);
-		}
-		FString CPPType = BaseCPPType;
-		const bool bIsClass = CPPType.StartsWith(TSubclassOfPrefix);
-
-		static const FString PrefixObjectPtr = TObjectPtrPrefix;
-		static const FString PrefixSubclassOf = TSubclassOfPrefix;
-		static const FString PrefixScriptInterface = TScriptInterfacePrefix;
-
-		if (CPPType.StartsWith(PrefixObjectPtr))
-		{
-			// Chop the prefix + the U indicating object class
-			CPPType = CPPType.Mid(PrefixObjectPtr.Len() + 1, CPPType.Len() - (PrefixObjectPtr.Len() + 2));
-		}
-		else if (CPPType.StartsWith(PrefixSubclassOf))
-		{
-			// Chop the prefix + the U indicating object class
-			CPPType = CPPType.Mid(PrefixSubclassOf.Len() + 1, CPPType.Len() - (PrefixSubclassOf.Len() + 2));
-		}
-		else if (CPPType.StartsWith(TScriptInterfacePrefix))
-		{
-			// Chop the prefix + the I indicating interface class
-			CPPType = CPPType.Mid(PrefixScriptInterface.Len() + 1, CPPType.Len() - (PrefixScriptInterface.Len() + 2));
-		}
-
-		UObject* CPPTypeObject = FindObjectGlobally(*CPPType, bUseRedirector);
-		if (CPPTypeObject == nullptr)
-		{
-			// If we've mistakenly stored the struct type with the 'F', 'U', or 'A' prefixes, we need to strip them
-			// off first. Enums are always named with their prefix intact.
-			if (!CPPType.IsEmpty() && (CPPType[0] == TEXT('F') || CPPType[0] == TEXT('U') || CPPType[0] == TEXT('A')))
-			{
-				CPPType = CPPType.Mid(1);
-			}
-			CPPTypeObject = RigVMTypeUtils::FindObjectGlobally(*CPPType, bUseRedirector);
-		}
-
-		if(CPPTypeObject == nullptr)
-		{
-			CPPType = BaseCPPType;
-			CPPTypeObject = UserDefinedTypeFromCPPType(CPPType, InTypeResolver);
-		}
-
-		if(CPPTypeObject == nullptr)
-		{
-			InOutCPPType.Reset();
-			return nullptr;
-		}
-
-		CPPType = CPPTypeFromObject(CPPTypeObject, bIsClass ? EClassArgType::AsClass : EClassArgType::AsObject);
-		InOutCPPType.ReplaceInline(*BaseCPPType, *CPPType);
-		return CPPTypeObject;
-	}
+	RIGVM_API UObject* ObjectFromCPPType(FString& InOutCPPType, bool bUseRedirector = true, const FRigVMUserDefinedTypeResolver* InTypeResolver = nullptr);
 
 	static void CleanupCPPType(FString& CPPType)
 	{
@@ -464,104 +259,12 @@ namespace RigVMTypeUtils
 		return CPPType;
 	}
 
-	static FString PostProcessCPPType(const FString& InCPPType, UObject* InCPPTypeObject = nullptr, const FRigVMUserDefinedTypeResolver* InResolvalInfo = nullptr)
-	{
-		FString CPPType = InCPPType;
-		CleanupCPPType(CPPType);
-		
-		if (InCPPTypeObject)
-		{
-			const bool bIsClass = CPPType.StartsWith(TSubclassOfPrefix);
-			CPPType = CPPTypeFromObject(InCPPTypeObject, bIsClass ? EClassArgType::AsClass : EClassArgType::AsObject);
-			if(CPPType != InCPPType)
-			{
-				FString TemplateType = InCPPType;
-				while (RigVMTypeUtils::IsArrayType(TemplateType))
-				{
-					CPPType = RigVMTypeUtils::ArrayTypeFromBaseType(CPPType);
-					TemplateType = RigVMTypeUtils::BaseTypeFromArrayType(TemplateType);
-				}		
-			}
-		}
-		else if (RequiresCPPTypeObject(CPPType))
-		{
-			// Uses redirectors and updates the CPPType if necessary
-			ObjectFromCPPType(CPPType, true, InResolvalInfo);
-		}
-	
-		return CPPType;
-	}
-
-	static bool FixCPPTypeAndObject(FString& InOutCPPType, TObjectPtr<UObject>& InOutCPPTypeObject, const FRigVMUserDefinedTypeResolver* InResolvalInfo = nullptr)
-	{
-		const FString OldCPPType = InOutCPPType;
-		UObject* OldCPPTypeObject = InOutCPPTypeObject;
-		InOutCPPType = PostProcessCPPType(InOutCPPType, InOutCPPTypeObject, InResolvalInfo);
-
-		// Type object might have changed through a redirect
-		bool bFindObjectAgain = true;
+	RIGVM_API FString PostProcessCPPType(const FString& InCPPType, UObject* InCPPTypeObject = nullptr, const FRigVMUserDefinedTypeResolver* InResolvalInfo = nullptr);
 
-		// If its a user defined enum, do not attempt to find the object from the cpp type
-		if (const UEnum* Enum = Cast<UEnum>(InOutCPPTypeObject))
-		{
-			if (Enum->CppType.IsEmpty())
-			{
-				bFindObjectAgain = false;
-			}
-		}
-
-		if (bFindObjectAgain)
-		{
-			InOutCPPTypeObject = ObjectFromCPPType(InOutCPPType);
-		}
-		
-		return !InOutCPPType.Equals(OldCPPType) || InOutCPPTypeObject != OldCPPTypeObject;
-	}
+	RIGVM_API bool FixCPPTypeAndObject(FString& InOutCPPType, TObjectPtr<UObject>& InOutCPPTypeObject, const FRigVMUserDefinedTypeResolver* InResolvalInfo = nullptr);
 
 	// helper function to retrieve an object from a path
-	static UObject* FindObjectFromCPPTypeObjectPath(const FString& InObjectPath)
-	{
-		UObject* Result = nullptr;
-		if (InObjectPath.IsEmpty())
-		{
-			return Result;
-		}
-
-		if (InObjectPath == FName(NAME_None).ToString())
-		{
-			return Result;
-		}
-
-		// we do this to avoid ambiguous searches for 
-		// common names such as "transform" or "vector"
-		UPackage* Package = nullptr;
-		FString PackageName;
-		FString CPPTypeObjectName = InObjectPath;
-		if (InObjectPath.Split(TEXT("."), &PackageName, &CPPTypeObjectName))
-		{
-			Package = FindPackage(nullptr, *PackageName);
-		}
-	
-		if (UObject* ObjectWithinPackage = FindObject<UObject>(Package, *CPPTypeObjectName))
-		{
-			Result = ObjectWithinPackage;
-		}
-
-		if (!Result)
-		{
-			Result = FindFirstObject<UObject>(*InObjectPath, EFindFirstObjectOptions::NativeFirst | EFindFirstObjectOptions::EnsureIfAmbiguous);
-		}
-		if (!Result)
-		{
-			const FCoreRedirectObjectName OldObjectName(InObjectPath);
-			const FCoreRedirectObjectName NewObjectName = FCoreRedirects::GetRedirectedName(ECoreRedirectFlags::Type_Struct, OldObjectName);
-			if (OldObjectName != NewObjectName)
-			{
-				Result = FindObjectFromCPPTypeObjectPath(NewObjectName.ToString());
-			}
-		}
-		return Result;
-	}
+	RIGVM_API UObject* FindObjectFromCPPTypeObjectPath(const FString& InObjectPath);
 	
 	template<class T>
 	static T* FindObjectFromCPPTypeObjectPath(const FString& InObjectPath)
@@ -569,52 +272,7 @@ namespace RigVMTypeUtils
 		return Cast<T>(FindObjectFromCPPTypeObjectPath(InObjectPath));
 	}
 
-	static bool AreCompatible(const FProperty* InSourceProperty, const FProperty* InTargetProperty)
-	{
-		bool bCompatible = InSourceProperty->SameType(InTargetProperty);
-		if (!bCompatible)
-		{
-			if(const FFloatProperty* TargetFloatProperty = CastField<FFloatProperty>(InTargetProperty))
-			{
-				bCompatible = InSourceProperty->IsA<FDoubleProperty>();
-			}
-			else if(const FDoubleProperty* TargetDoubleProperty = CastField<FDoubleProperty>(InTargetProperty))
-			{
-				bCompatible = InSourceProperty->IsA<FFloatProperty>();
-			}
-			else if (const FByteProperty* TargetByteProperty = CastField<FByteProperty>(InTargetProperty))
-			{
-				bCompatible = InSourceProperty->IsA<FEnumProperty>();
-			}
-			else if (const FEnumProperty* TargetEnumProperty = CastField<FEnumProperty>(InTargetProperty))
-			{
-				bCompatible = InSourceProperty->IsA<FByteProperty>();
-			}
-			else if(const FArrayProperty* TargetArrayProperty = CastField<FArrayProperty>(InTargetProperty))
-			{
-				if(const FArrayProperty* SourceArrayProperty = CastField<FArrayProperty>(InSourceProperty))
-				{
-					if(TargetArrayProperty->Inner->IsA<FFloatProperty>())
-					{
-						bCompatible = SourceArrayProperty->Inner->IsA<FDoubleProperty>();
-					}
-					else if(TargetArrayProperty->Inner->IsA<FDoubleProperty>())
-					{
-						bCompatible = SourceArrayProperty->Inner->IsA<FFloatProperty>();
-					}
-					else if(FByteProperty* TargetArrayInnerByteProperty = CastField<FByteProperty>(TargetArrayProperty->Inner))
-					{
-						bCompatible = SourceArrayProperty->Inner->IsA<FEnumProperty>();
-					}
-					else if(FEnumProperty* TargetArrayInnerEnumProperty = CastField<FEnumProperty>(TargetArrayProperty->Inner))
-					{
-						bCompatible = SourceArrayProperty->Inner->IsA<FByteProperty>();
-					}
-				}
-			}
-		}
-		return bCompatible;
-	}
+	RIGVM_API bool AreCompatible(const FProperty* InSourceProperty, const FProperty* InTargetProperty);
 }
 
 #undef UE_API
