Fixed crash in URigVM::RefreshArgumentNameCaches.

Static initialization order was sometimes causing an FName "Float" to be passed into RigVMTypeUtils::RequiresCPPTypeObject. This causes RequiresCPPTypeObject to return true when it should be returning false.

#tests replayrun, win64 editor pie

diff --git a/Engine/Plugins/Runtime/RigVM/Source/RigVM/Private/RigVMTypeUtils.cpp b/Engine/Plugins/Runtime/RigVM/Source/RigVM/Private/RigVMTypeUtils.cpp
--- a/Engine/Plugins/Runtime/RigVM/Source/RigVM/Private/RigVMTypeUtils.cpp
+++ b/Engine/Plugins/Runtime/RigVM/Source/RigVM/Private/RigVMTypeUtils.cpp
@@ -82,48 +82,33 @@ bool RigVMTypeUtils::RequiresCPPTypeObject(const FString& InCPPType)
 {
 	static const TCHAR* PrefixesRequiringCPPTypeObject[] = {
 		TEXT("F"), 
-		TEXT("TArray<F"), 
-		TEXT("TArray<TArray<F"), 
 		TEXT("E"), 
-		TEXT("TArray<E"), 
-		TEXT("TArray<TArray<E"), 
 		TEXT("U"), 
-		TEXT("TArray<U"),
-		TEXT("TArray<TArray<U"),
 		TEXT("TObjectPtr<"), 
-		TEXT("TArray<TObjectPtr<"),
-		TEXT("TArray<TArray<TObjectPtr<"),
 		TEXT("TScriptInterface<"),
 		TEXT("TSubclassOf<")
 	};
-	static const TCHAR* PrefixesNotRequiringCPPTypeObject[] = {
-		TEXT("UInt"), 
-		TEXT("TArray<UInt"),
-		TEXT("TArray<TArray<UInt"),
-		TEXT("uint"), 
-		TEXT("TArray<uint"),
-		TEXT("TArray<TArray<uint")
-	};
 	static const TCHAR* CPPTypesNotRequiringCPPTypeObject[] = {
 		TEXT("FString"), 
-		TEXT("TArray<FString>"), 
-		TEXT("TArray<TArray<FString>>"), 
 		TEXT("FName"), 
-		TEXT("TArray<FName>"),
-		TEXT("TArray<TArray<FName>>") 
+		TEXT("float"),
+		TEXT("uint"),
+		TEXT("uint8"),
+		TEXT("uint16"),
+		TEXT("uint32"),
+		TEXT("uint64")
 	};
 
-	for(const TCHAR* Type : CPPTypesNotRequiringCPPTypeObject)
+	FStringView InCPPTypeView = InCPPType;
+	while (InCPPTypeView.StartsWith(TEXT("TArray<")) && InCPPTypeView.EndsWith(TEXT(">")))
 	{
-		if(InCPPType == Type)
-		{
-			return false;
-		}
+		InCPPTypeView.RemovePrefix(UE_ARRAY_COUNT("TArray<") - 1);
+		InCPPTypeView.RemoveSuffix(1); // >
 	}
 
-	for(const TCHAR* Prefix : PrefixesNotRequiringCPPTypeObject)
+	for(const TCHAR* Type : CPPTypesNotRequiringCPPTypeObject)
 	{
-		if(InCPPType.StartsWith(Prefix, ESearchCase::CaseSensitive))
+		if(InCPPTypeView == Type)
 		{
 			return false;
 		}
@@ -131,7 +116,7 @@ bool RigVMTypeUtils::RequiresCPPTypeObject(const FString& InCPPType)
 
 	for(const TCHAR* Prefix : PrefixesRequiringCPPTypeObject)
 	{
-		if(InCPPType.StartsWith(Prefix, ESearchCase::CaseSensitive))
+		if(InCPPTypeView.StartsWith(Prefix, ESearchCase::CaseSensitive))
 		{
 			return true;
 		}
