Fix missing search path in the type removal process for RigVMRegistry

#rnx

diff --git a/Engine/Plugins/Runtime/RigVM/Source/RigVM/Private/RigVMCore/RigVMRegistry.cpp b/Engine/Plugins/Runtime/RigVM/Source/RigVM/Private/RigVMCore/RigVMRegistry.cpp
--- a/Engine/Plugins/Runtime/RigVM/Source/RigVM/Private/RigVMCore/RigVMRegistry.cpp
+++ b/Engine/Plugins/Runtime/RigVM/Source/RigVM/Private/RigVMCore/RigVMRegistry.cpp
@@ -913,6 +913,10 @@ TArray<FString> FRigVMRegistry_NoLock::GetPluginModulePrefixes(const IPlugin& In
 	{
 		static constexpr TCHAR Format[] = TEXT("/Script/%s.");
 		ModulePrefixes.Add(FString::Printf(Format, *Module.Name.ToString()));
+
+		FName PluginName = IPluginManager::Get().PackageNameFromModuleName(Module.Name);
+		static constexpr TCHAR PluginFormat[] = TEXT("/%s/");
+		ModulePrefixes.AddUnique(FString::Printf(PluginFormat, *PluginName.ToString()));
 	}
 
 	return ModulePrefixes;
