Fix RigVM unmounting code that wrongly assumed a plugin always has modules, and update test accordingly

#rnx

diff --git a/Engine/Plugins/Runtime/RigVM/Source/RigVM/Private/RigVMCore/RigVMRegistry.cpp b/Engine/Plugins/Runtime/RigVM/Source/RigVM/Private/RigVMCore/RigVMRegistry.cpp
--- a/Engine/Plugins/Runtime/RigVM/Source/RigVM/Private/RigVMCore/RigVMRegistry.cpp
+++ b/Engine/Plugins/Runtime/RigVM/Source/RigVM/Private/RigVMCore/RigVMRegistry.cpp
@@ -913,12 +913,11 @@ TArray<FString> FRigVMRegistry_NoLock::GetPluginModulePrefixes(const IPlugin& In
 	{
 		static constexpr TCHAR Format[] = TEXT("/Script/%s.");
 		ModulePrefixes.Add(FString::Printf(Format, *Module.Name.ToString()));
-
-		FName PluginName = IPluginManager::Get().PackageNameFromModuleName(Module.Name);
-		static constexpr TCHAR PluginFormat[] = TEXT("/%s/");
-		ModulePrefixes.AddUnique(FString::Printf(PluginFormat, *PluginName.ToString()));
 	}
 
+	static constexpr TCHAR PluginFormat[] = TEXT("/%s/");
+	ModulePrefixes.AddUnique(FString::Printf(PluginFormat, *InPlugin.GetName()));
+
 	return ModulePrefixes;
 }
 
