RigVM: Rewind debugger integration - first steps. recording / tracing infrastructure

#rnx

diff --git a/Engine/Plugins/Animation/ControlRig/Source/ControlRig/Private/ControlRig.cpp b/Engine/Plugins/Animation/ControlRig/Source/ControlRig/Private/ControlRig.cpp
--- a/Engine/Plugins/Animation/ControlRig/Source/ControlRig/Private/ControlRig.cpp
+++ b/Engine/Plugins/Animation/ControlRig/Source/ControlRig/Private/ControlRig.cpp
@@ -942,6 +942,8 @@ bool UControlRig::Execute(const FName& InEventName)
 	PublicContext.bDebugExecution = bDebugExecutionEnabled;
 #endif
 
+	ExtendedExecuteContext.Host = this;
+	
 #if WITH_EDITOR
 	ExtendedExecuteContext.SetInstructionVisitInfo(&InstructionVisitInfo);
 
@@ -2609,6 +2611,56 @@ void UControlRig::EmitPostConstructionEventFinished()
 	}
 }
 
+#if RIGVM_TRACE_ENABLED
+
+void UControlRig::TraceExecuteData(FRigVMTraceArchiveWriter& InArchive) const
+{
+	InArchive.UsingCustomVersion(FControlRigObjectVersion::GUID);
+	
+	Super::TraceExecuteData(InArchive);
+
+	URigHierarchy* Hierarchy = GetHierarchy();
+	bool bHasHierarchy = Hierarchy != nullptr;
+	InArchive << bHasHierarchy;
+	if (bHasHierarchy)
+	{
+		int64 PosBeforeHierarchy = InArchive.Tell();
+		int64 PosAfterHierarchy = 0;
+		InArchive << PosAfterHierarchy;
+		
+		Hierarchy->Serialize(InArchive);
+
+		PosAfterHierarchy = InArchive.Tell();
+		InArchive.Seek(PosBeforeHierarchy);
+		InArchive << PosAfterHierarchy;
+		InArchive.Seek(PosAfterHierarchy);
+	}
+}
+
+void UControlRig::LoadTracedExecuteData(FRigVMTraceArchiveReader& InArchive)
+{
+	Super::LoadTracedExecuteData(InArchive);
+
+	bool bHasHierarchy = false;
+	InArchive << bHasHierarchy;
+	if (bHasHierarchy)
+	{
+		int64 PosAfterHierarchy = 0;
+		InArchive << PosAfterHierarchy;
+		
+		if (URigHierarchy* Hierarchy = GetHierarchy())
+		{
+			Hierarchy->Serialize(InArchive);
+		}
+		else
+		{
+			InArchive.Seek(PosAfterHierarchy);
+		}
+	}
+}
+
+#endif
+
 void UControlRig::HandleHierarchyModified(ERigHierarchyNotification InNotification, URigHierarchy* InHierarchy,
                                           const FRigNotificationSubject& InSubject)
 {
@@ -3753,7 +3805,10 @@ void UControlRig::PostInitInstance(URigVMHost* InCDO)
 
 	if(UControlRig* CDOControlRig = Cast<UControlRig>(InCDO))
 	{
-		ElementKeyRedirector = FRigElementKeyRedirector(CDOControlRig->ElementKeyRedirector, DynamicHierarchy);
+		if (DynamicHierarchy)
+		{
+			ElementKeyRedirector = FRigElementKeyRedirector(CDOControlRig->ElementKeyRedirector, DynamicHierarchy);
+		}
 	}
 
 	RequestInit();
diff --git a/Engine/Plugins/Animation/ControlRig/Source/ControlRig/Public/ControlRig.h b/Engine/Plugins/Animation/ControlRig/Source/ControlRig/Public/ControlRig.h
--- a/Engine/Plugins/Animation/ControlRig/Source/ControlRig/Public/ControlRig.h
+++ b/Engine/Plugins/Animation/ControlRig/Source/ControlRig/Public/ControlRig.h
@@ -564,6 +564,11 @@ protected:
 	
 	CONTROLRIG_API virtual void EmitPostConstructionEventFinished();
 
+#if RIGVM_TRACE_ENABLED
+	CONTROLRIG_API virtual void TraceExecuteData(FRigVMTraceArchiveWriter& InArchive) const override;
+	CONTROLRIG_API virtual void LoadTracedExecuteData(FRigVMTraceArchiveReader& InArchive) override;
+#endif
+
 private:
 #if WITH_EDITOR
 	/** Add a transient / temporary control used to interact with a node */
diff --git a/Engine/Plugins/Runtime/RigVM/RigVM.uplugin b/Engine/Plugins/Runtime/RigVM/RigVM.uplugin
--- a/Engine/Plugins/Runtime/RigVM/RigVM.uplugin
+++ b/Engine/Plugins/Runtime/RigVM/RigVM.uplugin
@@ -34,5 +34,9 @@
 	],
 	"Plugins" :
 	[
+		{
+			"Name": "GameplayInsights",
+			"Enabled": true
+		}
 	]
 }
\ No newline at end of file
diff --git a/Engine/Plugins/Runtime/RigVM/Source/RigVM/Private/RigVMCore/RigVM.cpp b/Engine/Plugins/Runtime/RigVM/Source/RigVM/Private/RigVMCore/RigVM.cpp
--- a/Engine/Plugins/Runtime/RigVM/Source/RigVM/Private/RigVMCore/RigVM.cpp
+++ b/Engine/Plugins/Runtime/RigVM/Source/RigVM/Private/RigVMCore/RigVM.cpp
@@ -21,6 +21,7 @@
 #include "RigVMCore/RigVMTrait.h"
 #include "RigVMStringUtils.h"
 #include "Misc/MTAccessDetector.h"
+#include "RigVMTrace.h"
 
 #include UE_INLINE_GENERATED_CPP_BY_NAME(RigVM)
 
@@ -3154,6 +3155,7 @@ void URigVM::StopProfiling(FRigVMExtendedExecuteContext& Context)
 		RigVMProfilingInfo->StopProfiling();
 	}
 #endif
+	TRACE_RIGVM_EVALUATION(Context);
 }
 
 #undef UE_RIGVM_GET_READ_REGISTRY_CONTEXT
\ No newline at end of file
diff --git a/Engine/Plugins/Runtime/RigVM/Source/RigVM/Private/RigVMHost.cpp b/Engine/Plugins/Runtime/RigVM/Source/RigVM/Private/RigVMHost.cpp
--- a/Engine/Plugins/Runtime/RigVM/Source/RigVM/Private/RigVMHost.cpp
+++ b/Engine/Plugins/Runtime/RigVM/Source/RigVM/Private/RigVMHost.cpp
@@ -589,6 +589,52 @@ bool URigVMHost::DisableExecution()
 	return CVarRigVMDisableExecutionAll->GetInt() == 1;
 }
 
+#if RIGVM_TRACE_ENABLED
+
+void URigVMHost::TraceConstantData(FRigVMTraceArchiveWriter& InArchive) const
+{
+	InArchive.UsingCustomVersion(FRigVMObjectVersion::GUID);
+
+	uint32 VMHash = VM ? VM->GetVMHash() : 0;
+	InArchive << VMHash;
+	uint32 ByteCodeHash = VM ? VM->GetByteCode().GetByteCodeHash() : 0;
+	InArchive << ByteCodeHash;
+}
+
+bool URigVMHost::MatchesTracedConstantData(FRigVMTraceArchiveReader& InArchive) const
+{
+	const uint32 ExpectedVMHash = VM ? VM->GetVMHash() : 0;
+	uint32 VMHash = 0;
+	InArchive << VMHash;
+	const uint32 ExpectedByteCodeHash = VM ? VM->GetByteCode().GetByteCodeHash() : 0;
+	uint32 ByteCodeHash = 0;
+	InArchive << ByteCodeHash;
+
+	return VMHash == ExpectedVMHash && ByteCodeHash == ExpectedByteCodeHash;
+}
+
+void URigVMHost::TraceExecuteData(FRigVMTraceArchiveWriter& InArchive) const
+{
+	InArchive.UsingCustomVersion(FRigVMObjectVersion::GUID);
+	
+	const FRigVMExtendedExecuteContext& ExtendedExecuteContext = GetRigVMExtendedExecuteContext();
+
+	TArray<FName> CurrentEventQueue = EventQueue;
+	InArchive << CurrentEventQueue;
+}
+
+void URigVMHost::LoadTracedExecuteData(FRigVMTraceArchiveReader& InArchive)
+{
+	FRigVMExtendedExecuteContext& ExtendedExecuteContext = GetRigVMExtendedExecuteContext();
+
+	TArray<FName> CurrentEventQueue;
+	InArchive << CurrentEventQueue;
+
+	// todo: load the data onto the host.
+}
+
+#endif
+
 bool URigVMHost::InitializeCDOVM()
 {
 	check(VM != nullptr);
diff --git a/Engine/Plugins/Runtime/RigVM/Source/RigVM/Private/RigVMModule.cpp b/Engine/Plugins/Runtime/RigVM/Source/RigVM/Private/RigVMModule.cpp
--- a/Engine/Plugins/Runtime/RigVM/Source/RigVM/Private/RigVMModule.cpp
+++ b/Engine/Plugins/Runtime/RigVM/Source/RigVM/Private/RigVMModule.cpp
@@ -14,6 +14,9 @@
 #include "PropertyEditorModule.h"
 #endif
 
+#if RIGVM_TRACE_ENABLED
+#include "Features/IModularFeatures.h"
+#endif
 
 IMPLEMENT_MODULE(FRigVMModule, RigVM)
 
@@ -40,6 +43,10 @@ void FRigVMModule::StartupModule()
 	PropertyModule.RegisterCustomPropertyTypeLayout("RigVMMemoryStorageStruct", FOnGetPropertyTypeCustomizationInstance::CreateStatic(&FPropertyBagDetails::MakeInstance));
 	PropertyModule.NotifyCustomizationModuleChanged();
 #endif
+
+#if RIGVM_TRACE_ENABLED
+	IModularFeatures::Get().RegisterModularFeature(IRewindDebuggerRuntimeExtension::ModularFeatureName, &RewindDebuggerRigVMExtension);
+#endif
 }
 
 void FRigVMModule::ShutdownModule()
@@ -52,6 +59,10 @@ void FRigVMModule::ShutdownModule()
 		PropertyModule->NotifyCustomizationModuleChanged();
 	}
 #endif
+
+#if RIGVM_TRACE_ENABLED
+	IModularFeatures::Get().UnregisterModularFeature(IRewindDebuggerRuntimeExtension::ModularFeatureName, &RewindDebuggerRigVMExtension);
+#endif
 }
 
 bool RigVMCore::SupportsUObjects()
diff --git a/Engine/Plugins/Runtime/RigVM/Source/RigVM/Private/RigVMTrace.cpp b/Engine/Plugins/Runtime/RigVM/Source/RigVM/Private/RigVMTrace.cpp
new file mode 100644
--- /dev/null
+++ b/Engine/Plugins/Runtime/RigVM/Source/RigVM/Private/RigVMTrace.cpp
@@ -0,0 +1,246 @@
+﻿// Copyright Epic Games, Inc. All Rights Reserved.
+
+#include "RigVMTrace.h"
+#include "ObjectTrace.h"
+#include "RigVMCore/RigVMObjectArchive.h"
+#include "RigVMHost.h"
+#include "RigVMCore/RigVMExecuteContext.h"
+
+#if RIGVM_TRACE_ENABLED
+
+UE_TRACE_CHANNEL_DEFINE(RigVMChannel)
+UE_TRACE_EVENT_DEFINE(RigVM, Literals);
+UE_TRACE_EVENT_DEFINE(RigVM, Execute);
+
+TAutoConsoleVariable<bool> CVarRigVMEnableRewindDebuggerTracing(TEXT("RigVM.EnableRewindDebugger"), 0, TEXT("Enable tracing of RigVM instances in RewindDebugger"));
+
+void FRigVMTrace::RecordingStarted()
+{
+	if (!CVarRigVMEnableRewindDebuggerTracing.GetValueOnAnyThread())
+	{
+		return;
+	}
+	UE::Trace::ToggleChannel(TEXT("RigVM"), true);
+}
+
+void FRigVMTrace::RecordingStopped()
+{
+	UE::Trace::ToggleChannel(TEXT("RigVM"), false);
+}
+
+void FRigVMTrace::TraceRigVMEvaluation(const FRigVMExtendedExecuteContext& InContext)
+{
+	URigVM* VM = InContext.VM;
+	
+	bool bChannelEnabled = UE_TRACE_CHANNELEXPR_IS_ENABLED(RigVMChannel);
+	if (!bChannelEnabled || VM == nullptr)
+	{
+		return;
+	}
+
+	URigVMHost* Host = InContext.Host;
+	if (!Host || CANNOT_TRACE_OBJECT(Host))
+	{
+		return;
+	}
+	
+	// determine the world this host is using within.
+	// for control rig this funnels through the object
+	// binding and will find the level this is instantiated within.
+	UWorld* World = Host->GetWorld();
+	if (!World || CANNOT_TRACE_OBJECT(VM))
+	{
+		return;
+	}
+
+	TRACE_OBJECT(Host);
+
+	const bool bTraceLiterals = TRACE_OBJECT(VM);
+	constexpr bool bTraceExecute = true; 
+
+	FRigVMTraceArchive LiteralMemoryArchive;
+	if (bTraceLiterals)
+	{
+		FRigVMTraceArchiveWriter LiteralMemoryWriter(LiteralMemoryArchive);
+		LiteralMemoryWriter << VM->GetDefaultLiteralMemory();
+	}
+
+	FRigVMTraceArchive WorkMemoryArchive;
+	if(bTraceExecute)
+	{
+		FRigVMTraceArchiveWriter WorkMemoryWriter(WorkMemoryArchive);
+		WorkMemoryWriter << const_cast<FRigVMMemoryStorageStruct&>(InContext.WorkMemoryStorage);
+	}
+
+	FRigVMTraceArchive DebugMemoryArchive;
+	if(bTraceExecute)
+	{
+		FRigVMTraceArchiveWriter DebugMemoryWriter(DebugMemoryArchive);
+		DebugMemoryWriter << const_cast<FRigVMMemoryStorageStruct&>(InContext.DebugMemoryStorage);
+	}
+
+	FRigVMTraceArchive ExternalMemoryArchive;
+	if (bTraceExecute)
+	{
+		FRigVMTraceArchiveWriter ExternalMemoryWriter(ExternalMemoryArchive);
+		{
+			FBinaryArchiveFormatter BinaryArchiveFormatter(ExternalMemoryWriter);
+			{
+				FStructuredArchive StructuredArchive(BinaryArchiveFormatter);
+				{
+					FStructuredArchive::FRecord RootRecord = StructuredArchive.Open().EnterRecord();
+					{
+						const TArray<FRigVMExternalVariable> ExternalVariables = VM->GetExternalVariables(InContext);
+						TArray<FRigVMExternalVariable> FilteredExternalVariables;
+						for (const FRigVMExternalVariable& ExternalVariable : ExternalVariables)
+						{
+							if (ExternalVariable.Property && ExternalVariable.Memory)
+							{
+								FilteredExternalVariables.Add(ExternalVariable);
+							}
+						}
+						
+						int32 NumFilteredExternalVariables = FilteredExternalVariables.Num();
+						ExternalMemoryWriter << NumFilteredExternalVariables;
+						for (FRigVMExternalVariable& FilteredExternalVariable : FilteredExternalVariables)
+						{
+							check (FilteredExternalVariable.Property && FilteredExternalVariable.Memory)
+
+							FStructuredArchiveSlot PropertyField = RootRecord.EnterField(*FilteredExternalVariable.Property->GetName());
+							FilteredExternalVariable.Property->SerializeItem(PropertyField, FilteredExternalVariable.Memory);
+						}
+					}
+				}
+			}
+		}
+	}
+
+	FRigVMTraceArchive DrawInterfaceArchive;
+	if (bTraceExecute)
+	{
+		FRigVMTraceArchiveWriter DrawInterfaceWriter(DrawInterfaceArchive);
+		{
+			FBinaryArchiveFormatter BinaryArchiveFormatter(DrawInterfaceWriter);
+			{
+				FStructuredArchive StructuredArchive(BinaryArchiveFormatter);
+				{
+					FStructuredArchive::FRecord RootRecord = StructuredArchive.Open().EnterRecord();
+					{
+						if (const FProperty* DrawInterfaceProperty = URigVMHost::StaticClass()->FindPropertyByName(GET_MEMBER_NAME_CHECKED(URigVMHost, DrawInterface)))
+						{
+							FStructuredArchiveSlot PropertyField = RootRecord.EnterField(*DrawInterfaceProperty->GetName());
+							DrawInterfaceProperty->SerializeItem(PropertyField, &Host->GetDrawInterface());
+						}
+						if (const FProperty* DrawContainerProperty = URigVMHost::StaticClass()->FindPropertyByName(GET_MEMBER_NAME_CHECKED(URigVMHost, DrawContainer)))
+						{
+							FStructuredArchiveSlot PropertyField = RootRecord.EnterField(*DrawContainerProperty->GetName());
+							DrawContainerProperty->SerializeItem(PropertyField, &Host->GetDrawContainer());
+						}
+					}
+				}
+			}
+		}
+	}
+
+	FRigVMTraceArchive HostConstantDataArchive;
+	if (bTraceLiterals)
+	{
+		FRigVMTraceArchiveWriter HostConstantDataWriter(HostConstantDataArchive);
+		Host->TraceConstantData(HostConstantDataWriter);
+	}
+
+	FRigVMTraceArchive HostExecuteDataArchive;
+	if (bTraceExecute)
+	{
+		FRigVMTraceArchiveWriter HostExecuteDataWriter(HostExecuteDataArchive);
+		Host->TraceExecuteData(HostExecuteDataWriter);
+	}
+
+	FString JoinedEntries;
+	if (bTraceExecute)
+	{
+		TArray<FString> Entries;
+		for (const int32& EntryIndex : InContext.EntriesBeingExecuted)
+		{
+			if (EntryIndex >= 0 && EntryIndex < VM->GetByteCode().NumEntries())
+			{
+				Entries.AddUnique(VM->GetByteCode().GetEntry(EntryIndex).Name.ToString());
+			}
+		}
+		if (Entries.Num() > 0)
+		{
+			JoinedEntries = RigVMStringUtils::JoinStrings(Entries, TEXT(","));
+		}
+	}
+
+	static const TArray<int32> EmptyInt32Array;
+	static const TArray<uint64> EmptyUInt64Array;
+
+#if WITH_EDITOR
+	FRigVMInstructionVisitInfo* InstructionVisitInfo = InContext.GetRigVMInstructionVisitInfo();
+	const TArray<int32>& VisitedInstructions = InstructionVisitInfo ? InstructionVisitInfo->GetInstructionVisitOrder() : EmptyInt32Array;
+
+	FRigVMProfilingInfo* ProfilingInfo = InContext.GetRigVMProfilingInfo();
+	uint64 OverallCycles = ProfilingInfo ? ProfilingInfo->GetOverallCycles() : 0;
+	const TArray<uint64>& InstructionCycles = ProfilingInfo ? ProfilingInfo->InstructionCyclesDuringLastRun : EmptyUInt64Array;
+#else
+	const TArray<int32>& VisitedInstructions = EmptyInt32Array;
+	uint64 OverallCycles = 0;
+	const TArray<uint64>& InstructionCycles = EmptyUInt64Array;
+#endif
+
+	if (bTraceLiterals)
+	{
+		LiteralMemoryArchive.Compress();
+		HostConstantDataArchive.Compress();
+
+		FString HostGeneratedBy;
+#if WITH_EDITOR
+		if (const UClass* Class = Host->GetClass())
+		{
+			if (const UObject* GeneratedBy = Class->ClassGeneratedBy.Get())
+			{
+				if (const UPackage* Package = GeneratedBy->GetOutermost())
+				{
+					HostGeneratedBy = Package->GetPathName();
+				}
+			}
+		}
+#endif
+
+		UE_TRACE_LOG(RigVM, Literals, RigVMChannel)
+		<< Literals.Cycles(FPlatformTime::Cycles64())
+		<< Literals.RecordingTime(FObjectTrace::GetObjectWorldElapsedTime(World))
+		<< Literals.RigVMObjectVersion(FRigVMObjectVersion::LatestVersion)
+		<< Literals.HostId(FObjectTrace::GetObjectId(Host))
+		<< Literals.VMHash(InContext.VMHash)
+		<< Literals.ByteCodeHash(VM->GetByteCode().GetByteCodeHash())
+		<< Literals.LiteralMemory(LiteralMemoryArchive.GetOverallData(), LiteralMemoryArchive.OverallNum())
+		<< Literals.HostGeneratedBy(*HostGeneratedBy, HostGeneratedBy.Len())
+		<< Literals.HostConstantData(HostConstantDataArchive.GetOverallData(), HostConstantDataArchive.OverallNum());
+	}
+
+	if (bTraceExecute)
+	{
+		UE_TRACE_LOG(RigVM, Execute, RigVMChannel)
+		<< Execute.Cycles(FPlatformTime::Cycles64())
+		<< Execute.RecordingTime(FObjectTrace::GetObjectWorldElapsedTime(World))
+		<< Execute.AbsoluteTime(Host->GetAbsoluteTime())
+		<< Execute.DeltaTime(Host->GetDeltaTime())
+		<< Execute.HostId(FObjectTrace::GetObjectId(Host))
+		<< Execute.WorldId(FObjectTrace::GetObjectId(World))
+		<< Execute.RigVMObjectVersion(FRigVMObjectVersion::LatestVersion)
+		<< Execute.WorkMemory(WorkMemoryArchive.GetOverallData(), WorkMemoryArchive.OverallNum())
+		<< Execute.DebugMemory(DebugMemoryArchive.GetOverallData(), DebugMemoryArchive.OverallNum())
+		<< Execute.ExternalMemory(ExternalMemoryArchive.GetOverallData(), ExternalMemoryArchive.OverallNum())
+		<< Execute.Entries(*JoinedEntries, JoinedEntries.Len())
+		<< Execute.VisitedInstructions(VisitedInstructions.GetData(), VisitedInstructions.Num())
+		<< Execute.OverallCycles(OverallCycles)
+		<< Execute.InstructionCycles(InstructionCycles.GetData(), InstructionCycles.Num())
+		<< Execute.DrawInterface(DrawInterfaceArchive.GetOverallData(), DrawInterfaceArchive.OverallNum())
+		<< Execute.HostExecuteData(HostExecuteDataArchive.GetOverallData(), HostExecuteDataArchive.OverallNum());
+	}
+}
+
+#endif
+
diff --git a/Engine/Plugins/Runtime/RigVM/Source/RigVM/Private/RigVMTraceArchive.cpp b/Engine/Plugins/Runtime/RigVM/Source/RigVM/Private/RigVMTraceArchive.cpp
new file mode 100644
--- /dev/null
+++ b/Engine/Plugins/Runtime/RigVM/Source/RigVM/Private/RigVMTraceArchive.cpp
@@ -0,0 +1,467 @@
+// Copyright Epic Games, Inc. All Rights Reserved.
+
+#include "RigVMTraceArchive.h"
+#include "ObjectTrace.h"
+#include "RigVMTrace.h"
+#include "Misc/Compression.h"
+#include "RigVMObjectVersion.h"
+
+FRigVMTraceArchive::FRigVMTraceArchive()
+{
+	Reset();
+}
+
+void FRigVMTraceArchive::Reset()
+{
+	Buffer.Reset();
+	Buffer.AddZeroed(sizeof(FHeader));
+	Header() = FHeader();
+}
+
+void FRigVMTraceArchive::Empty()
+{
+	Buffer.Empty();
+	Buffer.AddZeroed(sizeof(FHeader));
+	Header() = FHeader();
+}
+
+bool FRigVMTraceArchive::IsPayloadEmpty() const
+{
+	return PayloadNum() == 0;
+}
+
+void FRigVMTraceArchive::Compress()
+{
+	if(IsCompressed())
+	{
+		return;
+	}
+
+	if(IsPayloadEmpty())
+	{
+		return;
+	}
+
+	FHeader& H = Header();
+	H.UncompressedSize = PayloadNum();
+	
+	// It is possible for compression to actually increase the size of the data, so we over allocate here to handle that.
+	H.CompressedSize = FCompression::CompressMemoryBound(NAME_Oodle, H.UncompressedSize);
+
+	TArray<uint8> CompressedBuffer;
+	CompressedBuffer.SetNumUninitialized(H.CompressedSize + sizeof(FHeader));
+	FMemory::Memcpy(CompressedBuffer.GetData(), GetOverallData(), sizeof(FHeader));
+	
+	if (FCompression::CompressMemory(NAME_Oodle, CompressedBuffer.GetData() + sizeof(FHeader), H.CompressedSize, GetPayloadData(), PayloadNum(), COMPRESS_BiasMemory))
+	{
+		// In the case that compressing it actually increases the size, we leave it uncompressed
+		if (H.CompressedSize < H.UncompressedSize)
+		{
+			const int32 CompressedSize = H.CompressedSize;
+			CompressedBuffer.SetNum(CompressedSize + sizeof(FHeader));
+			Buffer = MoveTemp(CompressedBuffer);
+			Buffer.Shrink();
+			Header().CompressedSize = CompressedSize;
+			Header().bIsCompressed = true;
+		}
+		else
+		{
+			H.CompressedSize = INDEX_NONE;
+		}
+	}
+	else
+	{
+		H.CompressedSize = INDEX_NONE;
+	}
+}
+
+bool FRigVMTraceArchive::Decompress()
+{
+	if(!IsCompressed())
+	{
+		return true;
+	}
+
+	FHeader& H = Header();
+	H.CompressedSize = PayloadNum();
+	
+	TArray<uint8> UncompressedBuffer;
+	UncompressedBuffer.SetNumUninitialized(H.UncompressedSize + sizeof(FHeader));
+	FMemory::Memcpy(UncompressedBuffer.GetData(), GetOverallData(), sizeof(FHeader));
+
+	if (FCompression::UncompressMemory(NAME_Oodle, UncompressedBuffer.GetData() + sizeof(FHeader), H.UncompressedSize, GetPayloadData(), PayloadNum()))
+	{
+		Buffer = MoveTemp(UncompressedBuffer);
+		Header().bIsCompressed = false;
+		return true;
+	}
+
+	return false;
+}
+
+const FRigVMTraceArchive::FHeader& FRigVMTraceArchive::GetHeader() const
+{
+	return const_cast<FRigVMTraceArchive*>(this)->Header();
+}
+
+FRigVMTraceArchive::FHeader& FRigVMTraceArchive::Header()
+{
+	check(Buffer.Num() >= sizeof(FHeader));
+	return *reinterpret_cast<FHeader*>(GetOverallData());
+}
+
+bool FRigVMTraceArchive::IsCompressed() const
+{
+	return GetHeader().bIsCompressed;
+}
+
+int32 FRigVMTraceArchive::GetUncompressedSize() const
+{
+	return GetHeader().UncompressedSize;
+}
+
+int32 FRigVMTraceArchive::GetCompressedSize() const
+{
+	return GetHeader().CompressedSize;
+}
+
+uint8* FRigVMTraceArchive::GetPayloadData()
+{
+	if (IsPayloadEmpty())
+	{
+		return nullptr;
+	}
+	return Buffer.GetData() + sizeof(FHeader);
+}
+
+const uint8* FRigVMTraceArchive::GetPayloadData() const
+{
+	if (IsPayloadEmpty())
+	{
+		return nullptr;
+	}
+	return Buffer.GetData() + sizeof(FHeader);
+}
+
+int32 FRigVMTraceArchive::PayloadNum() const
+{
+	return FMath::Max<int32>(0, Buffer.Num() - sizeof(FHeader));
+}
+
+uint8* FRigVMTraceArchive::GetOverallData()
+{
+	return Buffer.GetData();
+}
+
+const uint8* FRigVMTraceArchive::GetOverallData() const
+{
+	return Buffer.GetData();
+}
+
+int32 FRigVMTraceArchive::OverallNum() const
+{
+	return Buffer.Num();
+}
+
+void FRigVMTraceArchive::SetOverallBuffer(const TArrayView<const uint8>& InBuffer)
+{
+	SetOverallBuffer(TArrayView<uint8>(const_cast<uint8*>(InBuffer.GetData()), InBuffer.Num()));
+}
+
+void FRigVMTraceArchive::SetOverallBuffer(const TArrayView<uint8>& InBuffer)
+{
+	check(InBuffer.Num() >= sizeof(FHeader));
+	Buffer.Reset();
+	Buffer.Append(InBuffer);
+}
+
+FRigVMTraceArchiveWriter::FRigVMTraceArchiveWriter(FRigVMTraceArchive& InArchive)
+: Archive( InArchive )
+, Offset(0)
+{
+	SetIsSaving(true);
+	SetIsLoading(false);
+	UsingCustomVersion(FRigVMObjectVersion::GUID);
+	ArIgnoreOuterRef = 1;
+}
+
+FRigVMTraceArchiveWriter::~FRigVMTraceArchiveWriter()
+{
+	if (IsSaving())
+	{
+		WriteVersions();
+	}
+}
+
+void FRigVMTraceArchiveWriter::Serialize(void* V, int64 Length)
+{
+	if(Length == 0)
+	{
+		check(V == nullptr);
+	}
+	else
+	{
+		check(V);
+		int32 Start = IntCastChecked<int32, int64>(Offset) + sizeof(FRigVMTraceArchive::FHeader);
+
+		if(Archive.Buffer.IsValidIndex(Start))
+		{
+			const int32 MissingBytes = IntCastChecked<int32, int64>(Length) - (Archive.Buffer.Num() - Start);
+			if(MissingBytes > 0)
+			{
+				(void)Archive.Buffer.AddUninitialized(MissingBytes);
+			}
+		}
+		else if(Start == Archive.Buffer.Num())
+		{
+			Start = Archive.Buffer.AddUninitialized(IntCastChecked<int32, int64>(Length));
+		}
+		else
+		{
+			checkNoEntry();
+		}
+
+		FMemory::Memcpy( Archive.GetOverallData() + Start, V, Length );
+		Offset += Length;
+	}
+}
+
+int64 FRigVMTraceArchiveWriter::Tell()
+{
+	return Offset;
+}
+
+int64 FRigVMTraceArchiveWriter::TotalSize()
+{
+	return Archive.PayloadNum();
+}
+
+void FRigVMTraceArchiveWriter::Seek(int64 InPos)
+{
+	check((InPos >= 0) && (InPos <= Archive.PayloadNum()));
+	Offset = InPos;
+}
+
+void FRigVMTraceArchiveWriter::WriteUObject(const UObject* Obj)
+{
+	uint64 Id = 0;
+	
+#if RIGVM_TRACE_ENABLED
+	if (Obj)
+	{
+		if (const uint64* ExistingId = ObjectIdMap.Find(Obj))
+		{
+			Id = *ExistingId;
+		}
+		else
+		{
+			// first ensure the object has been traced 
+			if (const UClass* Class = Cast<UClass>(Obj))
+			{
+				TRACE_TYPE(Class);
+			}
+			TRACE_OBJECT(Obj);
+		
+			// then serialize the Object as a Trace Object id;
+			Id = FObjectTrace::GetObjectId(Obj);
+
+			// store this in the map for faster lookup
+			ObjectIdMap.Add(Obj, Id);
+		}
+	}
+#endif
+
+	*this << Id;
+}
+
+FArchive& FRigVMTraceArchiveWriter::operator<<(UObject*& Obj)
+{
+	WriteUObject(Obj);
+	return *this;
+}
+
+FArchive& FRigVMTraceArchiveWriter::operator<<(FName& Value)
+{
+	if(const int64* StringOffsetPtr = NameToOffset.Find(Value))
+	{
+		uint8 State = StoringStringAsOffset;
+		*this << State;
+		
+		int64 StringOffset = *StringOffsetPtr;
+		*this << StringOffset;
+	}
+	else
+	{
+		uint8 State = StoringStringAsString;
+		*this << State;
+		NameToOffset.Add(Value, Tell());
+		FString NameAsString = Value.IsNone() ? FString() : Value.ToString();
+		*this << NameAsString;
+	}
+	return *this;
+}
+
+FArchive& FRigVMTraceArchiveWriter::operator<<(FText& Value)
+{
+	FString ValueString = Value.ToString();
+	*this << ValueString;
+	return *this;
+}
+
+void FRigVMTraceArchiveWriter::WriteVersions()
+{
+	const FCustomVersionContainer& Versions = GetCustomVersions();
+	FCustomVersionArray VersionArray = Versions.GetAllVersions();
+	if (VersionArray.IsEmpty())
+	{
+		return;
+	}
+	Archive.Header().NumVersions = VersionArray.Num();
+	Archive.Header().VersionsOffset = Offset;
+	*this << VersionArray;
+}
+
+FRigVMTraceArchiveReader::FRigVMTraceArchiveReader(FRigVMTraceArchive& InArchive)
+: FRigVMTraceArchiveWriter(InArchive)
+{
+	SetIsSaving(false);
+	SetIsLoading(true);
+	if (!Archive.Decompress())
+	{
+		checkNoEntry();
+		return;
+	}
+	ReadVersions();
+}
+
+void FRigVMTraceArchiveReader::Serialize(void* V, int64 Length)
+{
+	if(Length == 0)
+	{
+		check(V == nullptr);
+	}
+	else
+	{
+		check(V);
+		check(Archive.PayloadNum() >= Offset + Length);
+		FMemory::Memcpy( V, Archive.GetPayloadData() + IntCastChecked<int32, int64>(Offset), Length );
+		Offset += Length;
+	}
+}
+
+void FRigVMTraceArchiveReader::ReadUObject(UObject*& Obj)
+{
+	uint64 Id = 0;
+	const int64 OffsetOfId = Tell();
+	*this << Id;
+
+#if RIGVM_TRACE_ENABLED
+	if (Id == 0)
+	{
+		Obj = nullptr;
+	}
+	else
+	{
+		Obj = FObjectTrace::GetObjectFromId(Id);
+	}
+#else
+	Obj = nullptr;
+#endif
+}
+
+FArchive& FRigVMTraceArchiveReader::operator<<(UObject*& Obj)
+{
+	ReadUObject(Obj);
+	return *this;
+}
+
+FArchive& FRigVMTraceArchiveReader::operator<<(FName& Value)
+{
+	uint8 State = UINT8_MAX;
+	*this << State;
+
+	auto ReadNameString = [this, &Value]()
+	{
+		const int64 OffsetOfString = Tell();
+		FString NameAsString;
+		*this << NameAsString;
+		if(NameAsString.IsEmpty())
+		{
+			Value = FName(NAME_None);
+		}
+		else
+		{
+			Value = *NameAsString;
+		}
+		OffsetToName.FindOrAdd(OffsetOfString) = Value;
+	};
+
+	if(State == StoringStringAsString)
+	{
+		// read the string at the current position
+		ReadNameString();
+	}
+	else if(State == StoringStringAsOffset)
+	{
+		int64 OffsetOfString = -1;
+		*this << OffsetOfString;
+
+		const int32 OffsetOfString32 = IntCastChecked<int32, int64>(OffsetOfString);
+		check(OffsetOfString32 >= 0 && OffsetOfString32 < Archive.PayloadNum());
+		
+		if(const FName* NamePtr = OffsetToName.Find(OffsetOfString))
+		{
+			Value = *NamePtr;
+		}
+		else
+		{
+			TGuardValue<int64> OffsetGuard(Offset, Offset);
+			Seek(OffsetOfString);
+			ReadNameString();
+		}
+	}
+	else
+	{
+		checkNoEntry();
+	}
+
+	return *this;
+}
+
+FArchive& FRigVMTraceArchiveReader::operator<<(FText& Value)
+{
+	FString ValueString;
+	*this << ValueString;
+	Value = FText::FromString(ValueString);
+	return *this;
+}
+
+void FRigVMTraceArchiveReader::ReadVersions()
+{
+	const int32 NumVersions = Archive.Header().NumVersions; 
+	if (NumVersions > 0)
+	{
+		const int64 VersionOffset = Archive.Header().VersionsOffset;
+		check(VersionOffset != INDEX_NONE);
+
+		TGuardValue<int64> OffsetGuard(Offset, Offset);
+		Seek(VersionOffset);
+
+		FCustomVersionArray VersionArray;
+		*this << VersionArray;
+
+		FCustomVersionContainer VersionContainer;
+		for (const FCustomVersion& Version : VersionArray)
+		{
+			const FName FriendlyName = Version.GetFriendlyName();
+			VersionContainer.SetVersion(Version.Key, Version.Version, FriendlyName);
+		}
+		SetCustomVersions(VersionContainer);
+	}
+}
+
+FArchive& operator<<(FArchive& Ar, FRigVMTraceArchive& Data)
+{
+	Ar << Data.Buffer;
+	return Ar;
+}
diff --git a/Engine/Plugins/Runtime/RigVM/Source/RigVM/Public/RigVMCore/RigVMExecuteContext.h b/Engine/Plugins/Runtime/RigVM/Source/RigVM/Public/RigVMCore/RigVMExecuteContext.h
--- a/Engine/Plugins/Runtime/RigVM/Source/RigVM/Public/RigVMCore/RigVMExecuteContext.h
+++ b/Engine/Plugins/Runtime/RigVM/Source/RigVM/Public/RigVMCore/RigVMExecuteContext.h
@@ -37,6 +37,7 @@
 #define UE_API RIGVM_API
 
 class URigVM;
+class URigVMHost;
 struct FRigVMDispatchFactory;
 struct FRigVMExecuteContext;
 struct FRigVMExtendedExecuteContext;
@@ -842,6 +843,7 @@ struct FRigVMExtendedExecuteContext
 #endif
 
 	FStructOnScope PublicDataScope;
+	URigVMHost* Host = nullptr;
 	URigVM* VM = nullptr;
 	TArray<FRigVMSlice> Slices;
 	TArray<uint16> SliceOffsets;
diff --git a/Engine/Plugins/Runtime/RigVM/Source/RigVM/Public/RigVMCore/RigVMProfilingInfo.h b/Engine/Plugins/Runtime/RigVM/Source/RigVM/Public/RigVMCore/RigVMProfilingInfo.h
--- a/Engine/Plugins/Runtime/RigVM/Source/RigVM/Public/RigVMCore/RigVMProfilingInfo.h
+++ b/Engine/Plugins/Runtime/RigVM/Source/RigVM/Public/RigVMCore/RigVMProfilingInfo.h
@@ -149,6 +149,8 @@ private:
 	uint64 OverallCycles = 0;
 
 	double LastExecutionMicroSeconds = 0.0;
+
+	friend class FRigVMTrace;
 };
 
 #undef UE_API
diff --git a/Engine/Plugins/Runtime/RigVM/Source/RigVM/Public/RigVMHost.h b/Engine/Plugins/Runtime/RigVM/Source/RigVM/Public/RigVMHost.h
--- a/Engine/Plugins/Runtime/RigVM/Source/RigVM/Public/RigVMHost.h
+++ b/Engine/Plugins/Runtime/RigVM/Source/RigVM/Public/RigVMHost.h
@@ -7,12 +7,16 @@
 #include "RigVMCore/RigVM.h"
 #include "RigVMCore/RigVMExecuteContext.h"
 #include "RigVMCore/RigVMGraphFunctionHost.h"
+#include "RigVMTrace.h"
 #include "Engine/AssetUserData.h"
 #include "Interfaces/Interface_AssetUserData.h"
 #include "Misc/TransactionallySafeCriticalSection.h"
 #if WITH_EDITOR
 #include "RigVMCore/RigVMDebugInfo.h"
 #endif
+#if RIGVM_TRACE_ENABLED
+#include "RigVMCore/RigVMObjectArchive.h"
+#endif
 #include "RigVMHost.generated.h"
 
 #define UE_API RIGVM_API 
@@ -194,6 +198,13 @@ protected:
 
 	static UE_API bool DisableExecution();
 
+#if RIGVM_TRACE_ENABLED
+	UE_API virtual void TraceConstantData(FRigVMTraceArchiveWriter& InArchive) const;
+	UE_API virtual bool MatchesTracedConstantData(FRigVMTraceArchiveReader& InArchive) const;
+	UE_API virtual void TraceExecuteData(FRigVMTraceArchiveWriter& InArchive) const;
+	UE_API virtual void LoadTracedExecuteData(FRigVMTraceArchiveReader& InArchive);
+#endif
+
 public:
 
 	template<class T>
@@ -415,6 +426,7 @@ public:
 	FRigVMDrawContainer DrawContainer;
 
 	/** The draw interface for the units to use */
+	UPROPERTY(transient)
 	FRigVMDrawInterface DrawInterface;
 
 	/** The event name used during an update */
@@ -539,6 +551,7 @@ protected:
 	friend class IRigVMAssetInterface;
 	friend class URigVMBlueprintGeneratedClass;
 	friend class FRigVMCompileSettingsDetails;
+	friend class FRigVMTrace;
 };
 
 class FRigVMBracketScope
diff --git a/Engine/Plugins/Runtime/RigVM/Source/RigVM/Public/RigVMModule.h b/Engine/Plugins/Runtime/RigVM/Source/RigVM/Public/RigVMModule.h
--- a/Engine/Plugins/Runtime/RigVM/Source/RigVM/Public/RigVMModule.h
+++ b/Engine/Plugins/Runtime/RigVM/Source/RigVM/Public/RigVMModule.h
@@ -9,6 +9,7 @@
 #include "CoreMinimal.h"
 #include "Logging/LogMacros.h"
 #include "Modules/ModuleInterface.h"
+#include "RigVMTrace.h"
 #include "RigVMDefines.h"
 
 #define UE_API RIGVM_API
@@ -24,6 +25,11 @@ public:
 	// Begin IModuleInterface
 	UE_API virtual void StartupModule() override;
 	UE_API virtual void ShutdownModule() override;
+private:
+
+#if RIGVM_TRACE_ENABLED
+	FRigVMTrace RewindDebuggerRigVMExtension;
+#endif
 };
 
 namespace RigVMCore
diff --git a/Engine/Plugins/Runtime/RigVM/Source/RigVM/Public/RigVMTrace.h b/Engine/Plugins/Runtime/RigVM/Source/RigVM/Public/RigVMTrace.h
new file mode 100644
--- /dev/null
+++ b/Engine/Plugins/Runtime/RigVM/Source/RigVM/Public/RigVMTrace.h
@@ -0,0 +1,174 @@
+// Copyright Epic Games, Inc. All Rights Reserved.
+
+#pragma once
+
+#include "CoreMinimal.h"
+#include "CoreTypes.h"
+#include "Trace/Config.h"
+#include "TraceFilter.h"
+#include "RigVMCore/RigVMMemoryStorageStruct.h"
+#include "RigVMTraceArchive.h"
+#include "RewindDebuggerRuntimeInterface/IRewindDebuggerRuntimeExtension.h"
+
+#define UE_API RIGVM_API
+
+#if UE_TRACE_ENABLED && !IS_PROGRAM && !UE_BUILD_SHIPPING && !UE_BUILD_TEST
+#define RIGVM_TRACE_ENABLED 1
+#else
+#define RIGVM_TRACE_ENABLED 0
+#endif
+
+class UObject;
+struct FRigVMExtendedExecuteContext;
+struct FRigVMMemoryStorageStruct;
+
+#if RIGVM_TRACE_ENABLED
+
+UE_API UE_TRACE_CHANNEL_EXTERN(RigVMChannel)
+
+/**
+ * Base structure for all RigVM trace structures
+ */
+struct FRigVMTraceBaseData
+{
+	FRigVMTraceBaseData()
+		: Cycles(0)
+		, ProfileTime(0)
+		, RecordingTime(0)
+		, RigVMObjectVersion(0)
+	{
+	}
+
+	// Cycles the frame took
+	uint64 Cycles;
+
+	// The time of the sample at machine time
+	double ProfileTime;
+	// The time of the sample at internal time of the host
+	double RecordingTime;
+
+	// The version used (FRigVMObjectVersion)
+	int32 RigVMObjectVersion;
+};
+
+UE_API UE_TRACE_EVENT_BEGIN_EXTERN(RigVM, Literals)
+	UE_TRACE_EVENT_FIELD(uint64, Cycles)
+	UE_TRACE_EVENT_FIELD(double, RecordingTime)
+	UE_TRACE_EVENT_FIELD(int32, RigVMObjectVersion)
+	UE_TRACE_EVENT_FIELD(uint64, HostId)
+	UE_TRACE_EVENT_FIELD(uint32, VMHash)
+	UE_TRACE_EVENT_FIELD(uint32, ByteCodeHash)
+	UE_TRACE_EVENT_FIELD(uint8[], LiteralMemory)
+	UE_TRACE_EVENT_FIELD(UE::Trace::WideString, HostGeneratedBy)
+	UE_TRACE_EVENT_FIELD(uint8[], HostConstantData)
+UE_TRACE_EVENT_END()
+
+/**
+ * The constant data stored in the stream only once
+ * when tracing a RigVMHost
+ */
+struct FRigVMTraceConstantData : FRigVMTraceBaseData
+{
+	FRigVMTraceConstantData()
+		: FRigVMTraceBaseData()
+		, HostId(0)
+		, VMHash(0)
+		, ByteCodeHash(0)
+	{
+	}
+	
+	// The object id of the traced RigVM Host
+	uint64 HostId;
+	// The hash of the traced VM
+	uint32 VMHash;
+	// The hash of the traced bytecode
+	uint32 ByteCodeHash;
+	// An archive storing the literal memory
+	FRigVMTraceArchive LiteralMemory;
+	// An archive storing auxiliary constant data from the host
+	FRigVMTraceArchive HostConstantData;
+};
+
+UE_API UE_TRACE_EVENT_BEGIN_EXTERN(RigVM, Execute)
+	UE_TRACE_EVENT_FIELD(uint64, Cycles)
+	UE_TRACE_EVENT_FIELD(double, RecordingTime)
+	UE_TRACE_EVENT_FIELD(int32, RigVMObjectVersion)
+	UE_TRACE_EVENT_FIELD(double, AbsoluteTime)
+	UE_TRACE_EVENT_FIELD(double, DeltaTime)
+	UE_TRACE_EVENT_FIELD(uint64, WorldId)
+	UE_TRACE_EVENT_FIELD(uint64, HostId)
+	UE_TRACE_EVENT_FIELD(uint8[], WorkMemory)
+	UE_TRACE_EVENT_FIELD(uint8[], DebugMemory)
+	UE_TRACE_EVENT_FIELD(uint8[], ExternalMemory)
+	UE_TRACE_EVENT_FIELD(UE::Trace::WideString, Entries)
+	UE_TRACE_EVENT_FIELD(int32[], VisitedInstructions)
+	UE_TRACE_EVENT_FIELD(uint64, OverallCycles)
+	UE_TRACE_EVENT_FIELD(uint64[], InstructionCycles)
+	UE_TRACE_EVENT_FIELD(uint8[], DrawInterface)
+	UE_TRACE_EVENT_FIELD(uint8[], HostExecuteData)
+UE_TRACE_EVENT_END()
+
+/**
+ * The data to be traced per frame for a RigVMHost.
+ */
+struct FRigVMTraceExecuteData : FRigVMTraceBaseData
+{
+	FRigVMTraceExecuteData()
+		: FRigVMTraceBaseData()
+		, AbsoluteTime(0)
+		, DeltaTime(0)
+		, WorldId(0)
+		, HostId(0)
+		, OverallCycles(0)
+	{
+	}
+
+	// The sample's absolute time of the host
+	double AbsoluteTime;
+	// The sample's delta time of the host
+	double DeltaTime;
+	// The id of the world this host belongs to
+	uint64 WorldId;
+	// The id of the host
+	uint64 HostId;
+	// An archive storing the work memory of the VM
+	FRigVMTraceArchive WorkMemory;
+	// An archive storing the debug memory of the VM
+	FRigVMTraceArchive DebugMemory;
+	// An archive storing the external memory (variables) of the VM
+	FRigVMTraceArchive ExternalMemory;
+	// The names of the entries being executed
+	FString Entries;
+	// The instructions visited during this sample
+	TArray<int32> VisitedInstructions;
+	// The overall duration in cycles
+	uint64 OverallCycles;
+	// The cycles spent in each instruction
+	TArray<uint64> InstructionCycles;
+	// The content of the debug draw interface
+	FRigVMTraceArchive DrawInterface;
+	// Auxiliary data provided by the host's virtual TraceExecuteData
+	FRigVMTraceArchive HostExecuteData;
+};
+
+/**
+ * The base service to allow tracing of RigVM host instances
+ */
+class FRigVMTrace : public IRewindDebuggerRuntimeExtension
+{
+public:
+	UE_API virtual void RecordingStarted() override;
+	UE_API virtual void RecordingStopped() override;
+
+	static UE_API void TraceRigVMEvaluation(const FRigVMExtendedExecuteContext& InContext);
+};
+
+#define TRACE_RIGVM_EVALUATION(InContext) \
+	FRigVMTrace::TraceRigVMEvaluation(InContext)
+#else
+
+#define TRACE_RIGVM_EVALUATION(...)
+
+#endif
+
+#undef UE_API
diff --git a/Engine/Plugins/Runtime/RigVM/Source/RigVM/Public/RigVMTraceArchive.h b/Engine/Plugins/Runtime/RigVM/Source/RigVM/Public/RigVMTraceArchive.h
new file mode 100644
--- /dev/null
+++ b/Engine/Plugins/Runtime/RigVM/Source/RigVM/Public/RigVMTraceArchive.h
@@ -0,0 +1,124 @@
+// Copyright Epic Games, Inc. All Rights Reserved.
+
+#pragma once
+
+#include "CoreMinimal.h"
+#include "Containers/Array.h"
+#include "Containers/Map.h"
+#include "Containers/ContainerAllocationPolicies.h"
+#include "Serialization/ArchiveUObject.h"
+
+#define UE_API RIGVM_API
+
+/**
+ * The trace archive is used to store compressed data for a RewindDebugger stream.
+ */
+struct FRigVMTraceArchive
+{
+public:
+	UE_API FRigVMTraceArchive();
+
+	struct FHeader
+	{
+		bool bIsCompressed = false;
+		int32 UncompressedSize = INDEX_NONE;
+		int32 CompressedSize = INDEX_NONE;
+		int32 NumVersions = INDEX_NONE;
+		int64 VersionsOffset = INDEX_NONE;
+	};
+	
+	UE_API void Reset();
+	UE_API void Empty();
+	UE_API bool IsPayloadEmpty() const;
+	UE_API void Compress();
+	UE_API bool Decompress();
+
+	UE_API const FHeader& GetHeader() const;
+	UE_API bool IsCompressed() const;
+	UE_API int32 GetUncompressedSize() const;
+	UE_API int32 GetCompressedSize() const;
+
+	UE_API uint8* GetPayloadData();
+	UE_API const uint8* GetPayloadData() const;
+	UE_API int32 PayloadNum() const;
+
+	UE_API uint8* GetOverallData();
+	UE_API const uint8* GetOverallData() const;
+	UE_API int32 OverallNum() const;
+
+	UE_API friend FArchive& operator<<(FArchive& Ar, FRigVMTraceArchive& Data);
+
+protected:
+
+	UE_API void SetOverallBuffer(const TArrayView<const uint8>& InBuffer);
+	UE_API void SetOverallBuffer(const TArrayView<uint8>& InBuffer);
+
+private:
+
+	TArray<uint8, TSizedDefaultAllocator<32>> Buffer;
+
+	FHeader& Header();
+
+	friend class FRigVMTraceArchiveWriter;
+	friend class FRigVMTraceArchiveReader;
+	friend class FRigVMTraceAnalyzer;
+};
+
+/**
+ * The trace archive writer is a wrapping archive used to write into
+ * a trace archive.
+ */
+class FRigVMTraceArchiveWriter : public FArchiveUObject
+{
+public:
+	UE_API FRigVMTraceArchiveWriter( FRigVMTraceArchive& InOutArchive );
+	UE_API ~FRigVMTraceArchiveWriter();
+	UE_API virtual void Serialize( void* V, int64 Length ) override;
+	UE_API virtual int64 Tell() override;
+	UE_API virtual int64 TotalSize() override;
+	UE_API virtual void Seek(int64 InPos) override;
+	using FArchiveUObject::operator<<; // For visibility of the overloads we don't override
+	UE_API void WriteUObject(const UObject* Obj);
+	UE_API virtual FArchive& operator<<(UObject*& Obj) override;
+	UE_API virtual FArchive& operator<<(FName& Value) override;
+	UE_API virtual FArchive& operator<<(FText& Value) override;
+
+protected:
+
+	void WriteVersions();
+	
+	FRigVMTraceArchive& Archive;
+	int64 Offset;
+	TMap<FName,int64> NameToOffset;
+	TMap<const UObject*, uint64> ObjectIdMap;
+
+	static inline constexpr uint8 StoringStringAsString = 0;
+	static inline constexpr uint8 StoringStringAsOffset = 1;
+};
+
+/**
+ * The trace archive reader is a wrapping archive used to read from
+ * a trace archive.
+ */
+class FRigVMTraceArchiveReader : public FRigVMTraceArchiveWriter
+{
+public:
+	
+	UE_API FRigVMTraceArchiveReader( FRigVMTraceArchive& InOutArchive );
+	UE_API virtual void Serialize( void* V, int64 Length ) override;
+	using FArchiveUObject::operator<<; // For visibility of the overloads we don't override
+	void ReadUObject(UObject*& Obj);
+	UE_API virtual FArchive& operator<<(UObject*& Obj) override;
+	UE_API virtual FArchive& operator<<(FName& Value) override;
+	UE_API virtual FArchive& operator<<(FText& Value) override;
+
+protected:
+
+	void ReadVersions();
+
+	TMap<uint64, UObject*> ObjectIdMap;
+	TMap<int64,FName> OffsetToName;
+	TMap<int64,FString> OffsetToString;
+};
+
+#undef UE_API
\ No newline at end of file
diff --git a/Engine/Plugins/Runtime/RigVM/Source/RigVM/RigVM.Build.cs b/Engine/Plugins/Runtime/RigVM/Source/RigVM/RigVM.Build.cs
--- a/Engine/Plugins/Runtime/RigVM/Source/RigVM/RigVM.Build.cs
+++ b/Engine/Plugins/Runtime/RigVM/Source/RigVM/RigVM.Build.cs
@@ -20,6 +20,8 @@ public class RigVM : ModuleRules
 				"DeveloperSettings",
 				"Engine",
 				"Projects",
+				"TraceLog",
+				"RewindDebuggerRuntimeInterface",
 			}
 		);
 
diff --git a/Engine/Plugins/Runtime/RigVM/Source/RigVMEditor/Private/RigVMEditorModule.cpp b/Engine/Plugins/Runtime/RigVM/Source/RigVMEditor/Private/RigVMEditorModule.cpp
--- a/Engine/Plugins/Runtime/RigVM/Source/RigVMEditor/Private/RigVMEditorModule.cpp
+++ b/Engine/Plugins/Runtime/RigVM/Source/RigVMEditor/Private/RigVMEditorModule.cpp
@@ -47,6 +47,7 @@
 #include "ToolMenus.h"
 #include "Widgets/SRigVMSwapFunctionsWidget.h"
 #include "Widgets/SRigVMBulkEditDialog.h"
+#include "Trace/RigVMTraceTrack.h"
 
 DEFINE_LOG_CATEGORY(LogRigVMEditor);
 
@@ -59,6 +60,10 @@ TAutoConsoleVariable<bool> CVarRigVMUseDualEditor(TEXT("RigVM.UseDualEditor"), 0
 
 UE::RigVMEditor::FRigVMEditorOnRequestFindNodeReferencesDelegate FRigVMEditorModule::OnRequestFindNodeReferences;
 
+#if RIGVM_TRACE_ENABLED
+FRewindDebuggerRigVMTrackCreator GRewindDebuggerRigVMTrackCreator;
+#endif
+
 FRigVMEditorModule& FRigVMEditorModule::Get()
 {
 	return FModuleManager::LoadModuleChecked< FRigVMEditorModule >(TEXT("RigVMEditor"));
@@ -148,6 +153,12 @@ void FRigVMEditorModule::StartupModule()
 				}));
 			}
 		}
+
+#if RIGVM_TRACE_ENABLED
+		IModularFeatures::Get().RegisterModularFeature( IRewindDebuggerExtension::ModularFeatureName, &RewindDebuggerForRigVM);
+		IModularFeatures::Get().RegisterModularFeature(TraceServices::ModuleFeatureName, &RigVMTraceModule);
+		IModularFeatures::Get().RegisterModularFeature(RewindDebugger::IRewindDebuggerTrackCreator::ModularFeatureName, &GRewindDebuggerRigVMTrackCreator);
+#endif
 	}
 
 	StartupModuleCommon();
@@ -195,6 +206,12 @@ void FRigVMEditorModule::ShutdownModule()
 			FBlueprintEditorModule& BlueprintEditorModule = FModuleManager::LoadModuleChecked<FBlueprintEditorModule>("Kismet");
 			BlueprintEditorModule.UnregisterVariableCustomization(FProperty::StaticClass(), BlueprintVariableCustomizationHandle);
 #endif
+
+#if RIGVM_TRACE_ENABLED
+			IModularFeatures::Get().UnregisterModularFeature( IRewindDebuggerExtension::ModularFeatureName, &RewindDebuggerForRigVM);
+			IModularFeatures::Get().UnregisterModularFeature(TraceServices::ModuleFeatureName, &RigVMTraceModule);
+			IModularFeatures::Get().UnregisterModularFeature(RewindDebugger::IRewindDebuggerTrackCreator::ModularFeatureName, &GRewindDebuggerRigVMTrackCreator);
+#endif
 		}
 	}
 
diff --git a/Engine/Plugins/Runtime/RigVM/Source/RigVMEditor/Private/Trace/RigVMTraceAnalyzer.cpp b/Engine/Plugins/Runtime/RigVM/Source/RigVMEditor/Private/Trace/RigVMTraceAnalyzer.cpp
new file mode 100644
--- /dev/null
+++ b/Engine/Plugins/Runtime/RigVM/Source/RigVMEditor/Private/Trace/RigVMTraceAnalyzer.cpp
@@ -0,0 +1,87 @@
+// Copyright Epic Games, Inc. All Rights Reserved.
+
+#include "RigVMTraceAnalyzer.h"
+
+#include "HAL/LowLevelMemTracker.h"
+#include "Serialization/MemoryReader.h"
+#include "TraceServices/Model/AnalysisSession.h"
+#include "RigVMTraceProvider.h"
+#include "TraceServices/Utils.h"
+
+#if RIGVM_TRACE_ENABLED
+
+FRigVMTraceAnalyzer::FRigVMTraceAnalyzer(TraceServices::IAnalysisSession& InSession, FRigVMTraceProvider& InProvider)
+	: Session(InSession)
+	, Provider(InProvider)
+{
+}
+
+void FRigVMTraceAnalyzer::OnAnalysisBegin(const FOnAnalysisContext& Context)
+{
+	auto& Builder = Context.InterfaceBuilder;
+
+	Builder.RouteEvent(RouteId_RigVMExecute, "RigVM", "Execute");
+	Builder.RouteEvent(RouteId_RigVMLiterals, "RigVM", "Literals");
+}
+
+bool FRigVMTraceAnalyzer::OnEvent(uint16 RouteId, EStyle Style, const FOnEventContext& Context)
+{
+	LLM_SCOPE_BYNAME(TEXT("Insights/FRigVMTraceAnalyzer"));
+
+	TraceServices::FAnalysisSessionEditScope _(Session);
+
+	const auto& EventData = Context.EventData;
+	switch (RouteId)
+	{
+		case RouteId_RigVMLiterals:
+		{
+			const uint64 HostId = EventData.GetValue<uint64>("HostId");
+			if (Provider.HasConstantData(HostId))
+			{
+				break;
+			}
+				
+			FRigVMTraceConstantData ConstantData;
+			ConstantData.Cycles = EventData.GetValue<uint64>("Cycles");
+			ConstantData.ProfileTime = Context.EventTime.AsSeconds(ConstantData.Cycles);
+			ConstantData.RecordingTime = EventData.GetValue<double>("RecordingTime");
+			ConstantData.RigVMObjectVersion = EventData.GetValue<int32>("RigVMObjectVersion");
+			ConstantData.HostId = HostId;
+			ConstantData.VMHash = EventData.GetValue<uint32>("VMHash");
+			ConstantData.ByteCodeHash = EventData.GetValue<uint32>("ByteCodeHash");
+			ConstantData.LiteralMemory.SetOverallBuffer(EventData.GetArrayView<uint8>("LiteralMemory"));
+			ConstantData.HostConstantData.SetOverallBuffer(EventData.GetArrayView<uint8>("HostConstantData"));
+				
+			Provider.AppendLiterals(ConstantData);
+			break;
+		}
+		case RouteId_RigVMExecute:
+		{
+			FRigVMTraceExecuteData ExecuteData;
+			ExecuteData.Cycles = EventData.GetValue<uint64>("Cycles");
+			ExecuteData.ProfileTime = Context.EventTime.AsSeconds(ExecuteData.Cycles);
+			ExecuteData.RecordingTime = EventData.GetValue<double>("RecordingTime");
+			ExecuteData.AbsoluteTime = EventData.GetValue<double>("AbsoluteTime");
+			ExecuteData.DeltaTime = EventData.GetValue<double>("DeltaTime");
+			ExecuteData.HostId = EventData.GetValue<uint64>("HostId");
+			ExecuteData.WorldId = EventData.GetValue<uint64>("WorldId");
+			ExecuteData.RigVMObjectVersion = EventData.GetValue<int32>("RigVMObjectVersion");
+			ExecuteData.WorkMemory.SetOverallBuffer(EventData.GetArrayView<uint8>("WorkMemory"));
+			ExecuteData.DebugMemory.SetOverallBuffer(EventData.GetArrayView<uint8>("DebugMemory"));
+			ExecuteData.ExternalMemory.SetOverallBuffer(EventData.GetArrayView<uint8>("ExternalMemory"));
+			(void)EventData.GetString("Entries", ExecuteData.Entries);
+			ExecuteData.VisitedInstructions = EventData.GetArrayView<int32>("VisitedInstructions");
+			ExecuteData.OverallCycles = EventData.GetValue<uint64>("OverallCycles");
+			ExecuteData.InstructionCycles = EventData.GetArrayView<uint64>("InstructionCycles");
+			ExecuteData.DrawInterface.SetOverallBuffer(EventData.GetArrayView<uint8>("DrawInterface"));
+			ExecuteData.HostExecuteData.SetOverallBuffer(EventData.GetArrayView<uint8>("HostExecuteData"));
+				
+			Provider.AppendExecute(ExecuteData);
+			break;
+		}
+	}
+
+	return true;
+}
+
+#endif
diff --git a/Engine/Plugins/Runtime/RigVM/Source/RigVMEditor/Private/Trace/RigVMTraceAnalyzer.h b/Engine/Plugins/Runtime/RigVM/Source/RigVMEditor/Private/Trace/RigVMTraceAnalyzer.h
new file mode 100644
--- /dev/null
+++ b/Engine/Plugins/Runtime/RigVM/Source/RigVMEditor/Private/Trace/RigVMTraceAnalyzer.h
@@ -0,0 +1,37 @@
+// Copyright Epic Games, Inc. All Rights Reserved.
+
+#pragma once
+
+#include "Trace/Analyzer.h"
+#include "RigVMTrace.h"
+
+#if RIGVM_TRACE_ENABLED
+
+class FRigVMTraceProvider;
+namespace TraceServices { class IAnalysisSession; }
+
+/**
+ * The analyzer is used to turn unstructured data from an analytics session
+ * into useful bits that can be send to the corresponding provider class.
+ */
+class FRigVMTraceAnalyzer : public UE::Trace::IAnalyzer
+{
+public:
+	FRigVMTraceAnalyzer(TraceServices::IAnalysisSession& InSession, FRigVMTraceProvider& InProvider);
+
+	virtual void OnAnalysisBegin(const FOnAnalysisContext& Context) override;
+	virtual void OnAnalysisEnd() override {}
+	virtual bool OnEvent(uint16 RouteId, EStyle Style, const FOnEventContext& Context) override;
+
+private:
+	enum : uint16
+	{
+		RouteId_RigVMLiterals,
+		RouteId_RigVMExecute,
+	};
+
+	TraceServices::IAnalysisSession& Session;
+	FRigVMTraceProvider& Provider;
+};
+
+#endif
\ No newline at end of file
diff --git a/Engine/Plugins/Runtime/RigVM/Source/RigVMEditor/Private/Trace/RigVMTraceModule.cpp b/Engine/Plugins/Runtime/RigVM/Source/RigVMEditor/Private/Trace/RigVMTraceModule.cpp
new file mode 100644
--- /dev/null
+++ b/Engine/Plugins/Runtime/RigVM/Source/RigVMEditor/Private/Trace/RigVMTraceModule.cpp
@@ -0,0 +1,87 @@
+﻿// Copyright Epic Games, Inc. All Rights Reserved.
+
+#include "Trace/RigVMTraceModule.h"
+#include "RigVMTraceProvider.h"
+#include "RigVMTraceAnalyzer.h"
+
+#if RIGVM_TRACE_ENABLED
+
+#define LOCTEXT_NAMESPACE "RigVMTraceModule"
+
+FName FRigVMTraceModule::ModuleName("RigVM");
+
+void FRigVMTraceModule::GetModuleInfo(TraceServices::FModuleInfo& OutModuleInfo)
+{
+	OutModuleInfo.Name = ModuleName;
+	OutModuleInfo.DisplayName = TEXT("RigVM");
+}
+
+void FRigVMTraceModule::OnAnalysisBegin(TraceServices::IAnalysisSession& InSession)
+{
+	TSharedPtr<FRigVMTraceProvider> RigVMTraceProvider = MakeShared<FRigVMTraceProvider>(InSession);
+	InSession.AddProvider(FRigVMTraceProvider::ProviderName, RigVMTraceProvider);
+	InSession.AddAnalyzer(new FRigVMTraceAnalyzer(InSession, *RigVMTraceProvider));
+}
+
+void FRigVMTraceModule::GetLoggers(TArray<const TCHAR *>& OutLoggers)
+{
+	OutLoggers.Add(TEXT("RigVM"));
+}
+
+void FRigVMTraceModule::GenerateReports(const TraceServices::IAnalysisSession& Session, const TCHAR* CmdLine, const TCHAR* OutputDirectory)
+{
+}
+
+void FRewindDebuggerForRigVM::Update(float DeltaTime, IRewindDebugger* RewindDebugger)
+{
+	if (RewindDebugger->IsPIESimulating() || RewindDebugger->GetRecordingDuration() == 0.0)
+	{
+		return;
+	}
+	
+	if (const TraceServices::IAnalysisSession* Session = RewindDebugger->GetAnalysisSession())
+	{
+		TraceServices::FAnalysisSessionReadScope SessionReadScope(*Session);
+		double CurrentTraceTime = RewindDebugger->CurrentTraceTime();
+		
+		if (CurrentTraceTime != LastTraceTime)
+		{
+			bool bRigVMTraceFound = false;
+
+			const TraceServices::IFrameProvider& FrameProvider = TraceServices::ReadFrameProvider(*Session);
+			TraceServices::FFrame Frame;
+
+			if(FrameProvider.GetFrameFromTime(ETraceFrameType::TraceFrameType_Game, CurrentTraceTime, Frame))
+			{
+				if (const FRigVMTraceProvider* Provider = Session->ReadProvider<FRigVMTraceProvider>(FRigVMTraceProvider::ProviderName))
+				{
+					Provider->EnumerateExecuteTimelines([&bRigVMTraceFound, Frame](uint64 HostId, const FRigVMTraceProvider::ExecuteTimeline& TimelineData)
+						{
+							TimelineData.EnumerateEvents(Frame.StartTime, Frame.EndTime,
+								[&bRigVMTraceFound, HostId](double InStartTime, double InEndTime, uint32 InDepth, const FRigVMTraceExecuteData& ExecuteData)
+								{
+									if (ExecuteData.HostId == HostId)
+									{
+										// todo: we want to make sure to update the host at this time.
+										bRigVMTraceFound = true;
+										return TraceServices::EEventEnumerate::Stop;
+									}
+									return TraceServices::EEventEnumerate::Continue;
+								}
+							);
+						}
+					);
+				};
+			}
+			
+			if (bRigVMTraceFound)
+			{
+				LastTraceTime = CurrentTraceTime;
+			}
+		}
+	}
+}
+
+#undef LOCTEXT_NAMESPACE
+
+#endif
\ No newline at end of file
diff --git a/Engine/Plugins/Runtime/RigVM/Source/RigVMEditor/Private/Trace/RigVMTraceProvider.cpp b/Engine/Plugins/Runtime/RigVM/Source/RigVMEditor/Private/Trace/RigVMTraceProvider.cpp
new file mode 100644
--- /dev/null
+++ b/Engine/Plugins/Runtime/RigVM/Source/RigVMEditor/Private/Trace/RigVMTraceProvider.cpp
@@ -0,0 +1,92 @@
+// Copyright Epic Games, Inc. All Rights Reserved.
+
+#include "RigVMTraceProvider.h"
+#include "ObjectTrace.h"
+
+#if RIGVM_TRACE_ENABLED
+
+FName FRigVMTraceProvider::ProviderName("RigVMTraceProvider");
+
+#define LOCTEXT_NAMESPACE "RigVMTraceProvider"
+
+FRigVMTraceProvider::FRigVMTraceProvider(TraceServices::IAnalysisSession& InSession)
+	: Session(InSession)
+	, ConstantData(InSession.GetLinearAllocator(), 256)
+{
+}
+
+bool FRigVMTraceProvider::HasConstantData(uint64 InObjectId) const
+{
+	Session.ReadAccessCheck();
+	
+	return ObjectIdToConstantData.Contains(InObjectId);
+}
+
+bool FRigVMTraceProvider::ReadExecuteTimeline(uint64 InObjectId, TFunctionRef<void(const ExecuteTimeline&)> Callback) const
+{
+	Session.ReadAccessCheck();
+
+	const uint32* IndexPtr = ObjectIdToExecuteTimelines.Find(InObjectId);
+	if(IndexPtr != nullptr)
+	{
+		if (*IndexPtr < uint32(ExecuteTimelines.Num()))
+		{
+			Callback(*ExecuteTimelines[*IndexPtr]);
+			return true;
+		}
+	}
+
+	return false;
+}
+
+void FRigVMTraceProvider::EnumerateExecuteTimelines(TFunctionRef<void(uint64 OwnerId, const ExecuteTimeline&)> Callback) const
+{
+	Session.ReadAccessCheck();
+
+	for (TTuple<uint64, uint32> Entry : ObjectIdToExecuteTimelines)
+	{
+		if (ExecuteTimelines.IsValidIndex(Entry.Value))
+		{
+			Callback(Entry.Key, *ExecuteTimelines[Entry.Value]);
+		}
+	}
+}
+
+uint32 FRigVMTraceProvider::GetExecuteTimelineIndex(uint64 InObjectId)
+{
+	uint32* IndexPtr = ObjectIdToExecuteTimelines.Find(InObjectId);
+	if(IndexPtr != nullptr)
+	{
+		return *IndexPtr;
+	}
+
+	uint32 Index = ExecuteTimelines.Num();
+	ObjectIdToExecuteTimelines.Add(InObjectId, Index);
+	ExecuteTimelines.Add( MakeShared<TraceServices::TPointTimeline<FRigVMTraceExecuteData>>(Session.GetLinearAllocator()));
+	return Index;
+}
+
+void FRigVMTraceProvider::AppendLiterals(const FRigVMTraceConstantData& InConstantData)
+{
+	Session.WriteAccessCheck();
+
+	if (ObjectIdToConstantData.Contains(InConstantData.HostId))
+	{
+		return;
+	}
+
+	ObjectIdToConstantData.Add(InConstantData.HostId, static_cast<uint32>(ConstantData.Num()));
+	ConstantData.PushBack() = InConstantData;
+}
+
+void FRigVMTraceProvider::AppendExecute(const FRigVMTraceExecuteData& InExecuteData)
+{
+	Session.WriteAccessCheck();
+
+	TSharedPtr<TraceServices::TPointTimeline<FRigVMTraceExecuteData>> Timeline = ExecuteTimelines[GetExecuteTimelineIndex(InExecuteData.HostId)];
+	Timeline->AppendEvent(InExecuteData.ProfileTime, InExecuteData);
+}
+
+#undef LOCTEXT_NAMESPACE
+
+#endif 
\ No newline at end of file
diff --git a/Engine/Plugins/Runtime/RigVM/Source/RigVMEditor/Private/Trace/RigVMTraceProvider.h b/Engine/Plugins/Runtime/RigVM/Source/RigVMEditor/Private/Trace/RigVMTraceProvider.h
new file mode 100644
--- /dev/null
+++ b/Engine/Plugins/Runtime/RigVM/Source/RigVMEditor/Private/Trace/RigVMTraceProvider.h
@@ -0,0 +1,52 @@
+// Copyright Epic Games, Inc. All Rights Reserved.
+
+#pragma once
+
+#include "Engine/EngineBaseTypes.h"
+#include "Model/PointTimeline.h"
+#include "Containers/Map.h"
+#include "Templates/SharedPointer.h"
+#include "Model/IntervalTimeline.h"
+#include "TraceServices/Model/AnalysisSession.h"
+#include "RigVMTrace.h"
+
+#if RIGVM_TRACE_ENABLED
+
+namespace TraceServices { class IAnalysisSession; }
+
+/**
+ * The provider is used to provide data to the world given a corresponding analyzer.
+ * The analyzer receives unstructured data from the session and forwards the data to the provider.
+ * The provider can then provide the data to the target RigVMHost during playback.
+ */
+class FRigVMTraceProvider : public TraceServices::IProvider
+{
+public:
+	typedef TraceServices::ITimeline<FRigVMTraceExecuteData> ExecuteTimeline;
+	
+	static FName ProviderName;
+
+	FRigVMTraceProvider(TraceServices::IAnalysisSession& InSession);
+
+	bool HasConstantData(uint64 InObjectId) const;
+	bool ReadConstantData(uint64 InObjectId, TFunctionRef<void(const FRigVMTraceConstantData&)> Callback) const;
+	void EnumerateConstantData(TFunctionRef<void(uint64 OwnerId, const FRigVMTraceConstantData&)> Callback) const;
+
+	bool ReadExecuteTimeline(uint64 InObjectId, TFunctionRef<void(const ExecuteTimeline&)> Callback) const;
+	void EnumerateExecuteTimelines(TFunctionRef<void(uint64 OwnerId, const ExecuteTimeline&)> Callback) const;
+
+	void AppendLiterals(const FRigVMTraceConstantData& InConstantData);
+	void AppendExecute(const FRigVMTraceExecuteData& InExecuteData);
+
+private:
+	uint32 GetExecuteTimelineIndex(uint64 ObjectId);
+	
+	TraceServices::IAnalysisSession& Session;
+
+	TMap<uint64, uint32> ObjectIdToConstantData;
+	TraceServices::TPagedArray<FRigVMTraceConstantData> ConstantData;
+	TMap<uint64, uint32> ObjectIdToExecuteTimelines;
+	TArray<TSharedRef<TraceServices::TPointTimeline<FRigVMTraceExecuteData>>> ExecuteTimelines;
+};
+
+#endif
\ No newline at end of file
diff --git a/Engine/Plugins/Runtime/RigVM/Source/RigVMEditor/Private/Trace/RigVMTraceTrack.cpp b/Engine/Plugins/Runtime/RigVM/Source/RigVMEditor/Private/Trace/RigVMTraceTrack.cpp
new file mode 100644
--- /dev/null
+++ b/Engine/Plugins/Runtime/RigVM/Source/RigVMEditor/Private/Trace/RigVMTraceTrack.cpp
@@ -0,0 +1,303 @@
+﻿// Copyright Epic Games, Inc. All Rights Reserved.
+
+#include "Trace/RigVMTraceTrack.h"
+#include "RigVMTraceProvider.h"
+#include "RigVMTraceAnalyzer.h"
+#include "DetailLayoutBuilder.h"
+#include "Widgets/Layout/SGridPanel.h"
+#include "Widgets/Input/SNumericEntryBox.h"
+
+#if RIGVM_TRACE_ENABLED
+
+#define LOCTEXT_NAMESPACE "RigVMTraceTrack"
+
+static const FLazyName RigVMName("RigVM");
+
+FRewindDebuggerRigVMExecuteTrack::FRewindDebuggerRigVMExecuteTrack(uint64 InHostId)
+	: HostId(InHostId)
+{
+	Icon = FSlateIcon(TEXT("RigVMEditorStyle"), "ExecutionStack.TabIcon");
+	EventData = MakeShared<SEventTimelineView::FTimelineEventData>();
+
+	// this is temporary
+	SAssignNew(DetailsVerticalBox, SVerticalBox)
+	+ SVerticalBox::Slot()
+	.AutoHeight()
+	.Padding(8)
+	.HAlign(HAlign_Left)
+	.VAlign(VAlign_Top)
+	[
+		SNew(SHorizontalBox)
+		+ SHorizontalBox::Slot()
+		.FillWidth(1)
+		.Padding(0)
+		.HAlign(HAlign_Left)
+		.VAlign(VAlign_Top)
+		[
+			SNew(SGridPanel)
+			.FillColumn(1, 1)
+			
+			+ SGridPanel::Slot(0, 0)
+			.VAlign(VAlign_Center)
+			.HAlign(HAlign_Left)
+			.Padding(2)
+			[
+				SNew(STextBlock)
+				.Text(LOCTEXT("AbsoluteTime", "Absolute Time"))
+				.Font(IDetailLayoutBuilder::GetDetailFont())
+			]
+			+ SGridPanel::Slot(1, 0)
+			.VAlign(VAlign_Center)
+			.HAlign(HAlign_Left)
+			.Padding(2)
+			[
+				SNew(SNumericEntryBox<double>)
+				.MaxFractionalDigits(3)
+				.IsEnabled(false)
+				.Value_Lambda([this]() { return AbsoluteTime; })
+			]
+			+ SGridPanel::Slot(0, 1)
+			.VAlign(VAlign_Center)
+			.HAlign(HAlign_Left)
+			.Padding(2)
+			[
+				SNew(STextBlock)
+				.Text(LOCTEXT("DeltaTime", "Delta Time"))
+				.Font(IDetailLayoutBuilder::GetDetailFont())
+			]
+			+ SGridPanel::Slot(1, 1)
+			.VAlign(VAlign_Center)
+			.HAlign(HAlign_Left)
+			.Padding(2)
+			[
+				SNew(SNumericEntryBox<double>)
+				.MaxFractionalDigits(3)
+				.IsEnabled(false)
+				.Value_Lambda([this]() { return DeltaTime; })
+			]
+		]
+	];
+}
+
+bool FRewindDebuggerRigVMExecuteTrack::UpdateInternal()
+{
+	TRACE_CPUPROFILER_EVENT_SCOPE(FRewindDebuggerRigVMExecuteTrack::UpdateInternal);
+	IRewindDebugger* RewindDebugger = IRewindDebugger::Instance();
+	
+	const TRange<double> RecordingTimeRange = RewindDebugger->GetCurrentViewRange();
+	const double RecordingStartTime = RecordingTimeRange.GetLowerBoundValue();
+	const TRange<double> TraceTimeRange = RewindDebugger->GetCurrentTraceRange();
+	const double TraceStartTime = TraceTimeRange.GetLowerBoundValue();
+	const double TraceEndTime = TraceTimeRange.GetUpperBoundValue();
+
+	auto TraceToRecordingTime = [RecordingStartTime, TraceStartTime](const double& InTraceTime) -> double
+	{
+		return (InTraceTime - TraceStartTime) + RecordingStartTime;
+	};
+	
+	const TraceServices::IAnalysisSession* AnalysisSession = RewindDebugger->GetAnalysisSession();
+	const FRigVMTraceProvider* Provider = AnalysisSession->ReadProvider<FRigVMTraceProvider>(FRigVMTraceProvider::ProviderName);
+	
+	if(EventUpdateRequested > 0 && Provider)
+	{
+		EventUpdateRequested = 0;
+	
+		EventData->Points.SetNum(0,EAllowShrinking::No);
+		EventData->Windows.SetNum(1,EAllowShrinking::No);
+
+		EventData->Windows[0].Color = FLinearColor::White;
+		EventData->Windows[0].TimeStart = DBL_MAX;
+		EventData->Windows[0].TimeEnd = DBL_MIN;
+
+		TraceServices::FAnalysisSessionReadScope SessionReadScope(*AnalysisSession);
+	
+		Provider->ReadExecuteTimeline(HostId, [this, TraceStartTime, TraceEndTime, TraceToRecordingTime](const FRigVMTraceProvider::ExecuteTimeline& InTimeline)
+		{
+			InTimeline.EnumerateEvents(TraceStartTime, TraceEndTime, [this, TraceToRecordingTime](double InStartTime, double InEndTime, uint32 InDepth, const FRigVMTraceExecuteData& ExecuteData)
+			{
+				if (ExecuteData.HostId == HostId)
+				{
+					const double StartRecordingTime = TraceToRecordingTime(InStartTime);
+					EventData->Windows[0].TimeStart = FMath::Min(StartRecordingTime, EventData->Windows[0].TimeStart);
+					EventData->Windows[0].TimeEnd = FMath::Max(StartRecordingTime, EventData->Windows[0].TimeEnd);
+					//EventData->Points.Add({TraceToRecordingTime(InStartTime), FText(), FText(), FLinearColor::White});
+				}
+				return TraceServices::EEventEnumerate::Continue;
+			});
+		});
+	}
+
+	double CurrentScrubTime = IRewindDebugger::Instance()->CurrentTraceTime();
+	if (PreviousScrubTime != CurrentScrubTime && Provider)
+	{
+		PreviousScrubTime = CurrentScrubTime;
+
+		const TraceServices::IFrameProvider& FramesProvider = TraceServices::ReadFrameProvider(*AnalysisSession);
+		TraceServices::FAnalysisSessionReadScope SessionReadScope(*AnalysisSession);
+
+		TraceServices::FFrame MarkerFrame;
+		if (FramesProvider.GetFrameFromTime(TraceFrameType_Game, CurrentScrubTime, MarkerFrame))
+		{
+			Provider->ReadExecuteTimeline(HostId,
+				[this, &MarkerFrame](const FRigVMTraceProvider::ExecuteTimeline& InTimeline)
+				{
+					InTimeline.EnumerateEvents(MarkerFrame.StartTime, MarkerFrame.EndTime,
+						[this](double InStartTime, double InEndTime, uint32 InDepth, const FRigVMTraceExecuteData& InExecuteData)
+						{
+							AbsoluteTime = InExecuteData.AbsoluteTime;
+							DeltaTime = InExecuteData.DeltaTime;
+							return TraceServices::EEventEnumerate::Stop;
+						}
+					);
+				}
+			);
+		}
+
+	}
+
+	EventUpdateRequested++;
+	
+	constexpr bool bChanged = false;
+	return bChanged;
+}
+
+TSharedPtr<SWidget> FRewindDebuggerRigVMExecuteTrack::GetDetailsViewInternal()
+{
+	return DetailsVerticalBox;
+}
+
+TSharedPtr<SWidget> FRewindDebuggerRigVMExecuteTrack::GetTimelineViewInternal()
+{
+	return SNew(SEventTimelineView)
+		.ViewRange_Lambda([]() { return IRewindDebugger::Instance()->GetCurrentViewRange(); })
+		.EventData_Raw(this, &FRewindDebuggerRigVMExecuteTrack::GetEventData);
+}
+
+TSharedPtr<SEventTimelineView::FTimelineEventData> FRewindDebuggerRigVMExecuteTrack::GetEventData() const
+{
+	if (!EventData.IsValid())
+	{
+		EventData = MakeShared<SEventTimelineView::FTimelineEventData>();
+	}
+
+	EventUpdateRequested++;
+
+	return EventData;
+}
+
+FRewindDebuggerRigVMTrack::FRewindDebuggerRigVMTrack(uint64 InObjectId)
+: HostId(InObjectId)
+{
+	Icon = FSlateIcon(TEXT("RigVMEditorStyle"), "ExecutionStack.TabIcon");
+}
+
+TSharedPtr<SWidget> FRewindDebuggerRigVMTrack::GetDetailsViewInternal() 
+{
+	return nullptr;
+}
+
+bool FRewindDebuggerRigVMTrack::UpdateInternal()
+{
+	TRACE_CPUPROFILER_EVENT_SCOPE(FRewindDebuggerRigVMTrack::UpdateInternal);
+	IRewindDebugger* RewindDebugger = IRewindDebugger::Instance();
+	
+	const TRange<double> TraceTimeRange = RewindDebugger->GetCurrentTraceRange();
+	const double TraceStartTime = TraceTimeRange.GetLowerBoundValue();
+	const double TraceEndTime = TraceTimeRange.GetUpperBoundValue();
+	
+	const TraceServices::IAnalysisSession* AnalysisSession = RewindDebugger->GetAnalysisSession();
+	const FRigVMTraceProvider* Provider = AnalysisSession->ReadProvider<FRigVMTraceProvider>(FRigVMTraceProvider::ProviderName);
+	
+	bool bChanged = false;
+	
+	if(Provider)
+	{
+		TArray<uint64> UniqueTrackIds;
+
+		TraceServices::FAnalysisSessionReadScope SessionReadScope(*AnalysisSession);
+
+		Provider->ReadExecuteTimeline(HostId, [this, TraceStartTime, TraceEndTime, Provider, AnalysisSession, &UniqueTrackIds](const FRigVMTraceProvider::ExecuteTimeline& InTimeline)
+			{
+				InTimeline.EnumerateEvents(TraceStartTime, TraceEndTime, [&UniqueTrackIds](double InStartTime, double InEndTime, uint32 InDepth, const FRigVMTraceExecuteData& InExecuteData)
+					{
+						UniqueTrackIds.AddUnique(InExecuteData.HostId);
+						return TraceServices::EEventEnumerate::Continue;
+					}
+				);
+			}
+		);
+
+		UniqueTrackIds.StableSort();
+		const int32 TrackCount = UniqueTrackIds.Num();
+
+		if (Children.Num() != TrackCount)
+		{
+			bChanged = true;
+		}
+
+		Children.SetNum(UniqueTrackIds.Num());
+		for (int i = 0; i < TrackCount; i++)
+		{
+			if (!Children[i].IsValid() || Children[i].Get()->GetAssociatedObjectId().GetMainId() != UniqueTrackIds[i])
+			{
+				Children[i] = MakeShared<FRewindDebuggerRigVMExecuteTrack>(UniqueTrackIds[i]);
+				bChanged = true;
+			}
+
+			if (Children[i]->Update())
+			{
+				bChanged = true;
+			}
+		}
+	}
+	
+	return bChanged;
+}
+
+TConstArrayView<TSharedPtr<RewindDebugger::FRewindDebuggerTrack>> FRewindDebuggerRigVMTrack::GetChildrenInternal(
+	TArray<TSharedPtr<FRewindDebuggerTrack>>& OutTracks) const
+{
+	return MakeConstArrayView<TSharedPtr<FRewindDebuggerTrack>>(Children.GetData(), Children.Num());
+}
+
+FName FRewindDebuggerRigVMTrackCreator::GetTargetTypeNameInternal() const
+{
+	static const FLazyName RigVMHostName("RigVMHost");
+	return RigVMHostName;
+}
+	
+FName FRewindDebuggerRigVMTrackCreator::GetNameInternal() const
+{
+	return RigVMName;
+}
+
+void FRewindDebuggerRigVMTrackCreator::GetTrackTypesInternal(TArray<RewindDebugger::FRewindDebuggerTrackType>& Types) const
+{
+	Types.Add({RigVMName, LOCTEXT("RigVM", "RigVM")});
+}
+
+TSharedPtr<RewindDebugger::FRewindDebuggerTrack> FRewindDebuggerRigVMTrackCreator::CreateTrackInternal(const RewindDebugger::FObjectId& InObjectId) const
+{
+	return MakeShared<FRewindDebuggerRigVMTrack>(InObjectId.GetMainId());
+}
+
+bool FRewindDebuggerRigVMTrackCreator::HasDebugInfoInternal(const RewindDebugger::FObjectId& InObjectId) const
+{
+	TRACE_CPUPROFILER_EVENT_SCOPE(FRewindDebuggerRigVMTrack::HasDebugInfoInternal);
+	const TraceServices::IAnalysisSession* AnalysisSession = IRewindDebugger::Instance()->GetAnalysisSession();
+	
+	TraceServices::FAnalysisSessionReadScope SessionReadScope(*AnalysisSession);
+	bool bHasData = false;
+	if (const FRigVMTraceProvider* Provider = AnalysisSession->ReadProvider<FRigVMTraceProvider>(FRigVMTraceProvider::ProviderName))
+	{
+		Provider->ReadExecuteTimeline(InObjectId.GetMainId(), [&bHasData](const FRigVMTraceProvider::ExecuteTimeline& InTimeline)
+	 	{
+	 		bHasData = true;
+	 	});
+	}
+	return bHasData;
+}
+
+#undef LOCTEXT_NAMESPACE
+
+#endif
\ No newline at end of file
diff --git a/Engine/Plugins/Runtime/RigVM/Source/RigVMEditor/Public/RigVMEditorModule.h b/Engine/Plugins/Runtime/RigVM/Source/RigVMEditor/Public/RigVMEditorModule.h
--- a/Engine/Plugins/Runtime/RigVM/Source/RigVMEditor/Public/RigVMEditorModule.h
+++ b/Engine/Plugins/Runtime/RigVM/Source/RigVMEditor/Public/RigVMEditorModule.h
@@ -16,6 +16,10 @@
 #include "Toolkits/AssetEditorToolkit.h"
 #include "Kismet2/StructureEditorUtils.h"
 #include "Modules/ModuleManager.h"
+#if RIGVM_TRACE_ENABLED
+#include "RigVMTrace.h"
+#include "Trace/RigVMTraceModule.h"
+#endif
 
 #define UE_API RIGVMEDITOR_API
 
@@ -188,6 +192,11 @@ protected:
 	FDelegateHandle ReconstructAllNodesDelegateHandle;
 	FDelegateHandle BlueprintVariableCustomizationHandle;
 
+#if RIGVM_TRACE_ENABLED
+	FRewindDebuggerForRigVM RewindDebuggerForRigVM;
+	FRigVMTraceModule RigVMTraceModule;
+#endif
+	
 private:
 
 	/** StaticStruct is not safe on shutdown, so we cache the name, and use this to unregister on shut down */
diff --git a/Engine/Plugins/Runtime/RigVM/Source/RigVMEditor/Public/Trace/RigVMTraceModule.h b/Engine/Plugins/Runtime/RigVM/Source/RigVMEditor/Public/Trace/RigVMTraceModule.h
new file mode 100644
--- /dev/null
+++ b/Engine/Plugins/Runtime/RigVM/Source/RigVMEditor/Public/Trace/RigVMTraceModule.h
@@ -0,0 +1,55 @@
+// Copyright Epic Games, Inc. All Rights Reserved.
+
+#pragma once
+
+#include "CoreMinimal.h"
+#include "IRewindDebugger.h"
+#include "IRewindDebuggerExtension.h"
+#include "RewindDebuggerTrack.h"
+#include "IRewindDebuggerTrackCreator.h"
+#include "TraceServices/ModuleService.h"
+#include "SEventTimelineView.h"
+#include "RigVMTrace.h"
+
+#if RIGVM_TRACE_ENABLED
+
+#define UE_API RIGVMEDITOR_API
+
+/**
+ * The trace module implements the rewind debugger extension to record traces
+ */
+class FRigVMTraceModule : public TraceServices::IModule
+{
+public:
+	virtual ~FRigVMTraceModule()
+	{
+	}
+	
+	// TraceServices::IModule interface
+	UE_API virtual void GetModuleInfo(TraceServices::FModuleInfo& OutModuleInfo) override;
+	UE_API virtual void OnAnalysisBegin(TraceServices::IAnalysisSession& Session) override;
+	UE_API virtual void GetLoggers(TArray<const TCHAR *>& OutLoggers) override;
+	UE_API virtual void GenerateReports(const TraceServices::IAnalysisSession& Session, const TCHAR* CmdLine, const TCHAR* OutputDirectory) override;
+	virtual const TCHAR* GetCommandLineArgument() override { return TEXT("rigvm"); }
+
+private:
+	static FName ModuleName;
+};
+
+class FRewindDebuggerForRigVM : public IRewindDebuggerExtension
+{
+public:
+	FRewindDebuggerForRigVM() {}
+	virtual ~FRewindDebuggerForRigVM() {}
+
+	UE_API virtual void Update(float DeltaTime, IRewindDebugger* RewindDebugger) override;
+	virtual FString GetName() override { return TEXT("RigVMDebugger"); }
+
+private:
+
+	double LastTraceTime = 0.0;
+};
+
+#undef UE_API
+
+#endif
\ No newline at end of file
diff --git a/Engine/Plugins/Runtime/RigVM/Source/RigVMEditor/Public/Trace/RigVMTraceTrack.h b/Engine/Plugins/Runtime/RigVM/Source/RigVMEditor/Public/Trace/RigVMTraceTrack.h
new file mode 100644
--- /dev/null
+++ b/Engine/Plugins/Runtime/RigVM/Source/RigVMEditor/Public/Trace/RigVMTraceTrack.h
@@ -0,0 +1,101 @@
+// Copyright Epic Games, Inc. All Rights Reserved.
+
+#pragma once
+
+#include "CoreMinimal.h"
+#include "IRewindDebugger.h"
+#include "RewindDebuggerTrack.h"
+#include "IRewindDebuggerTrackCreator.h"
+#include "TraceServices/ModuleService.h"
+#include "Widgets/SBoxPanel.h"
+#include "SEventTimelineView.h"
+#include "RigVMTrace.h"
+
+#if RIGVM_TRACE_ENABLED
+
+#define UE_API RIGVMEDITOR_API
+
+/**
+ * The execute track advertises the relevant data per frame for a given RigVM host.
+ */
+class FRewindDebuggerRigVMExecuteTrack : public RewindDebugger::FRewindDebuggerTrack
+{
+public:
+	FRewindDebuggerRigVMExecuteTrack(uint64 InHostId);
+
+private:
+	
+	virtual bool UpdateInternal() override;
+	virtual TSharedPtr<SWidget> GetDetailsViewInternal() override;
+	virtual TSharedPtr<SWidget> GetTimelineViewInternal() override;
+
+	virtual FSlateIcon GetIconInternal() override
+	{
+		return Icon;
+	}
+
+	virtual uint64 GetObjectIdInternal() const override
+	{
+		return HostId;
+	}
+
+	virtual FName GetNameInternal() const override { return "RigVMExecute"; }
+	virtual FText GetDisplayNameInternal() const override { return NSLOCTEXT("RewindDebugger", "RigVMExecuteTrackName", "Execute"); }
+
+	TSharedPtr<SEventTimelineView::FTimelineEventData> GetEventData() const;
+	
+	mutable TSharedPtr<SEventTimelineView::FTimelineEventData> EventData;
+	mutable int EventUpdateRequested = 0;
+	
+	uint64 HostId;
+	FSlateIcon Icon;
+	FText TrackName;
+
+	double PreviousScrubTime = 0.0f;
+	double AbsoluteTime = 0.0;
+	double DeltaTime = 0.0;
+
+	TSharedPtr<SVerticalBox> DetailsVerticalBox;
+};
+
+
+/**
+ * The RigVM rewind debugger track is the main track / parent track for all
+ * other RigVM related tracks
+ */
+class FRewindDebuggerRigVMTrack : public RewindDebugger::FRewindDebuggerTrack
+{
+public:
+	FRewindDebuggerRigVMTrack(uint64 InObjectId);
+private:
+	virtual bool UpdateInternal() override;
+	virtual TSharedPtr<SWidget> GetDetailsViewInternal() override;
+
+	virtual FSlateIcon GetIconInternal() override { return Icon; }
+	virtual FName GetNameInternal() const override { return "RigVM"; }
+	virtual FText GetDisplayNameInternal() const override { return NSLOCTEXT("RewindDebugger", "RigVMTrackName", "RigVM"); }
+	virtual uint64 GetObjectIdInternal() const override { return HostId; }
+	virtual TConstArrayView<TSharedPtr<FRewindDebuggerTrack>> GetChildrenInternal(TArray<TSharedPtr<FRewindDebuggerTrack>>& OutTracks) const override;
+
+	FSlateIcon Icon;
+	uint64 HostId;
+
+	TArray<TSharedPtr<FRewindDebuggerTrack>> Children;
+};
+
+/**
+ * Factory class to create RigVM rewind debugger tracks
+ */
+class FRewindDebuggerRigVMTrackCreator : public RewindDebugger::IRewindDebuggerTrackCreator
+{
+private:
+	virtual FName GetTargetTypeNameInternal() const;
+	virtual FName GetNameInternal() const override;
+	virtual void GetTrackTypesInternal(TArray<RewindDebugger::FRewindDebuggerTrackType>& Types) const override;
+	virtual TSharedPtr<RewindDebugger::FRewindDebuggerTrack> CreateTrackInternal(const RewindDebugger::FObjectId& InObjectId) const override;
+	virtual bool HasDebugInfoInternal(const RewindDebugger::FObjectId& InObjectId) const override;
+};
+
+#undef UE_API
+
+#endif
diff --git a/Engine/Plugins/Runtime/RigVM/Source/RigVMEditor/RigVMEditor.Build.cs b/Engine/Plugins/Runtime/RigVM/Source/RigVMEditor/RigVMEditor.Build.cs
--- a/Engine/Plugins/Runtime/RigVM/Source/RigVMEditor/RigVMEditor.Build.cs
+++ b/Engine/Plugins/Runtime/RigVM/Source/RigVMEditor/RigVMEditor.Build.cs
@@ -39,6 +39,11 @@ public class RigVMEditor : ModuleRules
                 "EditorFramework",
                 "AnimationEditorWidgets",
                 "WidgetRegistration",
+                "TraceAnalysis",
+                "TraceLog",
+                "TraceServices",
+                "RewindDebuggerInterface",
+                "GameplayInsights",
             }
 		);
         
