RigVM: Add additional nullptr check for function list

#rnx

diff --git a/Engine/Plugins/Runtime/RigVM/Source/RigVM/Private/RigVMCore/RigVM.cpp b/Engine/Plugins/Runtime/RigVM/Source/RigVM/Private/RigVMCore/RigVM.cpp
--- a/Engine/Plugins/Runtime/RigVM/Source/RigVM/Private/RigVMCore/RigVM.cpp
+++ b/Engine/Plugins/Runtime/RigVM/Source/RigVM/Private/RigVMCore/RigVM.cpp
@@ -428,6 +428,11 @@ void URigVM::RefreshArgumentNameCaches()
 			const FRigVMExecuteOp& Op = ByteCodeStorage.GetOpAt<FRigVMExecuteOp>(Instruction);
 			if (Functions.IsValidIndex(Op.FunctionIndex))
 			{
+				if (!Functions[Op.FunctionIndex])
+				{
+					continue;
+				}
+				
 				if (const FRigVMDispatchFactory* Factory = Functions[Op.FunctionIndex]->Factory)
 				{
 					if(Factories.IsEmpty())
