RigVM: Tracking of required plugins

#rnx

#virtualized

diff --git a/Engine/Plugins/Animation/ControlRig/Source/ControlRig/Private/ControlRig.cpp b/Engine/Plugins/Animation/ControlRig/Source/ControlRig/Private/ControlRig.cpp
--- a/Engine/Plugins/Animation/ControlRig/Source/ControlRig/Private/ControlRig.cpp
+++ b/Engine/Plugins/Animation/ControlRig/Source/ControlRig/Private/ControlRig.cpp
@@ -3894,6 +3894,28 @@ void UControlRig::OnHierarchyTransformUndoRedo(URigHierarchy* InHierarchy, const
 	}
 }
 
+void UControlRig::GenerateRequiredPluginsData(FRigVMExtendedExecuteContext& Context)
+{
+	Super::GenerateRequiredPluginsData(Context);
+
+	if (const URigHierarchy* Hierarchy = GetHierarchy())
+	{
+		TMap<FString,FString> AllPluginPrefixes;
+		for (int32 Index = 0; Index < Hierarchy->ElementComponents.Num(); Index++)
+		{
+			if(Hierarchy->ElementComponents[Index].IsValid())
+			{
+				const FRigBaseComponent* Component = Hierarchy->ElementComponents[Index].GetPtr<FRigBaseComponent>();
+				const FString PluginName = RigVMTypeUtils::GetPluginName(Component->GetScriptStruct(), AllPluginPrefixes);
+				if (!PluginName.IsEmpty())
+				{
+					RequiredPlugins.AddUnique(PluginName);
+				}
+			}
+		}
+	}
+}
+
 FRigPhysicsSimulationBase* UControlRig::GetPhysicsSimulation(const FRigComponentKey& InComponentKey)
 {
 	TSharedPtr<FRigPhysicsSimulationBase>* PhysicsSimulation = PhysicsSimulations.Find(InComponentKey);
diff --git a/Engine/Plugins/Animation/ControlRig/Source/ControlRig/Private/ModularRig.cpp b/Engine/Plugins/Animation/ControlRig/Source/ControlRig/Private/ModularRig.cpp
--- a/Engine/Plugins/Animation/ControlRig/Source/ControlRig/Private/ModularRig.cpp
+++ b/Engine/Plugins/Animation/ControlRig/Source/ControlRig/Private/ModularRig.cpp
@@ -939,6 +939,26 @@ TArray<FName> UModularRig::GetModuleNames() const
 	return Names;
 }
 
+void UModularRig::GenerateRequiredPluginsData(FRigVMExtendedExecuteContext& Context)
+{
+	Super::GenerateRequiredPluginsData(Context);
+
+	ForEachModule([this](const FRigModuleInstance* Module) -> bool
+	{
+		if(UControlRig* ModuleRig = Module->GetRig())
+		{
+			ModuleRig->GenerateRequiredPluginsData(ModuleRig->GetRigVMExtendedExecuteContext());
+
+			const TArray<FString>& RequiredPluginsForModule = ModuleRig->GetRequiredPlugins();
+			for (const FString& RequiredPluginForModule : RequiredPluginsForModule)
+			{
+				RequiredPlugins.AddUnique(RequiredPluginForModule);
+			}
+		}
+		return true;
+	});
+}
+
 FRigModuleInstance* UModularRig::AddModuleInstance(const FName& InModuleName, TSubclassOf<UControlRig> InModuleClass, const FRigModuleInstance* InParent,
                                                    const FRigElementKeyRedirector::FKeyMap& InConnectionMap, const FControlRigOverrideContainer& InConfigValues) 
 {
diff --git a/Engine/Plugins/Animation/ControlRig/Source/ControlRig/Public/ControlRig.h b/Engine/Plugins/Animation/ControlRig/Source/ControlRig/Public/ControlRig.h
--- a/Engine/Plugins/Animation/ControlRig/Source/ControlRig/Public/ControlRig.h
+++ b/Engine/Plugins/Animation/ControlRig/Source/ControlRig/Public/ControlRig.h
@@ -742,6 +742,8 @@ protected:
 
 	FControlRigOverrideAssetsChanged OverrideAssetsChangedDelegate;
 
+ 	CONTROLRIG_API virtual void GenerateRequiredPluginsData(FRigVMExtendedExecuteContext& Context) override;
+	
 private:
 
 #if WITH_EDITORONLY_DATA
diff --git a/Engine/Plugins/Animation/ControlRig/Source/ControlRig/Public/ModularRig.h b/Engine/Plugins/Animation/ControlRig/Source/ControlRig/Public/ModularRig.h
--- a/Engine/Plugins/Animation/ControlRig/Source/ControlRig/Public/ModularRig.h
+++ b/Engine/Plugins/Animation/ControlRig/Source/ControlRig/Public/ModularRig.h
@@ -337,6 +337,10 @@ public:
 		}
 		return FModuleInstanceHandle();
 	}
+	
+protected:
+
+	CONTROLRIG_API virtual void GenerateRequiredPluginsData(FRigVMExtendedExecuteContext& Context) override;
 
 private:
 
diff --git a/Engine/Plugins/Experimental/UAF/UAF/Source/UAFUncookedOnly/Private/AnimNextRigVMAssetEditorData.cpp b/Engine/Plugins/Experimental/UAF/UAF/Source/UAFUncookedOnly/Private/AnimNextRigVMAssetEditorData.cpp
--- a/Engine/Plugins/Experimental/UAF/UAF/Source/UAFUncookedOnly/Private/AnimNextRigVMAssetEditorData.cpp
+++ b/Engine/Plugins/Experimental/UAF/UAF/Source/UAFUncookedOnly/Private/AnimNextRigVMAssetEditorData.cpp
@@ -1114,6 +1114,7 @@ void UAnimNextRigVMAssetEditorData::RecompileVM()
 		// Initialize right away, in packaged builds we initialize during PostLoad
 		Asset->VM->Initialize(Asset->ExtendedExecuteContext);
 		Asset->GenerateUserDefinedDependenciesData(Asset->ExtendedExecuteContext);
+		Asset->GenerateRequiredPluginsData(Asset->ExtendedExecuteContext);
 
 		// Notable difference with vanilla RigVM host behavior - we init the VM here at the moment as we only have one 'instance'
 		Asset->InitializeVM(FRigUnit_AnimNextBeginExecution::EventName);
diff --git a/Engine/Plugins/Runtime/RigVM/Source/RigVM/Private/RigVMCore/RigVM.cpp b/Engine/Plugins/Runtime/RigVM/Source/RigVM/Private/RigVMCore/RigVM.cpp
--- a/Engine/Plugins/Runtime/RigVM/Source/RigVM/Private/RigVMCore/RigVM.cpp
+++ b/Engine/Plugins/Runtime/RigVM/Source/RigVM/Private/RigVMCore/RigVM.cpp
@@ -3132,6 +3132,44 @@ TArray<const UObject*> URigVM::GetUserDefinedDependencies(const TArray<const FRi
 	return Dependencies;
 }
 
+void URigVM::GetRequiredPlugins(TArray<FString>& OutPlugins, TMap<FString, FString>& OutPluginPrefixes)
+{
+	GetDefaultMemoryByType(ERigVMMemoryType::Literal)->GetRequiredPlugins(OutPlugins, OutPluginPrefixes);
+	GetDefaultMemoryByType(ERigVMMemoryType::Work)->GetRequiredPlugins(OutPlugins, OutPluginPrefixes);
+
+	ResolveFunctionsIfRequired();
+
+#if WITH_EDITOR
+	TUniquePtr<UE::TScopeLock<FTransactionallySafeCriticalSection>> ResolveFunctionsScopeLock;
+	if(!LocalizedRegistry.IsValid())
+	{
+		ResolveFunctionsScopeLock = MakeUnique<UE::TScopeLock<FTransactionallySafeCriticalSection>>(ResolveFunctionsMutex);
+	}
+#endif
+
+	for(const FRigVMFunction* Function : FunctionsStorage)
+	{
+		if (!Function)
+		{
+			continue;
+		}
+		
+		FString PluginName;
+		if (Function->Struct)
+		{
+			PluginName = RigVMTypeUtils::GetPluginName(Function->Struct, OutPluginPrefixes);
+		}
+		else if (Function->Factory && Function->Factory->GetScriptStruct())
+		{
+			PluginName = RigVMTypeUtils::GetPluginName(Function->Factory->GetScriptStruct(), OutPluginPrefixes);
+		}
+		if (!PluginName.IsEmpty())
+		{
+			OutPlugins.AddUnique(PluginName);
+		}
+	}
+}
+
 void URigVM::SetupInstructionTracking(FRigVMExtendedExecuteContext& Context, int32 InInstructionCount)
 {
 #if WITH_EDITOR
diff --git a/Engine/Plugins/Runtime/RigVM/Source/RigVM/Private/RigVMCore/RigVMMemoryStorageStruct.cpp b/Engine/Plugins/Runtime/RigVM/Source/RigVM/Private/RigVMCore/RigVMMemoryStorageStruct.cpp
--- a/Engine/Plugins/Runtime/RigVM/Source/RigVM/Private/RigVMCore/RigVMMemoryStorageStruct.cpp
+++ b/Engine/Plugins/Runtime/RigVM/Source/RigVM/Private/RigVMCore/RigVMMemoryStorageStruct.cpp
@@ -69,6 +69,32 @@ void FRigVMMemoryStorageStruct::AddStructReferencedObjects(FReferenceCollector&
 
 void FRigVMMemoryStorageStruct::GetUserDefinedDependencies(TArray<const UObject*>& OutDependencies) const
 {
+	ForEachObjectDependency([&OutDependencies](const UObject* Dependency)
+	{
+		OutDependencies.AddUnique(Dependency);
+	});
+}
+
+void FRigVMMemoryStorageStruct::GetRequiredPlugins(TArray<FString>& OutPlugins, TMap<FString,FString>& OutPluginPrefixes) const
+{
+	TMap<FString,FString> PluginPrefixes;
+	ForEachObjectDependency([&OutPlugins, &PluginPrefixes](const UObject* Dependency)
+	{
+		const FString PluginName = RigVMTypeUtils::GetPluginName(Dependency, PluginPrefixes);
+		if (!PluginName.IsEmpty())
+		{
+			OutPlugins.AddUnique(PluginName);
+		}
+	});
+}
+
+void FRigVMMemoryStorageStruct::ForEachObjectDependency(const TFunction<void(const UObject*)>& InCallBack) const
+{
+	if (!InCallBack)
+	{
+		return;
+	}
+	
 	const TArray<const FProperty*>& Properties = GetProperties();
 	for (const FProperty* Property : Properties)
 	{
@@ -81,21 +107,21 @@ void FRigVMMemoryStorageStruct::GetUserDefinedDependencies(TArray<const UObject*
 		{
 			if (const UUserDefinedStruct* UserDefinedStruct = Cast<UUserDefinedStruct>(StructProperty->Struct))
 			{
-				OutDependencies.AddUnique(UserDefinedStruct);
+				InCallBack(UserDefinedStruct);
 			}
 		}
 		else if (const FEnumProperty* EnumProperty = CastField<FEnumProperty>(PropertyToVisit))
 		{
 			if (const UUserDefinedEnum* UserDefinedEnum = Cast<UUserDefinedEnum>(EnumProperty->GetEnum()))
 			{
-				OutDependencies.AddUnique(UserDefinedEnum);
+				InCallBack(UserDefinedEnum);
 			}
 		}
 		else if (const FByteProperty* ByteProperty = CastField<FByteProperty>(PropertyToVisit))
 		{
 			if (const UUserDefinedEnum* UserDefinedEnum = Cast<UUserDefinedEnum>(ByteProperty->Enum))
 			{
-				OutDependencies.AddUnique(UserDefinedEnum);
+				InCallBack(UserDefinedEnum);
 			}
 		}
 		else if (const FObjectProperty* ObjectProperty = CastField<FObjectProperty>(PropertyToVisit))
@@ -106,7 +132,7 @@ void FRigVMMemoryStorageStruct::GetUserDefinedDependencies(TArray<const UObject*
 				const UObject* Dependency = ObjectProperty->GetObjectPropertyValue_InContainer(Value.GetMemory());
 				if (Dependency && FInternalUObjectBaseUtilityIsValidFlagsChecker::CheckObjectValidBasedOnItsFlags(Dependency))
 				{
-					OutDependencies.AddUnique(Dependency);
+					InCallBack(Dependency);
 				}
 			}
 			else if (const FArrayProperty* ArrayProperty = CastField<FArrayProperty>(Property)) // if we are within an array property
@@ -119,7 +145,7 @@ void FRigVMMemoryStorageStruct::GetUserDefinedDependencies(TArray<const UObject*
 						UObject* Dependency = InnerObjectProperty->GetObjectPropertyValue(ArrayHelper.GetRawPtr(ElementIndex));
 						if (Dependency && FInternalUObjectBaseUtilityIsValidFlagsChecker::CheckObjectValidBasedOnItsFlags(Dependency))
 						{
-							OutDependencies.AddUnique(Dependency);
+							InCallBack(Dependency);
 						}
 					}
 				}
diff --git a/Engine/Plugins/Runtime/RigVM/Source/RigVM/Private/RigVMCore/RigVMRegistry.cpp b/Engine/Plugins/Runtime/RigVM/Source/RigVM/Private/RigVMCore/RigVMRegistry.cpp
--- a/Engine/Plugins/Runtime/RigVM/Source/RigVM/Private/RigVMCore/RigVMRegistry.cpp
+++ b/Engine/Plugins/Runtime/RigVM/Source/RigVM/Private/RigVMCore/RigVMRegistry.cpp
@@ -904,21 +904,7 @@ bool FRigVMRegistry_NoLock::OnPluginLoaded_NoLock(IPlugin& InPlugin)
 
 TArray<FString> FRigVMRegistry_NoLock::GetPluginModulePrefixes(const IPlugin& InPlugin)
 {
-	const FPluginDescriptor Descriptor = InPlugin.GetDescriptor();
-	
-	TArray<FString> ModulePrefixes;
-	ModulePrefixes.Reserve(Descriptor.Modules.Num());
-
-	for (const FModuleDescriptor& Module : Descriptor.Modules)
-	{
-		static constexpr TCHAR Format[] = TEXT("/Script/%s.");
-		ModulePrefixes.Add(FString::Printf(Format, *Module.Name.ToString()));
-	}
-
-	static constexpr TCHAR PluginFormat[] = TEXT("/%s/");
-	ModulePrefixes.AddUnique(FString::Printf(PluginFormat, *InPlugin.GetName()));
-
-	return ModulePrefixes;
+	return RigVMTypeUtils::GetPluginModulePrefixes(InPlugin);
 }
 
 bool FRigVMRegistry_NoLock::IsWithinPlugin(const IPlugin& InPlugin, const UObject* InObject)
diff --git a/Engine/Plugins/Runtime/RigVM/Source/RigVM/Private/RigVMHost.cpp b/Engine/Plugins/Runtime/RigVM/Source/RigVM/Private/RigVMHost.cpp
--- a/Engine/Plugins/Runtime/RigVM/Source/RigVM/Private/RigVMHost.cpp
+++ b/Engine/Plugins/Runtime/RigVM/Source/RigVM/Private/RigVMHost.cpp
@@ -185,6 +185,7 @@ void URigVMHost::PreSave(FObjectPreSaveContext SaveContext)
 	Super::PreSave(SaveContext);
 
 	GenerateUserDefinedDependenciesData(GetRigVMExtendedExecuteContext());
+	GenerateRequiredPluginsData(GetRigVMExtendedExecuteContext());
 }
 
 void URigVMHost::BeginDestroy()
@@ -1072,7 +1073,7 @@ void URigVMHost::PostInitInstance(URigVMHost* InCDO)
 	RequestInit();
 }
 
-void URigVMHost::GenerateUserDefinedDependenciesData(FRigVMExtendedExecuteContext& Context)
+void URigVMHost::GenerateUserDefinedDependenciesData(FRigVMExtendedExecuteContext& InContext)
 {
 	if (VM)
 	{
@@ -1138,6 +1139,35 @@ TArray<const UObject*> URigVMHost::GetUserDefinedDependencies(const TArray<const
 	return Dependencies;
 }
 
+const TArray<FString>& URigVMHost::GetRequiredPlugins() const
+{
+	return RequiredPlugins;
+}
+
+void URigVMHost::GenerateRequiredPluginsData(FRigVMExtendedExecuteContext& InContext)
+{
+	TMap<FString,FString> AllPluginPrefixes;
+	RequiredPlugins.Reset();
+	if (VM)
+	{
+		VM->GetRequiredPlugins(RequiredPlugins, AllPluginPrefixes);
+	}
+
+	const TArray<FRigVMExternalVariable> ExternalVariables = GetExternalVariables();
+	for (const FRigVMExternalVariable& ExternalVariable : ExternalVariables)
+	{
+		if (!ExternalVariable.TypeObject)
+		{
+			continue;
+		}
+		const FString PluginName = RigVMTypeUtils::GetPluginName(ExternalVariable.TypeObject, AllPluginPrefixes);
+		if (!PluginName.IsEmpty())
+		{
+			RequiredPlugins.AddUnique(PluginName);
+		}
+	}
+}
+
 void URigVMHost::HandleExecutionReachedExit(const FName& InEventName)
 {
 	if (bAccumulateTime)
diff --git a/Engine/Plugins/Runtime/RigVM/Source/RigVM/Private/RigVMTypeUtils.cpp b/Engine/Plugins/Runtime/RigVM/Source/RigVM/Private/RigVMTypeUtils.cpp
--- a/Engine/Plugins/Runtime/RigVM/Source/RigVM/Private/RigVMTypeUtils.cpp
+++ b/Engine/Plugins/Runtime/RigVM/Source/RigVM/Private/RigVMTypeUtils.cpp
@@ -10,6 +10,7 @@
 #include "Engine/UserDefinedEnum.h"
 #include "StructUtils/UserDefinedStruct.h"
 #include "UObject/LinkerLoad.h"
+#include "Interfaces/IPluginManager.h"
 
 namespace RigVMTypeUtils
 {
@@ -558,3 +559,74 @@ bool RigVMTypeUtils::AreCompatible(const FProperty* InSourceProperty, const FPro
 	}
 	return bCompatible;
 }
+
+FString RigVMTypeUtils::GetPluginName(const UObject* InObject, TMap<FString,FString>& InOutPluginPrefixes)
+{
+	if(!InObject)
+	{
+		return FString();
+	}
+	
+	if (InOutPluginPrefixes.IsEmpty())
+	{
+		TArray<TSharedRef<IPlugin>> Plugins = IPluginManager::Get().GetDiscoveredPlugins();
+		for (const TSharedRef<IPlugin>& Plugin : Plugins)
+		{
+			const TArray<FString> Prefixes = GetPluginModulePrefixes(Plugin.Get());
+			for (const FString& Prefix : Prefixes)
+			{
+				InOutPluginPrefixes.Add(Prefix, Plugin->GetName());
+			}
+		}
+	}
+
+	const FString Path = FSoftObjectPath(InObject).ToString();
+	if (Path.IsEmpty())
+	{
+		return FString();
+	}
+	
+	static const FString ScriptPrefix = TEXT("/Script/");
+	if (Path.StartsWith(ScriptPrefix))
+	{
+		int32 PeriodPos = INDEX_NONE;
+		if (Path.FindChar(TEXT('.'), PeriodPos))
+		{
+			if (const FString* PluginName = InOutPluginPrefixes.Find(Path.Left(PeriodPos + 1)))
+			{
+				return *PluginName;
+			}
+		}
+	}
+	else
+	{
+		const int32 SecondSlashPos = Path.Find(TEXT("/"), ESearchCase::IgnoreCase, ESearchDir::FromStart, 1);
+		if (SecondSlashPos != INDEX_NONE)
+		{
+			if (const FString* PluginName = InOutPluginPrefixes.Find(Path.Left(SecondSlashPos + 1)))
+			{
+				return *PluginName;
+			}
+		}
+	}
+	return FString();
+}
+
+TArray<FString> RigVMTypeUtils::GetPluginModulePrefixes(const IPlugin& InPlugin)
+{
+	const FPluginDescriptor Descriptor = InPlugin.GetDescriptor();
+	
+	TArray<FString> ModulePrefixes;
+	ModulePrefixes.Reserve(Descriptor.Modules.Num());
+
+	for (const FModuleDescriptor& Module : Descriptor.Modules)
+	{
+		static constexpr TCHAR Format[] = TEXT("/Script/%s.");
+		ModulePrefixes.Add(FString::Printf(Format, *Module.Name.ToString()));
+	}
+
+	static constexpr TCHAR PluginFormat[] = TEXT("/%s/");
+	ModulePrefixes.AddUnique(FString::Printf(PluginFormat, *InPlugin.GetName()));
+
+	return ModulePrefixes;
+}
\ No newline at end of file
diff --git a/Engine/Plugins/Runtime/RigVM/Source/RigVM/Public/RigVMCore/RigVM.h b/Engine/Plugins/Runtime/RigVM/Source/RigVM/Public/RigVMCore/RigVM.h
--- a/Engine/Plugins/Runtime/RigVM/Source/RigVM/Public/RigVMCore/RigVM.h
+++ b/Engine/Plugins/Runtime/RigVM/Source/RigVM/Public/RigVMCore/RigVM.h
@@ -996,6 +996,7 @@ public:
 	RIGVM_API bool ValidateBytecode();
 	RIGVM_API void RefreshArgumentNameCaches();
 	RIGVM_API TArray<const UObject*> GetUserDefinedDependencies(const TArray<const FRigVMMemoryStorageStruct*> InMemory);
+	RIGVM_API void GetRequiredPlugins(TArray<FString>& OutPlugins, TMap<FString,FString>& OutPluginPrefixes);
 
 private:
 
diff --git a/Engine/Plugins/Runtime/RigVM/Source/RigVM/Public/RigVMCore/RigVMMemoryStorageStruct.h b/Engine/Plugins/Runtime/RigVM/Source/RigVM/Public/RigVMCore/RigVMMemoryStorageStruct.h
--- a/Engine/Plugins/Runtime/RigVM/Source/RigVM/Public/RigVMCore/RigVMMemoryStorageStruct.h
+++ b/Engine/Plugins/Runtime/RigVM/Source/RigVM/Public/RigVMCore/RigVMMemoryStorageStruct.h
@@ -23,6 +23,7 @@ struct FRigVMMemoryStorageStruct : public FInstancedPropertyBag
 	UE_API bool Serialize(FArchive& Ar);
 	UE_API void AddStructReferencedObjects(FReferenceCollector& Collector);
 	UE_API void GetUserDefinedDependencies(TArray<const UObject*>& OutDependencies) const;
+	UE_API void GetRequiredPlugins(TArray<FString>& OutPlugins, TMap<FString,FString>& OutPluginPrefixes) const;
 
 	friend FArchive& operator<<(FArchive& Ar, FRigVMMemoryStorageStruct& P)
 	{
@@ -396,6 +397,8 @@ protected:
 	static UE_API const TArray<const FProperty*> EmptyProperties;
 	static UE_API const TArray<FRigVMPropertyPath> EmptyPropertyPaths;
 
+	void ForEachObjectDependency(const TFunction<void(const UObject*)>& InCallBack) const;
+
 	UE_API void RefreshLinkedProperties();
 
 	UE_API void SetDefaultValues(const TArray<FRigVMPropertyDescription>& InPropertyDescriptions);
diff --git a/Engine/Plugins/Runtime/RigVM/Source/RigVM/Public/RigVMHost.h b/Engine/Plugins/Runtime/RigVM/Source/RigVM/Public/RigVMHost.h
--- a/Engine/Plugins/Runtime/RigVM/Source/RigVM/Public/RigVMHost.h
+++ b/Engine/Plugins/Runtime/RigVM/Source/RigVM/Public/RigVMHost.h
@@ -351,6 +351,8 @@ public:
 
 	UE_API UObject* ResolveUserDefinedTypeById(const FString& InTypeName) const;
 
+	UE_API const TArray<FString>& GetRequiredPlugins() const;
+	
 protected:
 
 	UE_API virtual void PostInitInstance(URigVMHost* InCDO);
@@ -382,11 +384,14 @@ protected:
 
 	UE_API void GenerateUserDefinedDependenciesData(FRigVMExtendedExecuteContext& Context);
 	UE_API TArray<const UObject*> GetUserDefinedDependencies(const TArray<const FRigVMMemoryStorageStruct*> InMemory);
+	UE_API virtual void GenerateRequiredPluginsData(FRigVMExtendedExecuteContext& Context);
 
 	UPROPERTY()
 	TMap<FString, FSoftObjectPath> UserDefinedStructGuidToPathName;
 	UPROPERTY()
 	TMap<FString, FSoftObjectPath> UserDefinedEnumToPathName;
+	UPROPERTY()
+	TArray<FString> RequiredPlugins;
 
 private:
 	UPROPERTY(transient)
diff --git a/Engine/Plugins/Runtime/RigVM/Source/RigVM/Public/RigVMTypeUtils.h b/Engine/Plugins/Runtime/RigVM/Source/RigVM/Public/RigVMTypeUtils.h
--- a/Engine/Plugins/Runtime/RigVM/Source/RigVM/Public/RigVMTypeUtils.h
+++ b/Engine/Plugins/Runtime/RigVM/Source/RigVM/Public/RigVMTypeUtils.h
@@ -16,6 +16,7 @@
 #define UE_API RIGVM_API
 
 struct FRigVMTemplateArgumentType;
+class IPlugin;
 
 struct FRigVMUserDefinedTypeResolver
 {
@@ -273,6 +274,9 @@ namespace RigVMTypeUtils
 	}
 
 	RIGVM_API bool AreCompatible(const FProperty* InSourceProperty, const FProperty* InTargetProperty);
+
+	RIGVM_API FString GetPluginName(const UObject* InObject, TMap<FString,FString>& InOutPluginPrefixes);
+	RIGVM_API TArray<FString> GetPluginModulePrefixes(const IPlugin& InPlugin);
 }
 
 #undef UE_API
diff --git a/Engine/Plugins/Runtime/RigVM/Source/RigVMDeveloper/Private/RigVMAsset.cpp b/Engine/Plugins/Runtime/RigVM/Source/RigVMDeveloper/Private/RigVMAsset.cpp
--- a/Engine/Plugins/Runtime/RigVM/Source/RigVMDeveloper/Private/RigVMAsset.cpp
+++ b/Engine/Plugins/Runtime/RigVM/Source/RigVMDeveloper/Private/RigVMAsset.cpp
@@ -28,6 +28,7 @@
 #include "UObject/UE5MainStreamObjectVersion.h"
 #include "Stats/StatsHierarchical.h"
 #include "UObject/UObjectIterator.h"
+#include "Interfaces/IPluginManager.h"
 
 #include UE_INLINE_GENERATED_CPP_BY_NAME(RigVMAsset)
 
@@ -44,6 +45,9 @@
 #if !WITH_RIGVMLEGACYEDITOR
 #include "RigVMEditor/Private/Editor/Kismet/RigVMBlueprintCompilationManager.h"
 #endif
+#include "Framework/Notifications/NotificationManager.h"
+#include "Widgets/Notifications/SNotificationList.h"
+#include "Styling/AppStyle.h"
 #endif//WITH_EDITOR
 
 #define LOCTEXT_NAMESPACE "RigVMAsset"
@@ -1514,6 +1518,30 @@ void IRigVMAssetInterface::PreSave(FObjectPreSaveContext ObjectSaveContext)
 
 void IRigVMAssetInterface::PostLoad()
 {
+	const TArray<FString>& RequiredPlugins = GetRequiredPlugins(false);
+	for (const FString& RequiredPlugin : RequiredPlugins)
+	{
+		if (IPluginManager::Get().FindEnabledPlugin(*RequiredPlugin) == nullptr)
+		{
+			const FString Message = FString::Printf(TEXT("%s requires the '%s' plugin to be enabled."), *GetPackage()->GetPathName(), *RequiredPlugin);
+			UE_LOG(LogRigVMDeveloper, Error, TEXT("%s"), *Message);
+
+			FNotificationInfo Info(FText::FromString(Message));
+			Info.bUseSuccessFailIcons = true;
+			Info.Image = FAppStyle::GetBrush(TEXT("MessageLog.Error"));
+			Info.bFireAndForget = true;
+			Info.bUseThrobber = true;
+			Info.FadeOutDuration = 1.f;
+			Info.ExpireDuration = 10.f;
+	
+			TSharedPtr<SNotificationItem> NotificationPtr = FSlateNotificationManager::Get().AddNotification(Info);
+			if (NotificationPtr)
+			{
+				NotificationPtr->SetCompletionState(SNotificationItem::CS_Fail);
+			}
+		}
+	}
+	
 	FRigVMRegistry::Get().RefreshEngineTypesIfRequired();
 
 	bVMRecompilationRequired = true;
diff --git a/Engine/Plugins/Runtime/RigVM/Source/RigVMDeveloper/Private/RigVMBlueprintLegacy.cpp b/Engine/Plugins/Runtime/RigVM/Source/RigVMDeveloper/Private/RigVMBlueprintLegacy.cpp
--- a/Engine/Plugins/Runtime/RigVM/Source/RigVMDeveloper/Private/RigVMBlueprintLegacy.cpp
+++ b/Engine/Plugins/Runtime/RigVM/Source/RigVMDeveloper/Private/RigVMBlueprintLegacy.cpp
@@ -425,6 +425,25 @@ TSet<TObjectPtr<UObject>>& URigVMBlueprint::GetUserDefinedTypesInUse(bool bFromC
 	return UserDefinedTypesInUse;
 }
 
+const TArray<FString>& URigVMBlueprint::GetRequiredPlugins(bool bRefresh) const
+{
+	if (!bRefresh)
+	{
+		return RequiredPlugins;
+	}
+	
+	if (URigVMBlueprintGeneratedClass* RigClass = GetRigVMBlueprintGeneratedClass())
+	{
+		if (URigVMHost* CDO = Cast<URigVMHost>(RigClass->GetDefaultObject(true /* create if needed */)))
+		{
+			CDO->GenerateRequiredPluginsData(CDO->GetRigVMExtendedExecuteContext());
+			return CDO->GetRequiredPlugins();
+		}
+	}
+	// fall back to the serialized data if we don't have a CDO yet
+	return RequiredPlugins;
+}
+
 void URigVMBlueprint::PostEditChangeChainProperty(struct FPropertyChangedChainEvent& PropertyChangedEvent)
 {
 	UBlueprint::PostEditChangeChainProperty(PropertyChangedEvent);
@@ -467,6 +486,8 @@ void URigVMBlueprint::PreSave(FObjectPreSaveContext ObjectSaveContext)
 		}
 		RigClass->AssetVariant = AssetVariant;
 	}
+
+	RequiredPlugins = GetRequiredPlugins();
 }
 
 void URigVMBlueprint::PostLoad()
diff --git a/Engine/Plugins/Runtime/RigVM/Source/RigVMDeveloper/Private/RigVMCompiler/RigVMCompiler.cpp b/Engine/Plugins/Runtime/RigVM/Source/RigVMDeveloper/Private/RigVMCompiler/RigVMCompiler.cpp
--- a/Engine/Plugins/Runtime/RigVM/Source/RigVMDeveloper/Private/RigVMCompiler/RigVMCompiler.cpp
+++ b/Engine/Plugins/Runtime/RigVM/Source/RigVMDeveloper/Private/RigVMCompiler/RigVMCompiler.cpp
@@ -661,6 +661,7 @@ bool URigVMCompiler::Compile(const FRigVMCompileSettings& InSettings, TArray<URi
 	VisitedGraphs.Append(InGraphs);
 
 	const FRigVMRegistry& Registry = FRigVMRegistry::Get();
+	bool bSendMissingPluginNotification = false;
 
 	TSet<FName> FoundEvents;
 	for(int32 GraphIndex=0; GraphIndex<VisitedGraphs.Num(); GraphIndex++)
@@ -690,6 +691,28 @@ bool URigVMCompiler::Compile(const FRigVMCompileSettings& InSettings, TArray<URi
 				{
 					static const FString InvalidPinEncountered = TEXT("Pin @@ is not valid - potentially using an invalid type?");
 					Settings.ASTSettings.Report(EMessageSeverity::Error, Pin, InvalidPinEncountered);
+
+					if (!bSendMissingPluginNotification)
+					{
+						bSendMissingPluginNotification = true;
+
+						const FString Message = FString::Printf(TEXT("%s seems to require a missing a plugin."), *Pin->GetPackage()->GetPathName());
+						UE_LOG(LogRigVMDeveloper, Error, TEXT("%s"), *Message);
+
+						FNotificationInfo Info(FText::FromString(Message));
+						Info.bUseSuccessFailIcons = true;
+						Info.Image = FAppStyle::GetBrush(TEXT("MessageLog.Error"));
+						Info.bFireAndForget = true;
+						Info.bUseThrobber = true;
+						Info.FadeOutDuration = 1.f;
+						Info.ExpireDuration = 5.f;
+	
+						TSharedPtr<SNotificationItem> NotificationPtr = FSlateNotificationManager::Get().AddNotification(Info);
+						if (NotificationPtr)
+						{
+							NotificationPtr->SetCompletionState(SNotificationItem::CS_Fail);
+						}
+					}
 					return false;
 				}
 			}
diff --git a/Engine/Plugins/Runtime/RigVM/Source/RigVMDeveloper/Public/RigVMAsset.h b/Engine/Plugins/Runtime/RigVM/Source/RigVMDeveloper/Public/RigVMAsset.h
--- a/Engine/Plugins/Runtime/RigVM/Source/RigVMDeveloper/Public/RigVMAsset.h
+++ b/Engine/Plugins/Runtime/RigVM/Source/RigVMDeveloper/Public/RigVMAsset.h
@@ -375,6 +375,9 @@ public:
 	/** Enables or disables profiling for this asset */
 	void SetProfilingEnabled(const bool bEnabled);
 
+	/** Returns the names of the plugins required to be loaded for this asset to work */
+	virtual const TArray<FString>& GetRequiredPlugins(bool bRefresh = true) const = 0;
+	
 protected:
 	IRigVMAssetInterface() {}
 	IRigVMAssetInterface(const FObjectInitializer& ObjectInitializer);
diff --git a/Engine/Plugins/Runtime/RigVM/Source/RigVMDeveloper/Public/RigVMBlueprintLegacy.h b/Engine/Plugins/Runtime/RigVM/Source/RigVMDeveloper/Public/RigVMBlueprintLegacy.h
--- a/Engine/Plugins/Runtime/RigVM/Source/RigVMDeveloper/Public/RigVMBlueprintLegacy.h
+++ b/Engine/Plugins/Runtime/RigVM/Source/RigVMDeveloper/Public/RigVMBlueprintLegacy.h
@@ -108,6 +108,7 @@ public:
 	virtual bool IsRegeneratingOnLoad() const override { return bIsRegeneratingOnLoad;}
 	virtual TArray<FRigVMGraphVariableDescription> GetAssetVariables() const override;
 	virtual TArray<struct FEditedDocumentInfo>& GetLastEditedDocuments() override { return LastEditedDocuments;}
+	virtual const TArray<FString>& GetRequiredPlugins(bool bRefresh = true) const override;
 
 protected:
 	virtual UClass* GetRigVMGeneratedClass() override { return GeneratedClass; }
@@ -379,6 +380,9 @@ protected:
 	FRigVMVariant AssetVariant;
 
 #endif
+	
+	UPROPERTY(VisibleAnywhere, AssetRegistrySearchable, Category = "Dependencies")
+	TArray<FString> RequiredPlugins;
 
 public:
 
diff --git a/Engine/Plugins/Runtime/RigVM/Source/RigVMDeveloper/RigVMDeveloper.Build.cs b/Engine/Plugins/Runtime/RigVM/Source/RigVMDeveloper/RigVMDeveloper.Build.cs
--- a/Engine/Plugins/Runtime/RigVM/Source/RigVMDeveloper/RigVMDeveloper.Build.cs
+++ b/Engine/Plugins/Runtime/RigVM/Source/RigVMDeveloper/RigVMDeveloper.Build.cs
@@ -16,6 +16,7 @@ public class RigVMDeveloper : ModuleRules
                 "RigVM",
                 "VisualGraphUtils",
                 "KismetCompiler",
+                "Projects",
             }
         );
 
