Fixed RigVM equal dispatch node type check to allow comparing compatible types, like enums with uint8 encoded enums

diff --git a/Engine/Plugins/Runtime/RigVM/Source/RigVM/Private/RigVMFunctions/RigVMDispatch_Core.cpp b/Engine/Plugins/Runtime/RigVM/Source/RigVM/Private/RigVMFunctions/RigVMDispatch_Core.cpp
--- a/Engine/Plugins/Runtime/RigVM/Source/RigVM/Private/RigVMFunctions/RigVMDispatch_Core.cpp
+++ b/Engine/Plugins/Runtime/RigVM/Source/RigVM/Private/RigVMFunctions/RigVMDispatch_Core.cpp
@@ -3,6 +3,7 @@
 #include "RigVMFunctions/RigVMDispatch_Core.h"
 #include "RigVMCore/RigVMRegistry.h"
 #include "EulerTransform.h"
+#include "RigVMTypeUtils.h"
 
 #include UE_INLINE_GENERATED_CPP_BY_NAME(RigVMDispatch_Core)
 
@@ -111,7 +112,7 @@ void FRigVMDispatch_CoreEquals::Execute(FRigVMExtendedExecuteContext& InContext,
 	const FProperty* PropertyB = Handles[1].GetResolvedProperty(); 
 	check(PropertyA);
 	check(PropertyB);
-	check(PropertyA->SameType(PropertyB));
+	check(RigVMTypeUtils::AreCompatible(PropertyA, PropertyB));
 	check(Handles[2].IsBool());
 
 	const uint8* A = Handles[0].GetInputData();
