RigVM: Commenting all rigvm math functions

#rnx

diff --git a/Engine/Plugins/Runtime/RigVM/Source/RigVM/Public/RigVMFunctions/Math/RigVMFunction_MathBase.h b/Engine/Plugins/Runtime/RigVM/Source/RigVM/Public/RigVMFunctions/Math/RigVMFunction_MathBase.h
--- a/Engine/Plugins/Runtime/RigVM/Source/RigVM/Public/RigVMFunctions/Math/RigVMFunction_MathBase.h
+++ b/Engine/Plugins/Runtime/RigVM/Source/RigVM/Public/RigVMFunctions/Math/RigVMFunction_MathBase.h
@@ -6,7 +6,10 @@
 #include "RigVMFunctions/RigVMFunctionDefines.h"
 #include "RigVMFunction_MathBase.generated.h"
 
-USTRUCT(meta=(Abstract, NodeColor = "0.05 0.25 0.05"))
+/*
+ * The base class for all pure math nodes
+ */
+USTRUCT(meta=(Abstract, NodeColor = "0.05 0.25 0.05", DocumentationPolicy="Strict"))
 struct FRigVMFunction_MathBase : public FRigVMStruct
 {
 	GENERATED_BODY()
@@ -14,7 +17,10 @@ struct FRigVMFunction_MathBase : public FRigVMStruct
 	virtual void Execute() {};
 };
 
-USTRUCT(meta=(Abstract, NodeColor = "0.05 0.25 0.05"))
+/*
+ * The base class for all mutable math nodes
+ */
+USTRUCT(meta=(Abstract, NodeColor = "0.05 0.25 0.05", DocumentationPolicy="Strict"))
 struct FRigVMFunction_MathMutableBase : public FRigVMStructMutable
 {
 	GENERATED_BODY()
diff --git a/Engine/Plugins/Runtime/RigVM/Source/RigVM/Public/RigVMFunctions/Math/RigVMFunction_MathBool.h b/Engine/Plugins/Runtime/RigVM/Source/RigVM/Public/RigVMFunctions/Math/RigVMFunction_MathBool.h
--- a/Engine/Plugins/Runtime/RigVM/Source/RigVM/Public/RigVMFunctions/Math/RigVMFunction_MathBool.h
+++ b/Engine/Plugins/Runtime/RigVM/Source/RigVM/Public/RigVMFunctions/Math/RigVMFunction_MathBool.h
@@ -5,12 +5,18 @@
 #include "RigVMFunction_MathBase.h"
 #include "RigVMFunction_MathBool.generated.h"
 
+/*
+ * The base class for all pure bool math nodes
+ */
 USTRUCT(meta=(Abstract, Category="Math|Boolean"))
 struct FRigVMFunction_MathBoolBase : public FRigVMFunction_MathBase
 {
 	GENERATED_BODY()
 };
 
+/*
+ * The base class for all bool math constants
+ */
 USTRUCT(meta=(Abstract))
 struct FRigVMFunction_MathBoolConstant : public FRigVMFunction_MathBoolBase
 {
@@ -21,10 +27,14 @@ struct FRigVMFunction_MathBoolConstant : public FRigVMFunction_MathBoolBase
 		Value = false;
 	}
 
+	// The value of the constant
 	UPROPERTY(meta=(Output))
 	bool Value;
 };
 
+/*
+ * The base class for all unary bool math operations
+ */
 USTRUCT(meta=(Abstract))
 struct FRigVMFunction_MathBoolUnaryOp : public FRigVMFunction_MathBoolBase
 {
@@ -35,13 +45,18 @@ struct FRigVMFunction_MathBoolUnaryOp : public FRigVMFunction_MathBoolBase
 		Value = Result = 0.f;
 	}
 
+	// The input value
 	UPROPERTY(meta=(Input))
 	bool Value;
 
+	// The result of the operation
 	UPROPERTY(meta=(Output))
 	bool Result;
 };
 
+/*
+ * The base class for all binary bool math operations
+ */
 USTRUCT(meta=(Abstract))
 struct FRigVMFunction_MathBoolBinaryOp : public FRigVMFunction_MathBoolBase
 {
@@ -52,16 +67,22 @@ struct FRigVMFunction_MathBoolBinaryOp : public FRigVMFunction_MathBoolBase
 		A = B = Result = 0.f;
 	}
 
+	// The first value for the operation
 	UPROPERTY(meta=(Input))
 	bool A;
 
+	// The second value for the operation
 	UPROPERTY(meta=(Input))
 	bool B;
 
+	// The resulting value
 	UPROPERTY(meta=(Output))
 	bool Result;
 };
 
+/*
+ * The base class for all aggregational bool math operations
+ */
 USTRUCT(meta=(Abstract))
 struct FRigVMFunction_MathBoolBinaryAggregateOp : public FRigVMFunction_MathBoolBase
 {
@@ -72,12 +93,15 @@ struct FRigVMFunction_MathBoolBinaryAggregateOp : public FRigVMFunction_MathBool
 		A = B = Result = 0.f;
 	}
 
+	// The first value for the operation
 	UPROPERTY(meta=(Input, Aggregate))
 	bool A;
 
+	// The second value for the operation
 	UPROPERTY(meta=(Input, Aggregate))
 	bool B;
 
+	// The resulting value
 	UPROPERTY(meta=(Output, Aggregate))
 	bool Result;
 };
@@ -275,9 +299,11 @@ struct FRigVMFunction_MathBoolToggled : public FRigVMFunction_MathBoolBase
 		Value = Toggled = Initialized = LastValue = false;
 	}
 
+	// The value to check / compare
 	UPROPERTY(meta=(Input))
 	bool Value;
 
+	// True if the input value has been toggled
 	UPROPERTY(meta=(Output))
 	bool Toggled;
 
@@ -316,6 +342,7 @@ struct FRigVMFunction_MathBoolFlipFlop : public FRigVMFunction_MathBoolBase
 	UPROPERTY(meta = (Input))
 	float Duration;
 
+	// The resulting value
 	UPROPERTY(meta=(Output))
 	bool Result;
 
@@ -350,6 +377,7 @@ struct FRigVMFunction_MathBoolOnce : public FRigVMFunction_MathBoolBase
 	UPROPERTY(meta = (Visible))
 	float Duration;
 
+	// The resulting value - true once the first time this node is hit
 	UPROPERTY(meta=(Output))
 	bool Result;
 
@@ -376,9 +404,11 @@ struct FRigVMFunction_MathBoolToFloat : public FRigVMFunction_MathBoolBase
 		Result = 0.f;
 	}
 
+	// The bool value to convert
 	UPROPERTY(meta = (Input))
 	bool Value;
 
+	// The resulting float number (0 or 1)
 	UPROPERTY(meta= (Output))
 	float Result;
 };
@@ -399,9 +429,11 @@ struct FRigVMFunction_MathBoolToInteger : public FRigVMFunction_MathBoolBase
 		Result = 0;
 	}
 
+	// The bool value to convert
 	UPROPERTY(meta = (Input))
 	bool Value;
 
+	// The resulting int number (0 or 1)
 	UPROPERTY(meta= (Output))
 	int32 Result;
 };
diff --git a/Engine/Plugins/Runtime/RigVM/Source/RigVM/Public/RigVMFunctions/Math/RigVMFunction_MathBox.h b/Engine/Plugins/Runtime/RigVM/Source/RigVM/Public/RigVMFunctions/Math/RigVMFunction_MathBox.h
--- a/Engine/Plugins/Runtime/RigVM/Source/RigVM/Public/RigVMFunctions/Math/RigVMFunction_MathBox.h
+++ b/Engine/Plugins/Runtime/RigVM/Source/RigVM/Public/RigVMFunctions/Math/RigVMFunction_MathBox.h
@@ -6,6 +6,9 @@
 #include "RigVMMathLibrary.h"
 #include "RigVMFunction_MathBox.generated.h"
 
+/*
+ * The base class for all pure box math nodes
+ */
 USTRUCT(meta=(Abstract, Category="Math|Box", MenuDescSuffix="(Box)"))
 struct FRigVMFunction_MathBoxBase : public FRigVMFunction_MathBase
 {
@@ -30,21 +33,27 @@ struct FRigVMFunction_MathBoxFromArray : public FRigVMFunction_MathBoxBase
 	RIGVM_METHOD()
 	RIGVM_API virtual void Execute() override;
 
+	// The array of input positions
 	UPROPERTY(meta = (Input))
 	TArray<FVector> Array;
 
+	// The resulting bounding box
 	UPROPERTY(meta = (Output))
 	FBox Box;
 
+	// The resulting minimum value of the bounding box
 	UPROPERTY(meta = (Output))
 	FVector Minimum;
 
+	// The resulting maximum value of the bounding box
 	UPROPERTY(meta = (Output))
 	FVector Maximum;
 
+	// The resulting center value of the bounding box
 	UPROPERTY(meta = (Output))
 	FVector Center;
-	
+
+	// The resulting size value of the bounding box
 	UPROPERTY(meta = (Output))
 	FVector Size;
 };
@@ -67,9 +76,11 @@ struct FRigVMFunction_MathBoxIsValid : public FRigVMFunction_MathBoxBase
 	RIGVM_METHOD()
 	RIGVM_API virtual void Execute() override;
 
+	// The box to check
 	UPROPERTY(meta = (Input))
 	FBox Box;
 
+	// Returns true if the box has any content / is valid
 	UPROPERTY(meta = (Output))
 	bool Valid;
 };
@@ -92,9 +103,11 @@ struct FRigVMFunction_MathBoxGetCenter : public FRigVMFunction_MathBoxBase
 	RIGVM_METHOD()
 	RIGVM_API virtual void Execute() override;
 
+	// The box to retrieve the center for
 	UPROPERTY(meta = (Input))
 	FBox Box;
 
+	// The center of the box
 	UPROPERTY(meta = (Output))
 	FVector Center;
 };
@@ -117,6 +130,7 @@ struct FRigVMFunction_MathBoxGetSize : public FRigVMFunction_MathBoxBase
 	RIGVM_METHOD()
 	RIGVM_API virtual void Execute() override;
 
+	// The box to retrieve the size for
 	UPROPERTY(meta = (Input))
 	FBox Box;
 
@@ -147,6 +161,7 @@ struct FRigVMFunction_MathBoxShift : public FRigVMFunction_MathBoxBase
 	RIGVM_METHOD()
 	RIGVM_API virtual void Execute() override;
 
+	// the box to move
 	UPROPERTY(meta = (Input))
 	FBox Box;
 
@@ -154,6 +169,7 @@ struct FRigVMFunction_MathBoxShift : public FRigVMFunction_MathBoxBase
 	UPROPERTY(meta = (Input))
 	FVector Amount;
 
+	// The resulting moved box
 	UPROPERTY(meta = (Output))
 	FBox Result;
 };
@@ -176,6 +192,7 @@ struct FRigVMFunction_MathBoxMoveTo : public FRigVMFunction_MathBoxBase
 	RIGVM_METHOD()
 	RIGVM_API virtual void Execute() override;
 
+	// the box to move
 	UPROPERTY(meta = (Input))
 	FBox Box;
 
@@ -183,6 +200,7 @@ struct FRigVMFunction_MathBoxMoveTo : public FRigVMFunction_MathBoxBase
 	UPROPERTY(meta = (Input))
 	FVector Center;
 
+	// The resulting moved box
 	UPROPERTY(meta = (Output))
 	FBox Result;
 };
@@ -205,6 +223,7 @@ struct FRigVMFunction_MathBoxExpand : public FRigVMFunction_MathBoxBase
 	RIGVM_METHOD()
 	RIGVM_API virtual void Execute() override;
 
+	// the box to expand
 	UPROPERTY(meta = (Input))
 	FBox Box;
 
@@ -212,6 +231,7 @@ struct FRigVMFunction_MathBoxExpand : public FRigVMFunction_MathBoxBase
 	UPROPERTY(meta = (Input))
 	FVector Amount;
 
+	// the resulting expanded box
 	UPROPERTY(meta = (Output))
 	FBox Result;
 };
@@ -234,12 +254,15 @@ struct FRigVMFunction_MathBoxTransform : public FRigVMFunction_MathBoxBase
 	RIGVM_METHOD()
 	RIGVM_API virtual void Execute() override;
 
+	// the box to transform
 	UPROPERTY(meta = (Input))
 	FBox Box;
 
+	// the transform to apply to the box
 	UPROPERTY(meta = (Input))
 	FTransform Transform;
 
+	// the resulting transformed box
 	UPROPERTY(meta = (Output))
 	FBox Result;
 };
@@ -265,9 +288,11 @@ struct FRigVMFunction_MathBoxGetDistance : public FRigVMFunction_MathBoxBase
 	RIGVM_METHOD()
 	RIGVM_API virtual void Execute() override;
 
+	// The box to measure the distance to
 	UPROPERTY(meta = (Input))
 	FBox Box;
 
+	// The position to measure the distance to
 	UPROPERTY(meta = (Input))
 	FVector Position;
 
@@ -275,9 +300,11 @@ struct FRigVMFunction_MathBoxGetDistance : public FRigVMFunction_MathBoxBase
 	UPROPERTY(meta = (Input))
 	bool Square;
 
+	// Returns true if the computed distance is valid
 	UPROPERTY(meta = (Output))
 	bool Valid;
 
+	// The distance between the box and the position
 	UPROPERTY(meta = (Output))
 	float Distance;
 };
@@ -301,12 +328,15 @@ struct FRigVMFunction_MathBoxIsInside : public FRigVMFunction_MathBoxBase
 	RIGVM_METHOD()
 	RIGVM_API virtual void Execute() override;
 
+	// The box to test
 	UPROPERTY(meta = (Input))
 	FBox Box;
 
-	UPROPERTY(meta = (Output))
+	// The position to check for
+	UPROPERTY(meta = (Input))
 	FVector Position;
 
+	// Returns true if the given point is inside the given box
 	UPROPERTY(meta = (Output))
 	bool Result;
 };
@@ -330,9 +360,11 @@ struct FRigVMFunction_MathBoxGetVolume : public FRigVMFunction_MathBoxBase
 	RIGVM_METHOD()
 	RIGVM_API virtual void Execute() override;
 
+	// The box to retrieve the volume for
 	UPROPERTY(meta = (Input))
 	FBox Box;
 
+	// The volume of the box
 	UPROPERTY(meta = (Output))
 	float Volume;
 };
\ No newline at end of file
diff --git a/Engine/Plugins/Runtime/RigVM/Source/RigVM/Public/RigVMFunctions/Math/RigVMFunction_MathColor.h b/Engine/Plugins/Runtime/RigVM/Source/RigVM/Public/RigVMFunctions/Math/RigVMFunction_MathColor.h
--- a/Engine/Plugins/Runtime/RigVM/Source/RigVM/Public/RigVMFunctions/Math/RigVMFunction_MathColor.h
+++ b/Engine/Plugins/Runtime/RigVM/Source/RigVM/Public/RigVMFunctions/Math/RigVMFunction_MathColor.h
@@ -5,32 +5,18 @@
 #include "RigVMFunction_MathBase.h"
 #include "RigVMFunction_MathColor.generated.h"
 
+/*
+ * The base class for all pure color math nodes
+ */
 USTRUCT(meta=(Abstract, Category="Math|Color", MenuDescSuffix="(Color)"))
 struct FRigVMFunction_MathColorBase : public FRigVMFunction_MathBase
 {
 	GENERATED_BODY()
 };
 
-USTRUCT(meta=(Abstract))
-struct FRigVMFunction_MathColorBinaryOp : public FRigVMFunction_MathColorBase
-{
-	GENERATED_BODY()
-
-	FRigVMFunction_MathColorBinaryOp()
-	{
-		A = B = Result = FLinearColor::Black;
-	}
-
-	UPROPERTY(meta=(Input))
-	FLinearColor A;
-
-	UPROPERTY(meta=(Input))
-	FLinearColor B;
-
-	UPROPERTY(meta=(Output))
-	FLinearColor Result;
-};
-
+/*
+ * The base class for all aggregational color math operations
+ */
 USTRUCT(meta=(Abstract))
 struct FRigVMFunction_MathColorBinaryAggregateOp : public FRigVMFunction_MathColorBase
 {
@@ -41,12 +27,15 @@ struct FRigVMFunction_MathColorBinaryAggregateOp : public FRigVMFunction_MathCol
 		A = B = Result = FLinearColor::Black;
 	}
 
+	// The first color for the operation
 	UPROPERTY(meta=(Input, Aggregate))
 	FLinearColor A;
 
+	// The second color for the operation
 	UPROPERTY(meta=(Input, Aggregate))
 	FLinearColor B;
 
+	// The resulting color
 	UPROPERTY(meta=(Output, Aggregate))
 	FLinearColor Result;
 };
@@ -88,7 +77,7 @@ struct FRigVMFunction_MathColorMake : public FRigVMFunction_MathColorBase
 };
 
 /**
- * Makes a vector from a single float
+ * Makes a color from a single float
  */
 USTRUCT(meta=(DisplayName="From Float", TemplateName="Cast", ExecuteContext="FRigVMExecuteContext", Keywords="Make,Construct"))
 struct FRigVMFunction_MathColorFromFloat : public FRigVMFunction_MathColorBase
@@ -103,15 +92,17 @@ struct FRigVMFunction_MathColorFromFloat : public FRigVMFunction_MathColorBase
 		Result = FLinearColor::Black;
 	}
 
+	// The float value to convert to a color
 	UPROPERTY(meta=(Input))
 	float Value;
 
+	// The resulting color
 	UPROPERTY(meta=(Output))
 	FLinearColor Result;
 };
 
 /**
- * Makes a vector from a single double
+ * Makes a color from a single double
  */
 USTRUCT(meta=(DisplayName="From Double", TemplateName="Cast", ExecuteContext="FRigVMExecuteContext", Keywords="Make,Construct"))
 struct FRigVMFunction_MathColorFromDouble : public FRigVMFunction_MathColorBase
@@ -126,9 +117,11 @@ struct FRigVMFunction_MathColorFromDouble : public FRigVMFunction_MathColorBase
 		Result = FLinearColor::Black;
 	}
 
+	// The double value to convert to a color
 	UPROPERTY(meta=(Input))
 	double Value;
-
+	
+	// The resulting color
 	UPROPERTY(meta=(Output))
 	FLinearColor Result;
 };
@@ -149,7 +142,7 @@ struct FRigVMFunction_MathColorAdd : public FRigVMFunction_MathColorBinaryAggreg
  * Returns the difference of the two values
  */
 USTRUCT(meta=(DisplayName="Subtract", TemplateName="Subtract", Keywords="-"))
-struct FRigVMFunction_MathColorSub : public FRigVMFunction_MathColorBinaryOp
+struct FRigVMFunction_MathColorSub : public FRigVMFunction_MathColorBinaryAggregateOp
 {
 	GENERATED_BODY()
 
@@ -192,15 +185,19 @@ struct FRigVMFunction_MathColorLerp : public FRigVMFunction_MathColorBase
 	RIGVM_METHOD()
 	RIGVM_API virtual void Execute() override;
 
+	// The first color to interpolate from
 	UPROPERTY(meta=(Input))
 	FLinearColor A;
 
+	// The second color to interpolate to
 	UPROPERTY(meta=(Input))
 	FLinearColor B;
 
-	UPROPERTY(meta=(Input))
+	// The blend value for the interpolation
+	UPROPERTY(meta=(Input, UIMin = "0", UIMax = "1"))
 	float T;
 
+	// The resulting interpolated color
 	UPROPERTY(meta=(Output))
 	FLinearColor Result;
 };
diff --git a/Engine/Plugins/Runtime/RigVM/Source/RigVM/Public/RigVMFunctions/Math/RigVMFunction_MathDouble.h b/Engine/Plugins/Runtime/RigVM/Source/RigVM/Public/RigVMFunctions/Math/RigVMFunction_MathDouble.h
--- a/Engine/Plugins/Runtime/RigVM/Source/RigVM/Public/RigVMFunctions/Math/RigVMFunction_MathDouble.h
+++ b/Engine/Plugins/Runtime/RigVM/Source/RigVM/Public/RigVMFunctions/Math/RigVMFunction_MathDouble.h
@@ -5,12 +5,18 @@
 #include "RigVMFunction_MathBase.h"
 #include "RigVMFunction_MathDouble.generated.h"
 
+/*
+ * The base class for all pure double math nodes
+ */
 USTRUCT(meta=(Abstract, Category="Math|Double", MenuDescSuffix="(Double)"))
 struct FRigVMFunction_MathDoubleBase : public FRigVMFunction_MathBase
 {
 	GENERATED_BODY()
 };
 
+/*
+ * The base class for all double math constants
+ */
 USTRUCT(meta=(Abstract))
 struct FRigVMFunction_MathDoubleConstant : public FRigVMFunction_MathDoubleBase
 {
@@ -21,10 +27,14 @@ struct FRigVMFunction_MathDoubleConstant : public FRigVMFunction_MathDoubleBase
 		Value = 0.0;
 	}
 
+	// The value of the constant
 	UPROPERTY(meta=(Output))
 	double Value;
 };
 
+/*
+ * The base class for all unary double math operations
+ */
 USTRUCT(meta=(Abstract))
 struct FRigVMFunction_MathDoubleUnaryOp : public FRigVMFunction_MathDoubleBase
 {
@@ -35,13 +45,18 @@ struct FRigVMFunction_MathDoubleUnaryOp : public FRigVMFunction_MathDoubleBase
 		Value = Result = 0.0;
 	}
 
+	// The input value
 	UPROPERTY(meta=(Input))
 	double Value;
 
+	// The result of the operation
 	UPROPERTY(meta=(Output))
 	double Result;
 };
 
+/*
+ * The base class for all binary double math operations
+ */
 USTRUCT(meta=(Abstract))
 struct FRigVMFunction_MathDoubleBinaryOp : public FRigVMFunction_MathDoubleBase
 {
@@ -52,17 +67,22 @@ struct FRigVMFunction_MathDoubleBinaryOp : public FRigVMFunction_MathDoubleBase
 		A = B = Result = 0.0;
 	}
 
+	// The first number for the operation
 	UPROPERTY(meta=(Input))
 	double A;
 
+	// The second number for the operation
 	UPROPERTY(meta=(Input))
 	double B;
 
+	// The resulting number
 	UPROPERTY(meta=(Output))
 	double Result;
 };
 
-
+/*
+ * The base class for all aggregational double math operations
+ */
 USTRUCT(meta=(Abstract))
 struct FRigVMFunction_MathDoubleBinaryAggregateOp : public FRigVMFunction_MathDoubleBase
 {
@@ -73,12 +93,15 @@ struct FRigVMFunction_MathDoubleBinaryAggregateOp : public FRigVMFunction_MathDo
 		A = B = Result = 0.0;
 	}
 
+	// The first number for the operation
 	UPROPERTY(meta=(Input, Aggregate))
 	double A;
 
+	// The second number for the operation
 	UPROPERTY(meta=(Input, Aggregate))
 	double B;
 
+	// The resulting number
 	UPROPERTY(meta=(Output, Aggregate))
 	double Result;
 };
@@ -170,7 +193,7 @@ struct FRigVMFunction_MathDoubleAdd : public FRigVMFunction_MathDoubleBinaryAggr
  * Returns the difference of the two values
  */
 USTRUCT(meta=(DisplayName="Subtract", TemplateName="Subtract", Keywords="-"))
-struct FRigVMFunction_MathDoubleSub : public FRigVMFunction_MathDoubleBinaryOp
+struct FRigVMFunction_MathDoubleSub : public FRigVMFunction_MathDoubleBinaryAggregateOp
 {
 	GENERATED_BODY()
 
@@ -326,12 +349,15 @@ struct FRigVMFunction_MathDoubleFloor : public FRigVMFunction_MathDoubleBase
 	RIGVM_METHOD()
 	RIGVM_API virtual void Execute() override;
 
+	// The input value to apply the floor to
 	UPROPERTY(meta=(Input))
 	double Value;
 
+	// The resulting closest lower full number (integer) of the input value
 	UPROPERTY(meta=(Output))
 	double Result;
 
+	// The result as an integer value
 	UPROPERTY(meta=(Output))
 	int32 Int;
 };
@@ -353,18 +379,21 @@ struct FRigVMFunction_MathDoubleCeil : public FRigVMFunction_MathDoubleBase
 	RIGVM_METHOD()
 	RIGVM_API virtual void Execute() override;
 
+	// The input value to apply the ceiling to
 	UPROPERTY(meta=(Input))
 	double Value;
 
+	// The resulting closest higher full number (integer) of the input value
 	UPROPERTY(meta=(Output))
 	double Result;
 
+	// The result as an integer value
 	UPROPERTY(meta=(Output))
 	int32 Int;
 };
 
 /**
- * Returns the closest higher full number (integer) of the value
+ * Returns the rounded value of the given double number
  */
 USTRUCT(meta=(DisplayName="Round", TemplateName="Round"))
 struct FRigVMFunction_MathDoubleRound : public FRigVMFunction_MathDoubleBase
@@ -380,12 +409,15 @@ struct FRigVMFunction_MathDoubleRound : public FRigVMFunction_MathDoubleBase
 	RIGVM_METHOD()
 	RIGVM_API virtual void Execute() override;
 
+	// The input value to round
 	UPROPERTY(meta=(Input))
 	double Value;
 
+	// The rounded value of the given double number
 	UPROPERTY(meta=(Output))
 	double Result;
 
+	// The result as an integer
 	UPROPERTY(meta=(Output))
 	int32 Int;
 };
@@ -407,9 +439,11 @@ struct FRigVMFunction_MathDoubleToInt : public FRigVMFunction_MathDoubleBase
 	RIGVM_METHOD()
 	RIGVM_API virtual void Execute() override;
 
+	// The input value to floor to an int
 	UPROPERTY(meta=(Input))
 	double Value;
 
+	// The resulting (floored) integer
 	UPROPERTY(meta=(Output))
 	int32 Result;
 };
@@ -431,9 +465,11 @@ struct FRigVMFunction_MathDoubleToFloat : public FRigVMFunction_MathDoubleBase
 	RIGVM_METHOD()
 	RIGVM_API virtual void Execute() override;
 
+	// The double number to cast to float
 	UPROPERTY(meta=(Input))
 	double Value;
 
+	// The resulting float number
 	UPROPERTY(meta=(Output))
 	float Result;
 };
@@ -466,15 +502,19 @@ struct FRigVMFunction_MathDoubleClamp : public FRigVMFunction_MathDoubleBase
 	RIGVM_METHOD()
 	RIGVM_API virtual void Execute() override;
 
+	// The number to clamp
 	UPROPERTY(meta=(Input))
 	double Value;
 
+	// The Minimum for the resulting range
 	UPROPERTY(meta=(Input))
 	double Minimum;
 
+	// The Maximum for the resulting range
 	UPROPERTY(meta=(Input))
 	double Maximum;
 
+	// The resulting clamped value
 	UPROPERTY(meta=(Output))
 	double Result;
 };
@@ -496,15 +536,19 @@ struct FRigVMFunction_MathDoubleLerp : public FRigVMFunction_MathDoubleBase
 	RIGVM_METHOD()
 	RIGVM_API virtual void Execute() override;
 
+	// The first number to interpolate from
 	UPROPERTY(meta=(Input))
 	double A;
 
+	// The second number to interpolate to
 	UPROPERTY(meta=(Input))
 	double B;
 
-	UPROPERTY(meta=(Input))
+	// The blend value for the interpolation
+	UPROPERTY(meta=(Input, UIMin = "0", UIMax = "1"))
 	double T;
 
+	// The resulting interpolated number
 	UPROPERTY(meta=(Output))
 	double Result;
 };
@@ -527,18 +571,23 @@ struct FRigVMFunction_MathDoubleRemap : public FRigVMFunction_MathDoubleBase
 	RIGVM_METHOD()
 	RIGVM_API virtual void Execute() override;
 
+	// The input value to remap
 	UPROPERTY(meta=(Input))
 	double Value;
 
+	// The minimum of the range of the input / source value
 	UPROPERTY(meta=(Input))
 	double SourceMinimum;
 
+	// The maximum of the range of the input / source value
 	UPROPERTY(meta=(Input))
 	double SourceMaximum;
 
+	// The minimum of the range of the output / target value
 	UPROPERTY(meta=(Input))
 	double TargetMinimum;
 
+	// The maximum of the range of the output / target value
 	UPROPERTY(meta=(Input))
 	double TargetMaximum;
 
@@ -546,6 +595,7 @@ struct FRigVMFunction_MathDoubleRemap : public FRigVMFunction_MathDoubleBase
 	UPROPERTY(meta=(Input))
 	bool bClamp;
 
+	// The resulting remapped value
 	UPROPERTY(meta=(Output))
 	double Result;
 };
@@ -625,12 +675,15 @@ struct FRigVMFunction_MathDoubleGreater : public FRigVMFunction_MathDoubleBase
 	RIGVM_METHOD()
 	RIGVM_API virtual void Execute() override;
 
+	// The first value to compare
 	UPROPERTY(meta=(Input))
 	double A;
 
+	// The second value to compare
 	UPROPERTY(meta=(Input))
 	double B;
 
+	// True if the value A is greater than B
 	UPROPERTY(meta=(Output))
 	bool Result;
 };
@@ -652,12 +705,15 @@ struct FRigVMFunction_MathDoubleLess : public FRigVMFunction_MathDoubleBase
 	RIGVM_METHOD()
 	RIGVM_API virtual void Execute() override;
 
+	// The first value to compare
 	UPROPERTY(meta=(Input))
 	double A;
 
+	// The second value to compare
 	UPROPERTY(meta=(Input))
 	double B;
 
+	// True if the value A is less than B
 	UPROPERTY(meta=(Output))
 	bool Result;
 };
@@ -679,12 +735,15 @@ struct FRigVMFunction_MathDoubleGreaterEqual : public FRigVMFunction_MathDoubleB
 	RIGVM_METHOD()
 	RIGVM_API virtual void Execute() override;
 
+	// The first value to compare
 	UPROPERTY(meta=(Input))
 	double A;
 
+	// The second value to compare
 	UPROPERTY(meta=(Input))
 	double B;
 
+	// True if the value A is greater than or equal to B
 	UPROPERTY(meta=(Output))
 	bool Result;
 };
@@ -706,12 +765,15 @@ struct FRigVMFunction_MathDoubleLessEqual : public FRigVMFunction_MathDoubleBase
 	RIGVM_METHOD()
 	RIGVM_API virtual void Execute() override;
 
+	// The first value to compare
 	UPROPERTY(meta=(Input))
 	double A;
 
+	// The second value to compare
 	UPROPERTY(meta=(Input))
 	double B;
 
+	// True if the value A is less than or equal to B
 	UPROPERTY(meta=(Output))
 	bool Result;
 };
@@ -733,12 +795,15 @@ struct FRigVMFunction_MathDoubleIsNearlyZero : public FRigVMFunction_MathDoubleB
 	RIGVM_METHOD()
 	RIGVM_API virtual void Execute() override;
 
+	// The input value to check
 	UPROPERTY(meta=(Input))
 	double Value;
 
+	// The tolerance to apply for the comparison
 	UPROPERTY(meta=(Input))
 	double Tolerance;
 
+	// True if the value is nearly zero
 	UPROPERTY(meta=(Output))
 	bool Result;
 };
@@ -760,15 +825,19 @@ struct FRigVMFunction_MathDoubleIsNearlyEqual : public FRigVMFunction_MathDouble
 	RIGVM_METHOD()
 	RIGVM_API virtual void Execute() override;
 
+	// The first value to compare
 	UPROPERTY(meta=(Input))
 	double A;
 
+	// The second value to compare
 	UPROPERTY(meta=(Input))
 	double B;
 
+	// The tolerance to apply for the comparison
 	UPROPERTY(meta=(Input))
 	double Tolerance;
 
+	// True if the value A is almost equal to B
 	UPROPERTY(meta=(Output))
 	bool Result;
 };
@@ -898,24 +967,31 @@ struct FRigVMFunction_MathDoubleLawOfCosine : public FRigVMFunction_MathDoubleBa
 	RIGVM_METHOD()
 	RIGVM_API virtual void Execute() override;
 
+	// The length of the first edge of the triangle
 	UPROPERTY(meta = (Input))
 	double A;
 
+	// The length of the second edge of the triangle
 	UPROPERTY(meta = (Input))
 	double B;
 
+	// The length of the third edge of the triangle
 	UPROPERTY(meta = (Input))
 	double C;
 
+	// The angle of between B and C
 	UPROPERTY(meta = (Output))
 	double AlphaAngle;
 
+	// The angle of between A and C
 	UPROPERTY(meta = (Output))
 	double BetaAngle;
 
+	// The angle of between A and B
 	UPROPERTY(meta = (Output))
 	double GammaAngle;
 
+	// True if the results are valid
 	UPROPERTY(meta = (Output))
 	bool bValid;
 };
@@ -949,9 +1025,11 @@ struct FRigVMFunction_MathDoubleArraySum : public FRigVMFunction_MathDoubleBase
 	RIGVM_METHOD()
 	RIGVM_API virtual void Execute() override;
 
+	// The array of number to add up
 	UPROPERTY(meta = (Input))
 	TArray<double> Array;
 
+	// The resulting sum of all numbers
 	UPROPERTY(meta = (Output))
 	double Sum;
 };
@@ -973,9 +1051,11 @@ struct FRigVMFunction_MathDoubleArrayAverage : public FRigVMFunction_MathDoubleB
 	RIGVM_METHOD()
 	RIGVM_API virtual void Execute() override;
 
+	// The array of numbers to average
 	UPROPERTY(meta = (Input))
 	TArray<double> Array;
 
+	// The resulting average
 	UPROPERTY(meta = (Output))
 	double Average;
 };
diff --git a/Engine/Plugins/Runtime/RigVM/Source/RigVM/Public/RigVMFunctions/Math/RigVMFunction_MathFloat.h b/Engine/Plugins/Runtime/RigVM/Source/RigVM/Public/RigVMFunctions/Math/RigVMFunction_MathFloat.h
--- a/Engine/Plugins/Runtime/RigVM/Source/RigVM/Public/RigVMFunctions/Math/RigVMFunction_MathFloat.h
+++ b/Engine/Plugins/Runtime/RigVM/Source/RigVM/Public/RigVMFunctions/Math/RigVMFunction_MathFloat.h
@@ -7,12 +7,18 @@
 
 #define UE_API RIGVM_API
 
+/*
+ * The base class for all pure float math nodes
+ */
 USTRUCT(meta=(Abstract, Category="Math|Float", MenuDescSuffix="(Float)"))
 struct FRigVMFunction_MathFloatBase : public FRigVMFunction_MathBase
 {
 	GENERATED_BODY()
 };
 
+/*
+ * The base class for all float math constants
+ */
 USTRUCT(meta=(Abstract))
 struct FRigVMFunction_MathFloatConstant : public FRigVMFunction_MathFloatBase
 {
@@ -23,10 +29,14 @@ struct FRigVMFunction_MathFloatConstant : public FRigVMFunction_MathFloatBase
 		Value = 0.f;
 	}
 
+	// The value of the constant
 	UPROPERTY(meta=(Output))
 	float Value;
 };
 
+/*
+ * The base class for all unary float math operations
+ */
 USTRUCT(meta=(Abstract))
 struct FRigVMFunction_MathFloatUnaryOp : public FRigVMFunction_MathFloatBase
 {
@@ -37,13 +47,18 @@ struct FRigVMFunction_MathFloatUnaryOp : public FRigVMFunction_MathFloatBase
 		Value = Result = 0.f;
 	}
 
+	// The input value
 	UPROPERTY(meta=(Input))
 	float Value;
 
+	// The result of the operation
 	UPROPERTY(meta=(Output))
 	float Result;
 };
 
+/*
+ * The base class for all binary float math operations
+ */
 USTRUCT(meta=(Abstract))
 struct FRigVMFunction_MathFloatBinaryOp : public FRigVMFunction_MathFloatBase
 {
@@ -54,16 +69,22 @@ struct FRigVMFunction_MathFloatBinaryOp : public FRigVMFunction_MathFloatBase
 		A = B = Result = 0.f;
 	}
 
+	// The first number for the operation
 	UPROPERTY(meta=(Input))
 	float A;
 
+	// The second number for the operation
 	UPROPERTY(meta=(Input))
 	float B;
 
+	// The resulting float number
 	UPROPERTY(meta=(Output))
 	float Result;
 };
 
+/*
+ * The base class for all aggregational float math operations
+ */
 USTRUCT(meta=(Abstract))
 struct FRigVMFunction_MathFloatBinaryAggregateOp : public FRigVMFunction_MathFloatBase
 {
@@ -74,12 +95,15 @@ struct FRigVMFunction_MathFloatBinaryAggregateOp : public FRigVMFunction_MathFlo
 		A = B = Result = 0.f;
 	}
 
+	// The first number for the operation
 	UPROPERTY(meta=(Input, Aggregate))
 	float A;
 
+	// The second number for the operation
 	UPROPERTY(meta=(Input, Aggregate))
 	float B;
 
+	// The resulting float number
 	UPROPERTY(meta=(Output, Aggregate))
 	float Result;
 };
@@ -171,7 +195,7 @@ struct FRigVMFunction_MathFloatAdd : public FRigVMFunction_MathFloatBinaryAggreg
  * Returns the difference of the two values
  */
 USTRUCT(meta=(DisplayName="Subtract", TemplateName="Subtract", Keywords="-"))
-struct FRigVMFunction_MathFloatSub : public FRigVMFunction_MathFloatBinaryOp
+struct FRigVMFunction_MathFloatSub : public FRigVMFunction_MathFloatBinaryAggregateOp
 {
 	GENERATED_BODY()
 
@@ -327,12 +351,15 @@ struct FRigVMFunction_MathFloatFloor : public FRigVMFunction_MathFloatBase
 	RIGVM_METHOD()
 	UE_API virtual void Execute() override;
 
+	// The input value to apply the floor to
 	UPROPERTY(meta=(Input))
 	float Value;
 
+	// The resulting closest lower full number (integer) of the input value
 	UPROPERTY(meta=(Output))
 	float Result;
 
+	// The result as an integer value
 	UPROPERTY(meta=(Output))
 	int32 Int;
 };
@@ -354,18 +381,21 @@ struct FRigVMFunction_MathFloatCeil : public FRigVMFunction_MathFloatBase
 	RIGVM_METHOD()
 	UE_API virtual void Execute() override;
 
+	// The input value to apply the ceiling to
 	UPROPERTY(meta=(Input))
 	float Value;
 
+	// The resulting closest higher full number (integer) of the input value
 	UPROPERTY(meta=(Output))
 	float Result;
 
+	// The result as an integer value
 	UPROPERTY(meta=(Output))
 	int32 Int;
 };
 
 /**
- * Returns the closest higher full number (integer) of the value
+ * Returns the rounded value of the given float number
  */
 USTRUCT(meta=(DisplayName="Round", TemplateName="Round"))
 struct FRigVMFunction_MathFloatRound : public FRigVMFunction_MathFloatBase
@@ -381,12 +411,15 @@ struct FRigVMFunction_MathFloatRound : public FRigVMFunction_MathFloatBase
 	RIGVM_METHOD()
 	UE_API virtual void Execute() override;
 
+	// The input value to round
 	UPROPERTY(meta=(Input))
 	float Value;
 
+	// The rounded value of the given float number
 	UPROPERTY(meta=(Output))
 	float Result;
 
+	// The result as an integer
 	UPROPERTY(meta=(Output))
 	int32 Int;
 };
@@ -408,9 +441,11 @@ struct FRigVMFunction_MathFloatToInt : public FRigVMFunction_MathFloatBase
 	RIGVM_METHOD()
 	UE_API virtual void Execute() override;
 
+	// The input value to floor to an int
 	UPROPERTY(meta=(Input))
 	float Value;
 
+	// The resulting (floored) integer
 	UPROPERTY(meta=(Output))
 	int32 Result;
 };
@@ -432,9 +467,11 @@ struct FRigVMFunction_MathFloatToDouble : public FRigVMFunction_MathFloatBase
 	RIGVM_METHOD()
 	UE_API virtual void Execute() override;
 
+	// The float number to cast to double
 	UPROPERTY(meta=(Input))
 	float Value;
 
+	// The resulting double number
 	UPROPERTY(meta=(Output))
 	double Result;
 };
@@ -467,15 +504,19 @@ struct FRigVMFunction_MathFloatClamp : public FRigVMFunction_MathFloatBase
 	RIGVM_METHOD()
 	UE_API virtual void Execute() override;
 
+	// The number to clamp
 	UPROPERTY(meta=(Input))
 	float Value;
 
+	// The Minimum for the resulting range
 	UPROPERTY(meta=(Input))
 	float Minimum;
 
+	// The Maximum for the resulting range
 	UPROPERTY(meta=(Input))
 	float Maximum;
 
+	// The resulting clamped value
 	UPROPERTY(meta=(Output))
 	float Result;
 };
@@ -497,15 +538,19 @@ struct FRigVMFunction_MathFloatLerp : public FRigVMFunction_MathFloatBase
 	RIGVM_METHOD()
 	UE_API virtual void Execute() override;
 
+	// The first number to interpolate from
 	UPROPERTY(meta=(Input))
 	float A;
 
+	// The second number to interpolate to
 	UPROPERTY(meta=(Input))
 	float B;
 
-	UPROPERTY(meta=(Input))
+	// The blend value for the interpolation
+	UPROPERTY(meta=(Input, UIMin = "0", UIMax = "1"))
 	float T;
 
+	// The resulting interpolated number
 	UPROPERTY(meta=(Output))
 	float Result;
 };
@@ -528,18 +573,23 @@ struct FRigVMFunction_MathFloatRemap : public FRigVMFunction_MathFloatBase
 	RIGVM_METHOD()
 	UE_API virtual void Execute() override;
 
+	// The input value to remap
 	UPROPERTY(meta=(Input))
 	float Value;
 
+	// The minimum of the range of the input / source value
 	UPROPERTY(meta=(Input))
 	float SourceMinimum;
 
+	// The maximum of the range of the input / source value
 	UPROPERTY(meta=(Input))
 	float SourceMaximum;
 
+	// The minimum of the range of the output / target value
 	UPROPERTY(meta=(Input))
 	float TargetMinimum;
 
+	// The maximum of the range of the output / target value
 	UPROPERTY(meta=(Input))
 	float TargetMaximum;
 
@@ -547,6 +597,7 @@ struct FRigVMFunction_MathFloatRemap : public FRigVMFunction_MathFloatBase
 	UPROPERTY(meta=(Input))
 	bool bClamp;
 
+	// The resulting remapped value
 	UPROPERTY(meta=(Output))
 	float Result;
 };
@@ -626,12 +677,15 @@ struct FRigVMFunction_MathFloatGreater : public FRigVMFunction_MathFloatBase
 	RIGVM_METHOD()
 	UE_API virtual void Execute() override;
 
+	// The first value to compare
 	UPROPERTY(meta=(Input))
 	float A;
 
+	// The second value to compare
 	UPROPERTY(meta=(Input))
 	float B;
 
+	// True if the value A is greater than B
 	UPROPERTY(meta=(Output))
 	bool Result;
 };
@@ -653,12 +707,15 @@ struct FRigVMFunction_MathFloatLess : public FRigVMFunction_MathFloatBase
 	RIGVM_METHOD()
 	UE_API virtual void Execute() override;
 
+	// The first value to compare
 	UPROPERTY(meta=(Input))
 	float A;
 
+	// The second value to compare
 	UPROPERTY(meta=(Input))
 	float B;
 
+	// True if the value A is less than B
 	UPROPERTY(meta=(Output))
 	bool Result;
 };
@@ -680,12 +737,15 @@ struct FRigVMFunction_MathFloatGreaterEqual : public FRigVMFunction_MathFloatBas
 	RIGVM_METHOD()
 	UE_API virtual void Execute() override;
 
+	// The first value to compare
 	UPROPERTY(meta=(Input))
 	float A;
 
+	// The second value to compare
 	UPROPERTY(meta=(Input))
 	float B;
 
+	// True if the value A is greater than or equal to B
 	UPROPERTY(meta=(Output))
 	bool Result;
 };
@@ -707,12 +767,15 @@ struct FRigVMFunction_MathFloatLessEqual : public FRigVMFunction_MathFloatBase
 	RIGVM_METHOD()
 	UE_API virtual void Execute() override;
 
+	// The first value to compare
 	UPROPERTY(meta=(Input))
 	float A;
 
+	// The second value to compare
 	UPROPERTY(meta=(Input))
 	float B;
 
+	// True if the value A is less than or equal to B
 	UPROPERTY(meta=(Output))
 	bool Result;
 };
@@ -734,12 +797,15 @@ struct FRigVMFunction_MathFloatIsNearlyZero : public FRigVMFunction_MathFloatBas
 	RIGVM_METHOD()
 	UE_API virtual void Execute() override;
 
+	// The input value to compare
 	UPROPERTY(meta=(Input))
 	float Value;
 
+	// The tolerance to apply for the comparison
 	UPROPERTY(meta=(Input))
 	float Tolerance;
 
+	// True if the value is nearly zero
 	UPROPERTY(meta=(Output))
 	bool Result;
 };
@@ -761,15 +827,19 @@ struct FRigVMFunction_MathFloatIsNearlyEqual : public FRigVMFunction_MathFloatBa
 	RIGVM_METHOD()
 	UE_API virtual void Execute() override;
 
+	// The first value to compare
 	UPROPERTY(meta=(Input))
 	float A;
 
+	// The second value to compare
 	UPROPERTY(meta=(Input))
 	float B;
 
+	// The tolerance to apply for the comparison
 	UPROPERTY(meta=(Input))
 	float Tolerance;
 
+	// True if the value A is almost equal to B
 	UPROPERTY(meta=(Output))
 	bool Result;
 };
@@ -932,24 +1002,31 @@ struct FRigVMFunction_MathFloatLawOfCosine : public FRigVMFunction_MathFloatBase
 	RIGVM_METHOD()
 	UE_API virtual void Execute() override;
 
+	// The length of the first edge of the triangle
 	UPROPERTY(meta = (Input))
 	float A;
 
+	// The length of the second edge of the triangle
 	UPROPERTY(meta = (Input))
 	float B;
 
+	// The length of the third edge of the triangle
 	UPROPERTY(meta = (Input))
 	float C;
 
+	// The angle of between B and C
 	UPROPERTY(meta = (Output))
 	float AlphaAngle;
 
+	// The angle of between A and C
 	UPROPERTY(meta = (Output))
 	float BetaAngle;
 
+	// The angle of between A and B
 	UPROPERTY(meta = (Output))
 	float GammaAngle;
 
+	// True if the results are valid
 	UPROPERTY(meta = (Output))
 	bool bValid;
 };
@@ -983,9 +1060,11 @@ struct FRigVMFunction_MathFloatArraySum : public FRigVMFunction_MathFloatBase
 	RIGVM_METHOD()
 	UE_API virtual void Execute() override;
 
+	// The array of number to add up
 	UPROPERTY(meta = (Input))
 	TArray<float> Array;
 
+	// The resulting sum of all numbers
 	UPROPERTY(meta = (Output))
 	float Sum;
 };
@@ -1007,9 +1086,11 @@ struct FRigVMFunction_MathFloatArrayAverage : public FRigVMFunction_MathFloatBas
 	RIGVM_METHOD()
 	UE_API virtual void Execute() override;
 
+	// The array of numbers to average
 	UPROPERTY(meta = (Input))
 	TArray<float> Array;
 
+	// The resulting average
 	UPROPERTY(meta = (Output))
 	float Average;
 };
diff --git a/Engine/Plugins/Runtime/RigVM/Source/RigVM/Public/RigVMFunctions/Math/RigVMFunction_MathInt.h b/Engine/Plugins/Runtime/RigVM/Source/RigVM/Public/RigVMFunctions/Math/RigVMFunction_MathInt.h
--- a/Engine/Plugins/Runtime/RigVM/Source/RigVM/Public/RigVMFunctions/Math/RigVMFunction_MathInt.h
+++ b/Engine/Plugins/Runtime/RigVM/Source/RigVM/Public/RigVMFunctions/Math/RigVMFunction_MathInt.h
@@ -5,12 +5,18 @@
 #include "RigVMFunction_MathBase.h"
 #include "RigVMFunction_MathInt.generated.h"
 
+/*
+ * The base class for all pure int math nodes
+ */
 USTRUCT(meta=(Abstract, Category="Math|Int", MenuDescSuffix="(Int)"))
 struct FRigVMFunction_MathIntBase : public FRigVMFunction_MathBase
 {
 	GENERATED_BODY()
 };
 
+/*
+ * The base class for all unary int math operations
+ */
 USTRUCT(meta=(Abstract))
 struct FRigVMFunction_MathIntUnaryOp : public FRigVMFunction_MathIntBase
 {
@@ -21,13 +27,18 @@ struct FRigVMFunction_MathIntUnaryOp : public FRigVMFunction_MathIntBase
 		Value = Result = 0;
 	}
 
+	// The input value
 	UPROPERTY(meta=(Input))
 	int32 Value;
 
+	// The result of the operation
 	UPROPERTY(meta=(Output))
 	int32 Result;
 };
 
+/*
+ * The base class for all binary int math operations
+ */
 USTRUCT(meta=(Abstract))
 struct FRigVMFunction_MathIntBinaryOp : public FRigVMFunction_MathIntBase
 {
@@ -38,16 +49,22 @@ struct FRigVMFunction_MathIntBinaryOp : public FRigVMFunction_MathIntBase
 		A = B = Result = 0;
 	}
 
+	// The first number for the operation
 	UPROPERTY(meta=(Input))
 	int32 A;
 
+	// The second number for the operation
 	UPROPERTY(meta=(Input))
 	int32 B;
 
+	// The resulting int number
 	UPROPERTY(meta=(Output))
 	int32 Result;
 };
 
+/*
+ * The base class for all aggregational int math operations
+ */
 USTRUCT(meta=(Abstract))
 struct FRigVMFunction_MathIntBinaryAggregateOp : public FRigVMFunction_MathIntBase
 {
@@ -58,12 +75,15 @@ struct FRigVMFunction_MathIntBinaryAggregateOp : public FRigVMFunction_MathIntBa
 		A = B = Result = 0;
 	}
 
+	// The first number for the operation
 	UPROPERTY(meta=(Input, Aggregate))
 	int32 A;
 
+	// The second number for the operation
 	UPROPERTY(meta=(Input, Aggregate))
 	int32 B;
 
+	// The resulting int number
 	UPROPERTY(meta=(Output, Aggregate))
 	int32 Result;
 };
@@ -107,7 +127,7 @@ struct FRigVMFunction_MathIntAdd : public FRigVMFunction_MathIntBinaryAggregateO
  * Returns the difference of the two values
  */
 USTRUCT(meta=(DisplayName="Subtract", TemplateName="Subtract", Keywords="-"))
-struct FRigVMFunction_MathIntSub : public FRigVMFunction_MathIntBinaryOp
+struct FRigVMFunction_MathIntSub : public FRigVMFunction_MathIntBinaryAggregateOp
 {
 	GENERATED_BODY()
 
@@ -251,15 +271,17 @@ struct FRigVMFunction_MathIntToFloat : public FRigVMFunction_MathIntBase
 	RIGVM_METHOD()
 	RIGVM_API virtual void Execute() override;
 
+	// The int to cast
 	UPROPERTY(meta=(Input))
 	int32 Value;
 
+	// The resulting float value
 	UPROPERTY(meta=(Output))
 	float Result;
 };
 
 /**
- * Returns the int cast to a float
+ * Returns the int cast to a double
  */
 USTRUCT(meta=(DisplayName="To Double", TemplateName="Cast", ExecuteContext="FRigVMExecuteContext"))
 struct FRigVMFunction_MathIntToDouble : public FRigVMFunction_MathIntBase
@@ -275,9 +297,11 @@ struct FRigVMFunction_MathIntToDouble : public FRigVMFunction_MathIntBase
 	RIGVM_METHOD()
 	RIGVM_API virtual void Execute() override;
 
+	// The int to cast to double
 	UPROPERTY(meta=(Input))
 	int32 Value;
 
+	// The resulting double value
 	UPROPERTY(meta=(Output))
 	double Result;
 };
@@ -310,15 +334,19 @@ struct FRigVMFunction_MathIntClamp : public FRigVMFunction_MathIntBase
 	RIGVM_METHOD()
 	RIGVM_API virtual void Execute() override;
 
+	// The number to clamp
 	UPROPERTY(meta=(Input))
 	int32 Value;
 
+	// The Minimum for the resulting range
 	UPROPERTY(meta=(Input))
 	int32 Minimum;
-
+	
+	// The Maximum for the resulting range
 	UPROPERTY(meta=(Input))
 	int32 Maximum;
 
+	// The resulting clamped value
 	UPROPERTY(meta=(Output))
 	int32 Result;
 };
@@ -398,12 +426,15 @@ struct FRigVMFunction_MathIntGreater : public FRigVMFunction_MathIntBase
 	RIGVM_METHOD()
 	RIGVM_API virtual void Execute() override;
 
+	// The first number to compare
 	UPROPERTY(meta=(Input))
 	int32 A;
 
+	// The second number to compare
 	UPROPERTY(meta=(Input))
 	int32 B;
 
+	// True if the value A is greater than B
 	UPROPERTY(meta=(Output))
 	bool Result;
 };
@@ -425,12 +456,15 @@ struct FRigVMFunction_MathIntLess : public FRigVMFunction_MathIntBase
 	RIGVM_METHOD()
 	RIGVM_API virtual void Execute() override;
 
+	// The first number to compare
 	UPROPERTY(meta=(Input))
 	int32 A;
 
+	// The second number to compare
 	UPROPERTY(meta=(Input))
 	int32 B;
 
+	// True if the value A is less than B
 	UPROPERTY(meta=(Output))
 	bool Result;
 };
@@ -452,12 +486,15 @@ struct FRigVMFunction_MathIntGreaterEqual : public FRigVMFunction_MathIntBase
 	RIGVM_METHOD()
 	RIGVM_API virtual void Execute() override;
 
+	// The first number to compare
 	UPROPERTY(meta=(Input))
 	int32 A;
 
+	// The second number to compare
 	UPROPERTY(meta=(Input))
 	int32 B;
 
+	// True if the value A is greater than or equal to B
 	UPROPERTY(meta=(Output))
 	bool Result;
 };
@@ -479,12 +516,15 @@ struct FRigVMFunction_MathIntLessEqual : public FRigVMFunction_MathIntBase
 	RIGVM_METHOD()
 	RIGVM_API virtual void Execute() override;
 
+	// The first number to compare
 	UPROPERTY(meta=(Input))
 	int32 A;
 
+	// The second number to compare
 	UPROPERTY(meta=(Input))
 	int32 B;
 
+	// True if the value A is less than or equal to B
 	UPROPERTY(meta=(Output))
 	bool Result;
 };
@@ -506,9 +546,11 @@ struct FRigVMFunction_MathIntArraySum : public FRigVMFunction_MathIntBase
 	RIGVM_METHOD()
 	RIGVM_API virtual void Execute() override;
 
+	// The array of numbers to add up
 	UPROPERTY(meta = (Input))
 	TArray<int32> Array;
 
+	// The resulting sum
 	UPROPERTY(meta = (Output))
 	int32 Sum;
 };
@@ -531,9 +573,11 @@ struct FRigVMFunction_MathIntArrayAverage : public FRigVMFunction_MathIntBase
 	RIGVM_METHOD()
 	RIGVM_API virtual void Execute() override;
 
+	// The array of numbers to average
 	UPROPERTY(meta = (Input))
 	TArray<int32> Array;
 
+	// The resulting average
 	UPROPERTY(meta = (Output))
 	int32 Average;
 };
@@ -555,6 +599,7 @@ struct FRigVMFunction_MathIntToString : public FRigVMFunction_MathIntBase
 	RIGVM_METHOD()
 	RIGVM_API virtual void Execute() override;
 
+	// The number to convert
 	UPROPERTY(meta = (Input))
 	int32 Number;
 
@@ -565,6 +610,7 @@ struct FRigVMFunction_MathIntToString : public FRigVMFunction_MathIntBase
 	UPROPERTY(meta = (Input))
 	int32 PaddedSize;
 
+	// The resulting potentially padded string
 	UPROPERTY(meta = (Output))
 	FString Result;
 };
@@ -586,6 +632,7 @@ struct FRigVMFunction_MathIntToName : public FRigVMFunction_MathIntBase
 	RIGVM_METHOD()
 	RIGVM_API virtual void Execute() override;
 
+	// The number to convert
 	UPROPERTY(meta = (Input))
 	int32 Number;
 
@@ -596,6 +643,7 @@ struct FRigVMFunction_MathIntToName : public FRigVMFunction_MathIntBase
 	UPROPERTY(meta = (Input))
 	int32 PaddedSize;
 
+	// The resulting potentially padded name
 	UPROPERTY(meta = (Output))
 	FName Result;
 };
\ No newline at end of file
diff --git a/Engine/Plugins/Runtime/RigVM/Source/RigVM/Public/RigVMFunctions/Math/RigVMFunction_MathMatrix.h b/Engine/Plugins/Runtime/RigVM/Source/RigVM/Public/RigVMFunctions/Math/RigVMFunction_MathMatrix.h
--- a/Engine/Plugins/Runtime/RigVM/Source/RigVM/Public/RigVMFunctions/Math/RigVMFunction_MathMatrix.h
+++ b/Engine/Plugins/Runtime/RigVM/Source/RigVM/Public/RigVMFunctions/Math/RigVMFunction_MathMatrix.h
@@ -6,12 +6,18 @@
 #include "RigVMFunction_MathBase.h"
 #include "RigVMFunction_MathMatrix.generated.h"
 
+/*
+ * The base class for all pure matrix math nodes
+ */
 USTRUCT(meta=(Abstract, Category="Math|Matrix", MenuDescSuffix="(Matrix)"))
 struct FRigVMFunction_MathMatrixBase : public FRigVMFunction_MathBase
 {
 	GENERATED_BODY()
 };
 
+/*
+ * The base class for all unary matrix math operations
+ */
 USTRUCT(meta=(Abstract))
 struct FRigVMFunction_MathMatrixUnaryOp : public FRigVMFunction_MathMatrixBase
 {
@@ -22,13 +28,18 @@ struct FRigVMFunction_MathMatrixUnaryOp : public FRigVMFunction_MathMatrixBase
 		Value = Result = FMatrix::Identity;
 	}
 
+	// The input value
 	UPROPERTY(meta=(Input))
 	FMatrix Value;
 
+	// The result of the operation
 	UPROPERTY(meta=(Output))
 	FMatrix Result;
 };
 
+/*
+ * The base class for all binary matrix math operations
+ */
 USTRUCT(meta=(Abstract))
 struct FRigVMFunction_MathMatrixBinaryOp : public FRigVMFunction_MathMatrixBase
 {
@@ -39,16 +50,22 @@ struct FRigVMFunction_MathMatrixBinaryOp : public FRigVMFunction_MathMatrixBase
 		A = B = Result = FMatrix::Identity;
 	}
 
+	// The first matrix for the operation
 	UPROPERTY(meta=(Input))
 	FMatrix A;
 
+	// The second matrix for the operation
 	UPROPERTY(meta=(Input))
 	FMatrix B;
 
+	// The resulting matrix
 	UPROPERTY(meta=(Output))
 	FMatrix Result;
 };
 
+/*
+ * The base class for all aggregational matrix math nodes
+ */
 USTRUCT(meta=(Abstract))
 struct FRigVMFunction_MathMatrixBinaryAggregateOp : public FRigVMFunction_MathMatrixBase
 {
@@ -59,12 +76,15 @@ struct FRigVMFunction_MathMatrixBinaryAggregateOp : public FRigVMFunction_MathMa
 		A = B = Result = FMatrix::Identity;
 	}
 
+	// The first matrix for the operation
 	UPROPERTY(meta=(Input, Aggregate))
 	FMatrix A;
 
+	// The second matrix for the operation
 	UPROPERTY(meta=(Input, Aggregate))
 	FMatrix B;
 
+	// The resulting matrix
 	UPROPERTY(meta=(Output, Aggregate))
 	FMatrix Result;
 };
@@ -86,9 +106,11 @@ struct FRigVMFunction_MathMatrixToTransform : public FRigVMFunction_MathMatrixBa
 	RIGVM_METHOD()
 	RIGVM_API virtual void Execute() override;
 
+	// The matrix to convert to a transform
 	UPROPERTY(meta=(Input))
 	FMatrix Value;
 
+	// The resulting transform
 	UPROPERTY(meta=(Output))
 	FTransform Result;
 };
@@ -137,9 +159,11 @@ struct FRigVMFunction_MathMatrixFromTransformV2 : public FRigVMFunction_MathMatr
 	RIGVM_METHOD()
 	RIGVM_API virtual void Execute() override;
 
+	// The transform to convert to a matrix
 	UPROPERTY(meta=(Input))
 	FTransform Value;
 
+	// The resulting matrix
 	UPROPERTY(meta=(Output))
 	FMatrix Result;
 };
@@ -164,18 +188,23 @@ struct FRigVMFunction_MathMatrixToVectors : public FRigVMFunction_MathMatrixBase
 	RIGVM_METHOD()
 	RIGVM_API virtual void Execute() override;
 
+	// The matrix to split up to vectors
 	UPROPERTY(meta=(Input))
 	FMatrix Value;
 
+	// The resulting origin
 	UPROPERTY(meta=(Output))
 	FVector Origin;
 
+	// The resulting X component
 	UPROPERTY(meta=(Output))
 	FVector X;
 
+	// The resulting Y component
 	UPROPERTY(meta=(Output))
 	FVector Y;
 
+	// The resulting Z component
 	UPROPERTY(meta=(Output))
 	FVector Z;
 };
@@ -200,18 +229,23 @@ struct FRigVMFunction_MathMatrixFromVectors : public FRigVMFunction_MathMatrixBa
 	RIGVM_METHOD()
 	RIGVM_API virtual void Execute() override;
 
+	// The input origin for the matrix
 	UPROPERTY(meta=(Input))
 	FVector Origin;
 
+	// The input X component for the matrix
 	UPROPERTY(meta=(Input))
 	FVector X;
 
+	// The input Y component for the matrix
 	UPROPERTY(meta=(Input))
 	FVector Y;
 
+	// The input Z component for the matrix
 	UPROPERTY(meta=(Input))
 	FVector Z;
 
+	// The resulting matrix
 	UPROPERTY(meta=(Output))
 	FMatrix Result;
 };
diff --git a/Engine/Plugins/Runtime/RigVM/Source/RigVM/Public/RigVMFunctions/Math/RigVMFunction_MathQuaternion.h b/Engine/Plugins/Runtime/RigVM/Source/RigVM/Public/RigVMFunctions/Math/RigVMFunction_MathQuaternion.h
--- a/Engine/Plugins/Runtime/RigVM/Source/RigVM/Public/RigVMFunctions/Math/RigVMFunction_MathQuaternion.h
+++ b/Engine/Plugins/Runtime/RigVM/Source/RigVM/Public/RigVMFunctions/Math/RigVMFunction_MathQuaternion.h
@@ -8,12 +8,18 @@
 
 #define UE_API RIGVM_API
 
+/*
+ * The base class for all pure quaternion math nodes
+ */
 USTRUCT(meta=(Abstract, Category="Math|Quaternion", MenuDescSuffix="(Quaternion)"))
 struct FRigVMFunction_MathQuaternionBase : public FRigVMFunction_MathBase
 {
 	GENERATED_BODY()
 };
 
+/*
+ * The base class for all unary quaternion math operations
+ */
 USTRUCT(meta=(Abstract))
 struct FRigVMFunction_MathQuaternionUnaryOp : public FRigVMFunction_MathQuaternionBase
 {
@@ -24,13 +30,18 @@ struct FRigVMFunction_MathQuaternionUnaryOp : public FRigVMFunction_MathQuaterni
 		Value = Result = FQuat::Identity;
 	}
 
+	// The input value
 	UPROPERTY(meta=(Input))
 	FQuat Value;
 
+	// The result of the operation
 	UPROPERTY(meta=(Output))
 	FQuat Result;
 };
 
+/*
+ * The base class for all binary quaternion math operations
+ */
 USTRUCT(meta=(Abstract))
 struct FRigVMFunction_MathQuaternionBinaryOp : public FRigVMFunction_MathQuaternionBase
 {
@@ -41,16 +52,22 @@ struct FRigVMFunction_MathQuaternionBinaryOp : public FRigVMFunction_MathQuatern
 		A = B = Result = FQuat::Identity;
 	}
 
+	// The first quaternion for the operation
 	UPROPERTY(meta=(Input))
 	FQuat A;
 
+	// The second quaternion for the operation
 	UPROPERTY(meta=(Input))
 	FQuat B;
 
+	// The resulting quaternion
 	UPROPERTY(meta=(Output))
 	FQuat Result;
 };
 
+/*
+ * The base class for all aggregational quaternion math operations
+ */
 USTRUCT(meta=(Abstract))
 struct FRigVMFunction_MathQuaternionBinaryAggregateOp : public FRigVMFunction_MathQuaternionBase
 {
@@ -60,13 +77,16 @@ struct FRigVMFunction_MathQuaternionBinaryAggregateOp : public FRigVMFunction_Ma
 	{
 		A = B = Result = FQuat::Identity;
 	}
-
+	
+	// The first quaternion for the operation
 	UPROPERTY(meta=(Input, Aggregate))
 	FQuat A;
-
+	
+	// The second quaternion for the operation
 	UPROPERTY(meta=(Input, Aggregate))
 	FQuat B;
-
+	
+	// The resulting quaternion
 	UPROPERTY(meta=(Output, Aggregate))
 	FQuat Result;
 };
@@ -87,18 +107,23 @@ struct FRigVMFunction_MathQuaternionMake : public FRigVMFunction_MathQuaternionB
 		Result = FQuat::Identity;
 	}
 
+	// The X of the quaternion
 	UPROPERTY(meta=(Input))
 	float X;
 
+	// The Y of the quaternion
 	UPROPERTY(meta=(Input))
 	float Y;
 
+	// The Z of the quaternion
 	UPROPERTY(meta=(Input))
 	float Z;
 
+	// The W of the quaternion
 	UPROPERTY(meta=(Input))
 	float W;
 
+	// The resulting composed quaternion
 	UPROPERTY(meta=(Output))
 	FQuat Result;
 };
@@ -121,12 +146,15 @@ struct FRigVMFunction_MathQuaternionFromAxisAndAngle : public FRigVMFunction_Mat
 	RIGVM_METHOD()
 	UE_API virtual void Execute() override;
 
+	// The axis to use for the quaternion
 	UPROPERTY(meta=(Input))
 	FVector Axis;
 
+	// The angle in radians to use for the quaternion
 	UPROPERTY(meta=(Input))
 	float Angle;
 
+	// The resulting composed quaternion
 	UPROPERTY(meta=(Output))
 	FQuat Result;
 };
@@ -149,12 +177,15 @@ struct FRigVMFunction_MathQuaternionFromEuler : public FRigVMFunction_MathQuater
 	RIGVM_METHOD()
 	UE_API virtual void Execute() override;
 
+	// The euler representation of the rotation in degrees.
 	UPROPERTY(meta = (Input))
 	FVector Euler;
 
+	// The rotation order to use when interpreting the euler angles
 	UPROPERTY(meta = (Input))
 	EEulerRotationOrder RotationOrder;
 
+	// The resulting composed quaternion
 	UPROPERTY(meta=(Output))
 	FQuat Result;
 };
@@ -203,9 +234,11 @@ struct FRigVMFunction_MathQuaternionFromRotatorV2 : public FRigVMFunction_MathQu
 	RIGVM_METHOD()
 	UE_API virtual void Execute() override;
 
+	// The rotator representation of the rotation
 	UPROPERTY(meta=(Input))
 	FRotator Value;
 
+	// The resulting composed quaternion
 	UPROPERTY(meta=(Output))
 	FQuat Result;
 };
@@ -227,12 +260,15 @@ struct FRigVMFunction_MathQuaternionFromTwoVectors : public FRigVMFunction_MathQ
 	RIGVM_METHOD()
 	UE_API virtual void Execute() override;
 
+	// A vector representing a direction prior to rotation
 	UPROPERTY(meta=(Input))
 	FVector A;
 
+	// A vector representing a direction after a rotation
 	UPROPERTY(meta=(Input))
 	FVector B;
 
+	// The quaternion representing the rotation from A to B
 	UPROPERTY(meta=(Output))
 	FQuat Result;
 };
@@ -255,12 +291,15 @@ struct FRigVMFunction_MathQuaternionToAxisAndAngle : public FRigVMFunction_MathQ
 	RIGVM_METHOD()
 	UE_API virtual void Execute() override;
 
+	// The input quaternion
 	UPROPERTY(meta=(Input))
 	FQuat Value;
 
+	// The axis of the quaternion
 	UPROPERTY(meta=(Output))
 	FVector Axis;
 
+	// The angle in radians of the quaternion
 	UPROPERTY(meta=(Output))
 	float Angle;
 };
@@ -282,15 +321,19 @@ struct FRigVMFunction_MathQuaternionToVectors : public FRigVMFunction_MathQuater
 	RIGVM_METHOD()
 	UE_API virtual void Execute() override;
 
+	// The input quaternion
 	UPROPERTY(meta=(Input))
 	FQuat Value;
 
+	// The forward axis of the quaternion
 	UPROPERTY(meta=(Output))
 	FVector Forward;
 
+	// The right axis of the quaternion
 	UPROPERTY(meta=(Output))
 	FVector Right;
 
+	// The up axis of the quaternion
 	UPROPERTY(meta=(Output))
 	FVector Up;
 };
@@ -339,12 +382,15 @@ struct FRigVMFunction_MathQuaternionScaleV2 : public FRigVMFunction_MathQuaterni
 	RIGVM_METHOD()
 	UE_API virtual void Execute() override;
 
+	// The input quaternion
 	UPROPERTY(meta=(Input))
 	FQuat Value;
 
+	// The scale factor to scale the input by
 	UPROPERTY(meta=(Input))
 	float Factor;
 
+	// The resulting scaled quaternion
 	UPROPERTY(meta=(Output))
 	FQuat Result;
 };
@@ -367,12 +413,15 @@ struct FRigVMFunction_MathQuaternionToEuler : public FRigVMFunction_MathQuaterni
 	RIGVM_METHOD()
 	UE_API virtual void Execute() override;
 
+	// The input quaternion
 	UPROPERTY(meta=(Input))
 	FQuat Value;
 
+	// The rotation order to use when encoding the euler angles
 	UPROPERTY(meta = (Input))
 	EEulerRotationOrder RotationOrder;
 
+	// The euler angles in degrees representing the quaternion in the given rotation order
 	UPROPERTY(meta=(Output))
 	FVector Result;
 };
@@ -394,9 +443,11 @@ struct FRigVMFunction_MathQuaternionToRotator : public FRigVMFunction_MathQuater
 	RIGVM_METHOD()
 	UE_API virtual void Execute() override;
 
+	// The input quaternion
 	UPROPERTY(meta=(Input))
 	FQuat Value;
 
+	// The rotator representation of the rotation
 	UPROPERTY(meta=(Output))
 	FRotator Result;
 };
@@ -424,7 +475,7 @@ struct FRigVMFunction_MathQuaternionInverse : public FRigVMFunction_MathQuaterni
 };
 
 /**
- * Linearly interpolates between A and B using the ratio T
+ * Performs a spherical interpolation of the quaternions A and B based on the blend value T.
  */
 USTRUCT(meta=(DisplayName="Interpolate", TemplateName="Interpolate", Keywords="Lerp,Mix,Blend,Slerp,SphericalInterpolate"))
 struct FRigVMFunction_MathQuaternionSlerp : public FRigVMFunction_MathQuaternionBase
@@ -440,15 +491,19 @@ struct FRigVMFunction_MathQuaternionSlerp : public FRigVMFunction_MathQuaternion
 	RIGVM_METHOD()
 	UE_API virtual void Execute() override;
 
+	// The first quaternion to interpolate from
 	UPROPERTY(meta=(Input))
 	FQuat A;
 
+	// The second quaternion to interpolate to
 	UPROPERTY(meta=(Input))
 	FQuat B;
 
-	UPROPERTY(meta=(Input))
+	// The blend value for the interpolation
+	UPROPERTY(meta=(Input, UIMin = "0", UIMax = "1"))
 	float T;
 
+	// The resulting interpolated quaternion
 	UPROPERTY(meta=(Output))
 	FQuat Result;
 };
@@ -563,12 +618,15 @@ struct FRigVMFunction_MathQuaternionDot : public FRigVMFunction_MathQuaternionBa
 	RIGVM_METHOD()
 	UE_API virtual void Execute() override;
 
+	// The first quaternion to compute the dot product for
 	UPROPERTY(meta=(Input))
 	FQuat A;
 
+	// The second quaternion to compute the dot product for
 	UPROPERTY(meta=(Input))
 	FQuat B;
 
+	// The resulting dot product between two quaternions
 	UPROPERTY(meta=(Output))
 	float Result;
 };
@@ -601,18 +659,21 @@ struct FRigVMFunction_MathQuaternionRotateVector : public FRigVMFunction_MathQua
 	RIGVM_METHOD()
 	UE_API virtual void Execute() override;
 
+	// The input quaternion to rotate the input vector by
 	UPROPERTY(meta=(Input))
 	FQuat Transform;
 
+	// The input vector to rotate
 	UPROPERTY(meta=(Input))
 	FVector Vector;
 
+	// The resulting rotated vector
 	UPROPERTY(meta=(Output))
 	FVector Result;
 };
 
 /**
- * Rotates a given vector by the quaternion
+ * Extracts an axis from a quaternion rotation
  */
 USTRUCT(meta = (DisplayName = "Axis", TemplateName="Axis", Keywords = "GetAxis,xAxis,yAxis,zAxis"))
 struct FRigVMFunction_MathQuaternionGetAxis: public FRigVMFunction_MathQuaternionBase
@@ -629,12 +690,15 @@ struct FRigVMFunction_MathQuaternionGetAxis: public FRigVMFunction_MathQuaternio
 	RIGVM_METHOD()
 	UE_API virtual void Execute() override;
 
+	// The input quaternion
 	UPROPERTY(meta = (Input))
 	FQuat Quaternion;
 
+	// The axis to extract
 	UPROPERTY(meta = (Input))
 	TEnumAsByte<EAxis::Type> Axis;
 
+	// The resulting extracted axis
 	UPROPERTY(meta = (Output))
 	FVector Result;
 };
@@ -657,15 +721,19 @@ struct FRigVMFunction_MathQuaternionSwingTwist : public FRigVMFunction_MathQuate
 	RIGVM_METHOD()
 	UE_API virtual void Execute() override;
 
+	// The input quaternion
 	UPROPERTY(meta = (Input))
 	FQuat	Input;
 
+	// The axis to treat as the twist axis
 	UPROPERTY(meta = (Input))
 	FVector TwistAxis;
 
+	// The resulting swing quaternion of the input quaternion
 	UPROPERTY(meta = (Output))
 	FQuat Swing;
 
+	// The resulting twist quaternion of the input quaternion
 	UPROPERTY(meta = (Output))
 	FQuat Twist;
 };
@@ -683,6 +751,7 @@ struct FRigVMFunction_MathQuaternionRotationOrder : public FRigVMFunction_MathBa
 		RotationOrder = EEulerRotationOrder::ZYX;
 	}
 
+	// The constant rotation order
 	UPROPERTY(meta = (Input, Output))
 	EEulerRotationOrder RotationOrder;
 
@@ -707,12 +776,15 @@ struct FRigVMFunction_MathQuaternionMakeRelative : public FRigVMFunction_MathQua
 	RIGVM_METHOD()
 	UE_API virtual void Execute() override;
 
+	// The global rotation quaternion to make relative
 	UPROPERTY(meta=(Input))
 	FQuat Global;
 
+	// The parent global rotation to make the input relative to
 	UPROPERTY(meta=(Input))
 	FQuat Parent;
 
+	// The resulting relative rotation
 	UPROPERTY(meta=(Output))
 	FQuat Local;
 };
@@ -733,12 +805,15 @@ struct FRigVMFunction_MathQuaternionMakeAbsolute : public FRigVMFunction_MathQua
 	RIGVM_METHOD()
 	UE_API virtual void Execute() override;
 
+	// The input local rotation to make absolute
 	UPROPERTY(meta = (Input))
 	FQuat Local;
 
+	// The parent global rotation to make the local absolute to
 	UPROPERTY(meta = (Input))
 	FQuat Parent;
 
+	// The resulting global rotation
 	UPROPERTY(meta = (Output))
 	FQuat Global;
 };
@@ -762,6 +837,7 @@ struct FRigVMFunction_MathQuaternionMirrorTransform : public FRigVMFunction_Math
 	RIGVM_METHOD()
 	UE_API virtual void Execute() override;
 
+	// The input quaternion
 	UPROPERTY(meta=(Input))
 	FQuat Value;
 
@@ -777,6 +853,7 @@ struct FRigVMFunction_MathQuaternionMirrorTransform : public FRigVMFunction_Math
 	UPROPERTY(meta=(Input))
 	FTransform CentralTransform;
 
+	// The resulting mirrored rotation
 	UPROPERTY(meta=(Output))
 	FQuat Result;
 };
diff --git a/Engine/Plugins/Runtime/RigVM/Source/RigVM/Public/RigVMFunctions/Math/RigVMFunction_MathRBFInterpolate.h b/Engine/Plugins/Runtime/RigVM/Source/RigVM/Public/RigVMFunctions/Math/RigVMFunction_MathRBFInterpolate.h
--- a/Engine/Plugins/Runtime/RigVM/Source/RigVM/Public/RigVMFunctions/Math/RigVMFunction_MathRBFInterpolate.h
+++ b/Engine/Plugins/Runtime/RigVM/Source/RigVM/Public/RigVMFunctions/Math/RigVMFunction_MathRBFInterpolate.h
@@ -73,6 +73,9 @@ struct FRigVMFunction_MathRBFInterpolateVectorWorkData
 };
 
 
+/**
+ * The base class for all pure RBF nodes
+ */
 USTRUCT(meta = (Abstract, Category = "Math|RBF Interpolation"))
 struct FRigVMFunction_MathRBFInterpolateBase : 
 	public FRigVMFunction_MathBase
@@ -80,27 +83,36 @@ struct FRigVMFunction_MathRBFInterpolateBase :
 	GENERATED_BODY()
 };
 
+/**
+ * The base class for all quaternion RBF nodes
+ */
 USTRUCT(meta = (Abstract, TemplateName="RBF Quaternion", Keywords = "RBF,Interpolate,Quaternion"))
 struct FRigVMFunction_MathRBFInterpolateQuatBase :
 	public FRigVMFunction_MathRBFInterpolateBase
 {
 	GENERATED_BODY()
 
+	// The input quaternion
 	UPROPERTY(meta = (Input))
 	FQuat Input = FQuat::Identity;
 
+	// The distance function to use
 	UPROPERTY(meta = (Input, Constant))
 	ERBFQuatDistanceType DistanceFunction = ERBFQuatDistanceType::ArcLength;
 
+	// The smoothing function to use
 	UPROPERTY(meta = (Input, Constant))
 	ERBFKernelType SmoothingFunction = ERBFKernelType::Gaussian;
 
-	UPROPERTY(meta = (Input, Constant))
+	// The smoothing angle
+	UPROPERTY(meta = (Input, Constant, UIMin = "0", UIMax = "180"))
 	float SmoothingAngle = 45.0f;
 
+	// If true the resulting output will be normalized
 	UPROPERTY(meta = (Input, Constant))
 	bool bNormalizeOutput = false;
 
+	// The twist axis of the input quaternion
 	UPROPERTY(meta = (Input, EditCondition = "DistanceFunction == ERBFQuatDistanceType::SwingAngle || DistanceFunction == ERBFQuatDistanceType::TwistAngle"))
 	FVector TwistAxis = FVector::ForwardVector;
 
@@ -125,24 +137,32 @@ protected:
 	static UE_API uint64 HashTargets(const TArrayView<const FQuat>& Targets);
 };
 
+/**
+ * The base class for all vector RBF nodes
+ */
 USTRUCT(meta = (Abstract, TemplateName="RBF Vector", Keywords = "RBF,Interpolate,Vector"))
 struct FRigVMFunction_MathRBFInterpolateVectorBase :
 	public FRigVMFunction_MathRBFInterpolateBase
 {
 	GENERATED_BODY()
 
+	// The input vector
 	UPROPERTY(meta = (Input))
 	FVector Input = FVector::ZeroVector;
 
+	// The distance function to use
 	UPROPERTY(meta = (Input, Constant))
 	ERBFVectorDistanceType DistanceFunction = ERBFVectorDistanceType::Euclidean;
 
+	// The smoothing function to use
 	UPROPERTY(meta = (Input, Constant))
 	ERBFKernelType SmoothingFunction = ERBFKernelType::Gaussian;
 
+	// The radius to apply for the smoothing function
 	UPROPERTY(meta = (Input, Constant))
 	float SmoothingRadius = 5.0f;
 
+	// If true the resulting output will be normalized
 	UPROPERTY(meta = (Input, Constant))
 	bool bNormalizeOutput = false;
 
@@ -181,15 +201,18 @@ struct FMathRBFInterpolateQuatFloat_Target
 	float Value = 0.0f;
 };
 
+// A RBF interpolator from quaternion to float
 USTRUCT(meta = (Abstract, DisplayName = "RBF Quaternion to Float"))
 struct FRigVMFunction_MathRBFInterpolateQuatFloat : 
 	public FRigVMFunction_MathRBFInterpolateQuatBase
 {
 	GENERATED_BODY()
 
+	// The array of targets to interpolate within
 	UPROPERTY(meta = (Input))
 	TArray<FMathRBFInterpolateQuatFloat_Target> Targets;
 
+	// The interpolated result
 	UPROPERTY(meta = (Output))
 	float Output = 0.0f;
 
@@ -210,15 +233,18 @@ struct FMathRBFInterpolateQuatVector_Target
 	FVector Value = FVector::ZeroVector;
 };
 
+// A RBF interpolator from quaternion to vector
 USTRUCT(meta = (Abstract, DisplayName = "RBF Quaternion to Vector"))
 struct FRigVMFunction_MathRBFInterpolateQuatVector :
 	public FRigVMFunction_MathRBFInterpolateQuatBase
 {
 	GENERATED_BODY()
 
+	// The array of targets to interpolate within
 	UPROPERTY(meta = (Input))
 	TArray<FMathRBFInterpolateQuatVector_Target> Targets;
 
+	// The interpolated result
 	UPROPERTY(meta = (Output))
 	FVector Output = FVector::ZeroVector;
 
@@ -239,15 +265,18 @@ struct FMathRBFInterpolateQuatColor_Target
 	FLinearColor Value = FLinearColor::Transparent;
 };
 
+// A RBF interpolator from quaternion to color
 USTRUCT(meta = (Abstract, DisplayName = "RBF Quaternion to Color"))
 struct FRigVMFunction_MathRBFInterpolateQuatColor :
 	public FRigVMFunction_MathRBFInterpolateQuatBase
 {
 	GENERATED_BODY()
 
+	// The array of targets to interpolate within
 	UPROPERTY(meta = (Input))
 	TArray<FMathRBFInterpolateQuatColor_Target> Targets;
 
+	// The interpolated result
 	UPROPERTY(meta = (Output))
 	FLinearColor Output = FLinearColor::Transparent;
 
@@ -268,15 +297,18 @@ struct FMathRBFInterpolateQuatQuat_Target
 	FQuat Value = FQuat::Identity;
 };
 
+// A RBF interpolator from quaternion to quaternion
 USTRUCT(meta = (Abstract, DisplayName = "RBF Quaternion to Quaternion"))
 struct FRigVMFunction_MathRBFInterpolateQuatQuat :
 	public FRigVMFunction_MathRBFInterpolateQuatBase
 {
 	GENERATED_BODY()
 
+	// The array of targets to interpolate within
 	UPROPERTY(meta = (Input))
 	TArray<FMathRBFInterpolateQuatQuat_Target> Targets;
 
+	// The interpolated result
 	UPROPERTY(meta = (Output))
 	FQuat Output = FQuat::Identity;
 
@@ -297,15 +329,18 @@ struct FMathRBFInterpolateQuatXform_Target
 	FTransform Value = FTransform::Identity;
 };
 
+// A RBF interpolator from quaternion to transform
 USTRUCT(meta = (Abstract, DisplayName = "RBF Quaternion to Transform"))
 struct FRigVMFunction_MathRBFInterpolateQuatXform :
 	public FRigVMFunction_MathRBFInterpolateQuatBase
 {
 	GENERATED_BODY()
 
+	// The array of targets to interpolate within
 	UPROPERTY(meta = (Input))
 	TArray<FMathRBFInterpolateQuatXform_Target> Targets;
 
+	// The interpolated result
 	UPROPERTY(meta = (Output))
 	FTransform Output = FTransform::Identity;
 
@@ -328,15 +363,18 @@ struct FMathRBFInterpolateVectorFloat_Target
 	float Value = 0.0f;
 };
 
+// A RBF interpolator from vector to float
 USTRUCT(meta = (Abstract, DisplayName = "RBF Vector to Float", Keywords = "RBF,Interpolate,Vector"))
 struct FRigVMFunction_MathRBFInterpolateVectorFloat :
 	public FRigVMFunction_MathRBFInterpolateVectorBase
 {
 	GENERATED_BODY()
 
+	// The array of targets to interpolate within
 	UPROPERTY(meta = (Input))
 	TArray<FMathRBFInterpolateVectorFloat_Target> Targets;
 
+	// The interpolated result
 	UPROPERTY(meta = (Output))
 	float Output = 0.0f;
 
@@ -357,15 +395,18 @@ struct FMathRBFInterpolateVectorVector_Target
 	FVector Value = FVector::ZeroVector;
 };
 
+// A RBF interpolator from vector to vector
 USTRUCT(meta = (Abstract, DisplayName = "RBF Vector to Vector", Keywords = "RBF,Interpolate,Vector"))
 struct FRigVMFunction_MathRBFInterpolateVectorVector :
 	public FRigVMFunction_MathRBFInterpolateVectorBase
 {
 	GENERATED_BODY()
 
+	// The array of targets to interpolate within
 	UPROPERTY(meta = (Input))
 	TArray<FMathRBFInterpolateVectorVector_Target> Targets;
 
+	// The interpolated result
 	UPROPERTY(meta = (Output))
 	FVector Output = FVector::ZeroVector;
 
@@ -386,15 +427,18 @@ struct FMathRBFInterpolateVectorColor_Target
 	FLinearColor Value = FLinearColor::Transparent;
 };
 
+// A RBF interpolator from vector to color
 USTRUCT(meta = (Abstract, DisplayName = "RBF Vector to Color", Keywords = "RBF,Interpolate,Vector"))
 struct FRigVMFunction_MathRBFInterpolateVectorColor :
 	public FRigVMFunction_MathRBFInterpolateVectorBase
 {
 	GENERATED_BODY()
 
+	// The array of targets to interpolate within
 	UPROPERTY(meta = (Input))
 	TArray<FMathRBFInterpolateVectorColor_Target> Targets;
 
+	// The interpolated result
 	UPROPERTY(meta = (Output))
 	FLinearColor Output = FLinearColor::Transparent;
 
@@ -416,15 +460,18 @@ struct FMathRBFInterpolateVectorQuat_Target
 	FQuat Value = FQuat::Identity;
 };
 
+// A RBF interpolator from vector to quaternion
 USTRUCT(meta = (Abstract, DisplayName = "RBF Vector to Quat", Keywords = "RBF,Interpolate,Vector"))
 struct FRigVMFunction_MathRBFInterpolateVectorQuat :
 	public FRigVMFunction_MathRBFInterpolateVectorBase
 {
 	GENERATED_BODY()
 
+	// The array of targets to interpolate within
 	UPROPERTY(meta = (Input))
 	TArray<FMathRBFInterpolateVectorQuat_Target> Targets;
 
+	// The interpolated result
 	UPROPERTY(meta = (Output))
 	FQuat Output = FQuat::Identity;
 
@@ -444,15 +491,18 @@ struct FMathRBFInterpolateVectorXform_Target
 	FTransform Value = FTransform::Identity;
 };
 
+// A RBF interpolator from vector to transform
 USTRUCT(meta = (Abstract, DisplayName = "RBF Vector to Transform", Keywords = "RBF,Interpolate,Vector"))
 struct FRigVMFunction_MathRBFInterpolateVectorXform :
 	public FRigVMFunction_MathRBFInterpolateVectorBase
 {
 	GENERATED_BODY()
 
+	// The array of targets to interpolate within
 	UPROPERTY(meta = (Input))
 	TArray<FMathRBFInterpolateVectorXform_Target> Targets;
 
+	// The interpolated result
 	UPROPERTY(meta = (Output))
 	FTransform Output = FTransform::Identity;
 
diff --git a/Engine/Plugins/Runtime/RigVM/Source/RigVM/Public/RigVMFunctions/Math/RigVMFunction_MathRay.h b/Engine/Plugins/Runtime/RigVM/Source/RigVM/Public/RigVMFunctions/Math/RigVMFunction_MathRay.h
--- a/Engine/Plugins/Runtime/RigVM/Source/RigVM/Public/RigVMFunctions/Math/RigVMFunction_MathRay.h
+++ b/Engine/Plugins/Runtime/RigVM/Source/RigVM/Public/RigVMFunctions/Math/RigVMFunction_MathRay.h
@@ -7,6 +7,9 @@
 #include "RigVMFunction_MathBase.h"
 #include "RigVMFunction_MathRay.generated.h"
 
+/*
+ * The base class for all pure ray math nodes
+ */
 USTRUCT(meta=(Abstract, Category="Math|Ray", MenuDescSuffix="(Ray)"))
 struct FRigVMFunction_MathRayBase : public FRigVMFunction_MathBase
 {
@@ -31,21 +34,29 @@ struct FRigVMFunction_MathRayIntersectRay : public FRigVMFunction_MathRayBase
 	RIGVM_METHOD()
 	RIGVM_API virtual void Execute() override;
 
+	// The first ray to intersect
 	UPROPERTY(meta = (Input))
 	FRay A;
 
+	// The second ray to intersect
 	UPROPERTY(meta = (Input))
 	FRay B;
 
+	// The resulting intersection position.
+	// This is either on the rays themselves
+	// or at the closest position between the two.
 	UPROPERTY(meta = (Output))
 	FVector Result;
 
+	// The distance between the two rays (or 0 if they touch)
 	UPROPERTY(meta = (Output))
 	float Distance;
 
+	// The ratio on the first ray at the intersection
 	UPROPERTY(meta = (Output))
 	float RatioA;
 
+	// The ratio on the second ray at the intersection
 	UPROPERTY(meta = (Output))
 	float RatioB;
 };
@@ -70,21 +81,27 @@ struct FRigVMFunction_MathRayIntersectPlane : public FRigVMFunction_MathRayBase
 	RIGVM_METHOD()
 	RIGVM_API virtual void Execute() override;
 
+	// The ray to intersect with the plane
 	UPROPERTY(meta = (Input))
 	FRay Ray;
 
+	// The point on the plane to intersect the ray with
 	UPROPERTY(meta = (Input))
 	FVector PlanePoint;
 
+	// The normal of the plane to intersect the ray with
 	UPROPERTY(meta = (Input))
 	FVector PlaneNormal;
 
+	// The resulting intersection position
 	UPROPERTY(meta = (Output))
 	FVector Result;
 
+	// The distance of the ray origin to the plane
 	UPROPERTY(meta = (Output))
 	float Distance;
 
+	// The ratio along the ray up to the intersection point
 	UPROPERTY(meta = (Output))
 	float Ratio;
 };
@@ -107,18 +124,21 @@ struct FRigVMFunction_MathRayGetAt : public FRigVMFunction_MathRayBase
 	RIGVM_METHOD()
 	RIGVM_API virtual void Execute() override;
 
+	// The ray to query
 	UPROPERTY(meta = (Input))
 	FRay Ray;
 
+	// The ratio to query the ray at
 	UPROPERTY(meta = (Input))
 	float Ratio;
 
+	// The resulting position on the ray
 	UPROPERTY(meta = (Output))
 	FVector Result;
 };
 
 /**
- * Returns the division of the two values
+ * Transforms a ray
  */
 USTRUCT(meta=(DisplayName="Transform Ray", Keywords="Multiply,Project"))
 struct FRigVMFunction_MathRayTransform : public FRigVMFunction_MathRayBase
@@ -134,12 +154,15 @@ struct FRigVMFunction_MathRayTransform : public FRigVMFunction_MathRayBase
 	RIGVM_METHOD()
 	RIGVM_API virtual void Execute() override;
 
+	// The ray to transform
 	UPROPERTY(meta = (Input))
 	FRay Ray;
 
+	// The transform to apply to the ray
 	UPROPERTY(meta = (Input))
 	FTransform Transform;
 
+	// The resulting transformed ray
 	UPROPERTY(meta = (Output))
 	FRay Result;
 };
diff --git a/Engine/Plugins/Runtime/RigVM/Source/RigVM/Public/RigVMFunctions/Math/RigVMFunction_MathTransform.h b/Engine/Plugins/Runtime/RigVM/Source/RigVM/Public/RigVMFunctions/Math/RigVMFunction_MathTransform.h
--- a/Engine/Plugins/Runtime/RigVM/Source/RigVM/Public/RigVMFunctions/Math/RigVMFunction_MathTransform.h
+++ b/Engine/Plugins/Runtime/RigVM/Source/RigVM/Public/RigVMFunctions/Math/RigVMFunction_MathTransform.h
@@ -8,18 +8,27 @@
 
 #define UE_API RIGVM_API
 
+/*
+ * The base class for all pure transform math nodes
+ */
 USTRUCT(meta=(Abstract, Category="Math|Transform", MenuDescSuffix="(Transform)"))
 struct FRigVMFunction_MathTransformBase : public FRigVMFunction_MathBase
 {
 	GENERATED_BODY()
 };
 
+/*
+ * The base class for all mutable transform math nodes
+ */
 USTRUCT(meta=(Abstract, Category="Math|Transform", MenuDescSuffix="(Transform)"))
 struct FRigVMFunction_MathTransformMutableBase : public FRigVMFunction_MathMutableBase
 {
 	GENERATED_BODY()
 };
 
+/*
+ * The base class for all unary transform math operations
+ */
 USTRUCT(meta=(Abstract))
 struct FRigVMFunction_MathTransformUnaryOp : public FRigVMFunction_MathTransformBase
 {
@@ -30,13 +39,18 @@ struct FRigVMFunction_MathTransformUnaryOp : public FRigVMFunction_MathTransform
 		Value = Result = FTransform::Identity;
 	}
 
+	// The input value
 	UPROPERTY(meta=(Input))
 	FTransform Value;
 
+	// The result of the operation
 	UPROPERTY(meta=(Output))
 	FTransform Result;
 };
 
+/*
+ * The base class for all binary transform math operations
+ */
 USTRUCT(meta=(Abstract))
 struct FRigVMFunction_MathTransformBinaryOp : public FRigVMFunction_MathTransformBase
 {
@@ -47,16 +61,22 @@ struct FRigVMFunction_MathTransformBinaryOp : public FRigVMFunction_MathTransfor
 		A = B = Result = FTransform::Identity;
 	}
 
+	// The first transform for the operation
 	UPROPERTY(meta=(Input))
 	FTransform A;
 
+	// The second transform for the operation
 	UPROPERTY(meta=(Input))
 	FTransform B;
 
+	// The resulting transform
 	UPROPERTY(meta=(Output))
 	FTransform Result;
 };
 
+/*
+ * The base class for all aggregational transform math operations
+ */
 USTRUCT(meta=(Abstract))
 struct FRigVMFunction_MathTransformBinaryAggregateOp : public FRigVMFunction_MathTransformBase
 {
@@ -67,12 +87,15 @@ struct FRigVMFunction_MathTransformBinaryAggregateOp : public FRigVMFunction_Mat
 		A = B = Result = FTransform::Identity;
 	}
 
+	// The first transform for the operation
 	UPROPERTY(meta=(Input, Aggregate))
 	FTransform A;
 
+	// The second transform for the operation
 	UPROPERTY(meta=(Input, Aggregate))
 	FTransform B;
 
+	// The resulting transform
 	UPROPERTY(meta=(Output, Aggregate))
 	FTransform Result;
 };
@@ -155,9 +178,11 @@ struct FRigVMFunction_MathTransformFromEulerTransformV2 : public FRigVMFunction_
 	RIGVM_METHOD()
 	UE_API virtual void Execute() override;
 
+	// The input euler transform to convert
 	UPROPERTY(meta=(Input))
 	FEulerTransform Value;
 
+	// The resulting composed transform
 	UPROPERTY(meta=(Output))
 	FTransform Result;
 };
@@ -179,9 +204,11 @@ struct FRigVMFunction_MathTransformToEulerTransform : public FRigVMFunction_Math
 	RIGVM_METHOD()
 	UE_API virtual void Execute() override;
 
+	// The input transform to convert
 	UPROPERTY(meta=(Input))
 	FTransform Value;
 
+	// The resulting composed euler transform
 	UPROPERTY(meta=(Output))
 	FEulerTransform Result;
 };
@@ -203,15 +230,19 @@ struct FRigVMFunction_MathTransformToVectors : public FRigVMFunction_MathTransfo
 	RIGVM_METHOD()
 	UE_API virtual void Execute() override;
 
+	// The input transform
 	UPROPERTY(meta=(Input))
 	FTransform Value;
 
+	// The rotation's forward axis
 	UPROPERTY(meta=(Output))
 	FVector Forward;
 
+	// The rotation's right axis
 	UPROPERTY(meta=(Output))
 	FVector Right;
 
+	// The rotation's up axis
 	UPROPERTY(meta=(Output))
 	FVector Up;
 };
@@ -243,12 +274,15 @@ struct FRigVMFunction_MathTransformMakeRelative : public FRigVMFunction_MathTran
 	RIGVM_METHOD()
 	UE_API virtual void Execute() override;
 
+	// The global transform to make relative
 	UPROPERTY(meta=(Input))
 	FTransform Global;
 
+	// The parent global transform to make the input relative to
 	UPROPERTY(meta=(Input))
 	FTransform Parent;
 
+	// The resulting relative transform
 	UPROPERTY(meta=(Output))
 	FTransform Local;
 };
@@ -269,12 +303,15 @@ struct FRigVMFunction_MathTransformMakeAbsolute : public FRigVMFunction_MathTran
 	RIGVM_METHOD()
 	UE_API virtual void Execute() override;
 
+	// The input local transform to make absolute
 	UPROPERTY(meta = (Input))
 	FTransform Local;
 
+	// The parent global transform to make the local absolute to
 	UPROPERTY(meta = (Input))
 	FTransform Parent;
 
+	// The resulting global transform
 	UPROPERTY(meta = (Output))
 	FTransform Global;
 };
@@ -300,6 +337,7 @@ struct FRigVMFunction_MathTransformAccumulateArray : public FRigVMFunction_MathT
 	RIGVM_METHOD()
 	UE_API virtual void Execute() override;
 
+	// The input transform array to accumulate
 	UPROPERTY(meta=(Input, Output))
 	TArray<FTransform> Transforms;
 
@@ -351,15 +389,19 @@ struct FRigVMFunction_MathTransformLerp : public FRigVMFunction_MathTransformBas
 	RIGVM_METHOD()
 	UE_API virtual void Execute() override;
 
+	// The first transform to interpolate from
 	UPROPERTY(meta=(Input))
 	FTransform A;
 
+	// The second transform to interpolate to
 	UPROPERTY(meta=(Input))
 	FTransform B;
 
-	UPROPERTY(meta=(Input))
+	// The blend value for the interpolation
+	UPROPERTY(meta=(Input, UIMin = "0", UIMax = "1"))
 	float T;
 
+	// The resulting interpolated transform
 	UPROPERTY(meta=(Output))
 	FTransform Result;
 };
@@ -413,13 +455,16 @@ struct FRigVMFunction_MathTransformRotateVector : public FRigVMFunction_MathTran
 
 	RIGVM_METHOD()
 	UE_API virtual void Execute() override;
-
+	
+	// The input transform to rotate the input vector by
 	UPROPERTY(meta=(Input))
 	FTransform Transform;
 
+	// The input vector to rotate
 	UPROPERTY(meta=(Input))
 	FVector Vector;
 
+	// The resulting rotated vector
 	UPROPERTY(meta=(Output))
 	FVector Result;
 };
@@ -441,12 +486,15 @@ struct FRigVMFunction_MathTransformTransformVector : public FRigVMFunction_MathT
 	RIGVM_METHOD()
 	UE_API virtual void Execute() override;
 
+	// The input transform to transform the input vector by
 	UPROPERTY(meta=(Input))
 	FTransform Transform;
 
+	// The input vector to transform
 	UPROPERTY(meta=(Input))
 	FVector Location;
 
+	// The resulting transformed vector
 	UPROPERTY(meta=(Output))
 	FVector Result;
 };
@@ -472,21 +520,27 @@ struct FRigVMFunction_MathTransformFromSRT : public FRigVMFunction_MathTransform
 	RIGVM_METHOD()
 	UE_API virtual void Execute() override;
 
+	// The location for the desired transform
 	UPROPERTY(meta=(Input))
 	FVector Location;
 
+	// The euler angles in degrees for the desired transform
 	UPROPERTY(meta=(Input))
 	FVector Rotation;
 
+	// The rotation order to interpret the euler angles by
 	UPROPERTY(meta=(Input))
 	EEulerRotationOrder RotationOrder;
 
+	// The scale for the desired transform
 	UPROPERTY(meta=(Input))
 	FVector Scale;
 
+	// The result as a transform
 	UPROPERTY(meta=(Output))
 	FTransform Transform;
 
+	// The result as a euler transform
 	UPROPERTY(meta=(Output))
 	FEulerTransform EulerTransform;
 };
@@ -507,15 +561,19 @@ struct FRigVMFunction_MathTransformArrayToSRT : public FRigVMFunction_MathTransf
 	RIGVM_METHOD()
 	UE_API virtual void Execute() override;
 
+	// The input transform array to decompose
 	UPROPERTY(meta = (Input))
 	TArray<FTransform> Transforms;
 
+	// The array of translations - one for each input transform
 	UPROPERTY(meta = (Output))
 	TArray<FVector> Translations;
 
+	// The array of rotation - one for each input transform
 	UPROPERTY(meta = (Output))
 	TArray<FQuat> Rotations;
 
+	// The array of scale values - one for each input transform
 	UPROPERTY(meta = (Output))
 	TArray<FVector> Scales;
 };
@@ -544,12 +602,15 @@ struct FRigVMFunction_MathTransformClampSpatially : public FRigVMFunction_MathTr
 	RIGVM_METHOD()
 	UE_API virtual void Execute() override;
 
+	// The input transform to clamp
 	UPROPERTY(meta = (Input))
 	FTransform Value;
 
+	// The axis to use for the filter
 	UPROPERTY(meta = (Input))
 	TEnumAsByte<EAxis::Type> Axis;
 
+	// The filter / spatial mode to use
 	UPROPERTY(meta = (Input))
 	TEnumAsByte<ERigVMClampSpatialMode::Type> Type;
 
@@ -571,15 +632,19 @@ struct FRigVMFunction_MathTransformClampSpatially : public FRigVMFunction_MathTr
 	UPROPERTY(meta = (Input))
 	FTransform Space;
 
+	// Draws debug information if True
 	UPROPERTY(meta = (Input))
 	bool bDrawDebug;
 
+	// The color to use for the debug draw
 	UPROPERTY(meta = (Input))
 	FLinearColor DebugColor;
 
+	// The thickness to use for the debug draw
 	UPROPERTY(meta = (Input))
 	float DebugThickness;
 
+	// The resulting transform with a clamped position
 	UPROPERTY(meta = (Output))
 	FTransform Result;
 };
@@ -603,6 +668,7 @@ struct FRigVMFunction_MathTransformMirrorTransform : public FRigVMFunction_MathT
 	RIGVM_METHOD()
 	UE_API virtual void Execute() override;
 
+	// The input transform to mirror
 	UPROPERTY(meta=(Input))
 	FTransform Value;
 
@@ -618,6 +684,7 @@ struct FRigVMFunction_MathTransformMirrorTransform : public FRigVMFunction_MathT
 	UPROPERTY(meta=(Input))
 	FTransform CentralTransform;
 
+	// The resulting mirrored transform
 	UPROPERTY(meta=(Output))
 	FTransform Result;
 };
diff --git a/Engine/Plugins/Runtime/RigVM/Source/RigVM/Public/RigVMFunctions/Math/RigVMFunction_MathVector.h b/Engine/Plugins/Runtime/RigVM/Source/RigVM/Public/RigVMFunctions/Math/RigVMFunction_MathVector.h
--- a/Engine/Plugins/Runtime/RigVM/Source/RigVM/Public/RigVMFunctions/Math/RigVMFunction_MathVector.h
+++ b/Engine/Plugins/Runtime/RigVM/Source/RigVM/Public/RigVMFunctions/Math/RigVMFunction_MathVector.h
@@ -6,12 +6,18 @@
 #include "RigVMMathLibrary.h"
 #include "RigVMFunction_MathVector.generated.h"
 
+/*
+ * The base class for all pure vector math nodes
+ */
 USTRUCT(meta=(Abstract, Category="Math|Vector", MenuDescSuffix="(Vector)"))
 struct FRigVMFunction_MathVectorBase : public FRigVMFunction_MathBase
 {
 	GENERATED_BODY()
 };
 
+/*
+ * The base class for all unary vector math operations
+ */
 USTRUCT(meta=(Abstract))
 struct FRigVMFunction_MathVectorUnaryOp : public FRigVMFunction_MathVectorBase
 {
@@ -22,13 +28,18 @@ struct FRigVMFunction_MathVectorUnaryOp : public FRigVMFunction_MathVectorBase
 		Value = Result = FVector::ZeroVector;
 	}
 
+	// The input value
 	UPROPERTY(meta=(Input))
 	FVector Value;
 
+	// The result of the operation
 	UPROPERTY(meta=(Output))
 	FVector Result;
 };
 
+/*
+ * The base class for all binary vector math operations
+ */
 USTRUCT(meta=(Abstract))
 struct FRigVMFunction_MathVectorBinaryOp : public FRigVMFunction_MathVectorBase
 {
@@ -39,16 +50,22 @@ struct FRigVMFunction_MathVectorBinaryOp : public FRigVMFunction_MathVectorBase
 		A = B = Result = FVector::ZeroVector;
 	}
 
+	// The first vector for the operation
 	UPROPERTY(meta=(Input))
 	FVector A;
 
+	// The second vector for the operation
 	UPROPERTY(meta=(Input))
 	FVector B;
 
+	// The resulting vector
 	UPROPERTY(meta=(Output))
 	FVector Result;
 };
 
+/*
+ * The base class for all aggregational vector math operations
+ */
 USTRUCT(meta=(Abstract))
 struct FRigVMFunction_MathVectorBinaryAggregateOp : public FRigVMFunction_MathVectorBase
 {
@@ -59,12 +76,15 @@ struct FRigVMFunction_MathVectorBinaryAggregateOp : public FRigVMFunction_MathVe
 		A = B = Result = FVector::ZeroVector;
 	}
 
+	// The first vector for the operation
 	UPROPERTY(meta=(Input, Aggregate))
 	FVector A;
 
+	// The second vector for the operation
 	UPROPERTY(meta=(Input, Aggregate))
 	FVector B;
 
+	// The resulting vector
 	UPROPERTY(meta=(Output, Aggregate))
 	FVector Result;
 };
@@ -117,9 +137,11 @@ struct FRigVMFunction_MathVectorFromFloat : public FRigVMFunction_MathVectorBase
 		Result = FVector::ZeroVector;
 	}
 
+	// The input float value to use for all vector components 
 	UPROPERTY(meta=(Input))
 	float Value;
 
+	// The resulting (Value,Value,Value) vector
 	UPROPERTY(meta=(Output))
 	FVector Result;
 };
@@ -127,7 +149,7 @@ struct FRigVMFunction_MathVectorFromFloat : public FRigVMFunction_MathVectorBase
 /**
  * Makes a vector from a single double
  */
-USTRUCT(meta=(DisplayName="From Float", TemplateName="Cast", ExecuteContext="FRigVMExecuteContext", Keywords="Make,Construct"))
+USTRUCT(meta=(DisplayName="From Double", TemplateName="Cast", ExecuteContext="FRigVMExecuteContext", Keywords="Make,Construct"))
 struct FRigVMFunction_MathVectorFromDouble : public FRigVMFunction_MathVectorBase
 {
 	GENERATED_BODY()
@@ -140,9 +162,11 @@ struct FRigVMFunction_MathVectorFromDouble : public FRigVMFunction_MathVectorBas
 		Result = FVector::ZeroVector;
 	}
 
+	// The input value to use for all vector components
 	UPROPERTY(meta=(Input))
 	double Value;
 
+	// The resulting (Value,Value,Value) vector
 	UPROPERTY(meta=(Output))
 	FVector Result;
 };
@@ -163,7 +187,7 @@ struct FRigVMFunction_MathVectorAdd : public FRigVMFunction_MathVectorBinaryAggr
  * Returns the difference of the two values
  */
 USTRUCT(meta=(DisplayName="Subtract", TemplateName="Subtract", Keywords="-"))
-struct FRigVMFunction_MathVectorSub : public FRigVMFunction_MathVectorBinaryOp
+struct FRigVMFunction_MathVectorSub : public FRigVMFunction_MathVectorBinaryAggregateOp
 {
 	GENERATED_BODY()
 
@@ -205,12 +229,15 @@ struct FRigVMFunction_MathVectorScale : public FRigVMFunction_MathVectorBase
 	RIGVM_METHOD()
 	RIGVM_API virtual void Execute() override;
 
+	// The input vector to scale
 	UPROPERTY(meta = (Input))
 	FVector Value;
 
+	// The factor to scale the input vector by
 	UPROPERTY(meta = (Input))
 	float Factor;
 
+	// The resulting scaled vector
 	UPROPERTY(meta = (Output))
 	FVector Result;
 };
@@ -363,15 +390,19 @@ struct FRigVMFunction_MathVectorClamp : public FRigVMFunction_MathVectorBase
 	RIGVM_METHOD()
 	RIGVM_API virtual void Execute() override;
 
+	// The input vector to clamp
 	UPROPERTY(meta=(Input))
 	FVector Value;
 
+	// The Minimum vector to clamp the input by
 	UPROPERTY(meta=(Input))
 	FVector Minimum;
 
+	// The Maximum vector to clamp the input by
 	UPROPERTY(meta=(Input))
 	FVector Maximum;
 
+	// The resulting clamped vector
 	UPROPERTY(meta=(Output))
 	FVector Result;
 };
@@ -394,15 +425,19 @@ struct FRigVMFunction_MathVectorLerp : public FRigVMFunction_MathVectorBase
 	RIGVM_METHOD()
 	RIGVM_API virtual void Execute() override;
 
+	// The first vector to interpolate from
 	UPROPERTY(meta=(Input))
 	FVector A;
 
+	// The second vector to interpolate to
 	UPROPERTY(meta=(Input))
 	FVector B;
 
-	UPROPERTY(meta=(Input))
+	// The blend value for the interpolation
+	UPROPERTY(meta=(Input, UIMin = "0", UIMax = "1"))
 	float T;
 
+	// The resulting interpolated vector
 	UPROPERTY(meta=(Output))
 	FVector Result;
 };
@@ -425,18 +460,23 @@ struct FRigVMFunction_MathVectorRemap : public FRigVMFunction_MathVectorBase
 	RIGVM_METHOD()
 	RIGVM_API virtual void Execute() override;
 
+	// The input vector to remap
 	UPROPERTY(meta=(Input))
 	FVector Value;
 
+	// The minimum of the input vector for each component
 	UPROPERTY(meta=(Input))
 	FVector SourceMinimum;
 
+	// The maximum of the input vector for each component
 	UPROPERTY(meta=(Input))
 	FVector SourceMaximum;
 
+	// The minimum for the output vector for each component
 	UPROPERTY(meta=(Input))
 	FVector TargetMinimum;
 
+	// The maximum for the output vector for each component
 	UPROPERTY(meta=(Input))
 	FVector TargetMaximum;
 
@@ -444,6 +484,7 @@ struct FRigVMFunction_MathVectorRemap : public FRigVMFunction_MathVectorBase
 	UPROPERTY(meta=(Input))
 	bool bClamp;
 
+	// The resulting remapped vector
 	UPROPERTY(meta=(Output))
 	FVector Result;
 };
@@ -526,18 +567,21 @@ struct FRigVMFunction_MathVectorIsNearlyZero : public FRigVMFunction_MathVectorB
 	RIGVM_METHOD()
 	RIGVM_API virtual void Execute() override;
 
+	// The input vector to check
 	UPROPERTY(meta=(Input))
 	FVector Value;
 
+	// The tolerance to use when comparing with zero
 	UPROPERTY(meta=(Input))
 	float Tolerance;
 
+	// True if the input vector almost zero given the tolerance
 	UPROPERTY(meta=(Output))
 	bool Result;
 };
 
 /**
- * Returns true if the value A is almost equal to B
+ * Returns true if value A is almost equal to B
  */
 USTRUCT(meta=(DisplayName="Is Nearly Equal", TemplateName="IsNearlyEqual", Keywords="AlmostEqual,=="))
 struct FRigVMFunction_MathVectorIsNearlyEqual : public FRigVMFunction_MathVectorBase
@@ -554,15 +598,19 @@ struct FRigVMFunction_MathVectorIsNearlyEqual : public FRigVMFunction_MathVector
 	RIGVM_METHOD()
 	RIGVM_API virtual void Execute() override;
 
+	// The first vector to compare
 	UPROPERTY(meta=(Input))
 	FVector A;
 
+	// The second vector to compare
 	UPROPERTY(meta=(Input))
 	FVector B;
 
+	// The tolerance to use for the comparison
 	UPROPERTY(meta=(Input))
 	float Tolerance;
 
+	// True if value A is almost equal to B given the tolerance
 	UPROPERTY(meta=(Output))
 	bool Result;
 };
@@ -641,9 +689,11 @@ struct FRigVMFunction_MathVectorLengthSquared : public FRigVMFunction_MathVector
 	RIGVM_METHOD()
 	RIGVM_API virtual void Execute() override;
 
+	// The input vector to compute the squared length of
 	UPROPERTY(meta=(Input))
 	FVector Value;
 
+	// The resulting squared length of the vector
 	UPROPERTY(meta=(Output))
 	float Result;
 };
@@ -665,9 +715,11 @@ struct FRigVMFunction_MathVectorLength : public FRigVMFunction_MathVectorBase
 	RIGVM_METHOD()
 	RIGVM_API virtual void Execute() override;
 
+	// The input vector to compute the length for
 	UPROPERTY(meta=(Input))
 	FVector Value;
 
+	// The resulting length of the vector
 	UPROPERTY(meta=(Output))
 	float Result;
 };
@@ -689,12 +741,15 @@ struct FRigVMFunction_MathVectorDistance : public FRigVMFunction_MathVectorBase
 	RIGVM_METHOD()
 	RIGVM_API virtual void Execute() override;
 
+	// The first vector to measure the distance between
 	UPROPERTY(meta=(Input))
 	FVector A;
 
+	// The second vector to measure the distance between
 	UPROPERTY(meta=(Input))
 	FVector B;
 
+	// The distance between the two vectors A and B
 	UPROPERTY(meta=(Output))
 	float Result;
 };
@@ -727,12 +782,15 @@ struct FRigVMFunction_MathVectorDot : public FRigVMFunction_MathVectorBase
 	RIGVM_METHOD()
 	RIGVM_API virtual void Execute() override;
 
+	// The first vector compute the dot product for
 	UPROPERTY(meta=(Input))
 	FVector A;
 
+	// The second vector compute the dot product for
 	UPROPERTY(meta=(Input))
 	FVector B;
 
+	// The resulting dot product between A and B
 	UPROPERTY(meta=(Output))
 	float Result;
 };
@@ -766,12 +824,18 @@ struct FRigVMFunction_MathVectorSetLength: public FRigVMFunction_MathVectorBase
 	RIGVM_METHOD()
 	RIGVM_API virtual void Execute() override;
 
+	// The input vector to set the length of
 	UPROPERTY(meta = (Input))
 	FVector Value;
 
+	// The desired length of the vector
 	UPROPERTY(meta = (Input))
 	float Length;
 
+	/* 
+	 * The resulting scaled input vector
+	 * (or 0,0,0 in case the input was also 0,0,0)
+	 */
 	UPROPERTY(meta = (Output))
 	FVector Result;
 };
@@ -795,15 +859,22 @@ struct FRigVMFunction_MathVectorClampLength: public FRigVMFunction_MathVectorBas
 	RIGVM_METHOD()
 	RIGVM_API virtual void Execute() override;
 
+	// The input vector to clamp the length of 
 	UPROPERTY(meta = (Input))
 	FVector Value;
 
+	// The minimum length to clamp by
 	UPROPERTY(meta = (Input))
 	float MinimumLength;
 
+	// The maximum length to clamp by
 	UPROPERTY(meta = (Input))
 	float MaximumLength;
 
+	/* 
+	 * The resulting clamped input vector
+	 * (or 0,0,0 in case the input was also 0,0,0)
+	 */
 	UPROPERTY(meta = (Output))
 	FVector Result;
 };
@@ -825,12 +896,15 @@ struct FRigVMFunction_MathVectorMirror : public FRigVMFunction_MathVectorBase
 	RIGVM_METHOD()
 	RIGVM_API virtual void Execute() override;
 
+	// The input vector to mirror
 	UPROPERTY(meta=(Input))
 	FVector Value;
 
+	// The normal to mirror the vector by
 	UPROPERTY(meta=(Input))
 	FVector Normal;
 
+	// The resulting mirrored vector
 	UPROPERTY(meta=(Output))
 	FVector Result;
 };
@@ -852,12 +926,15 @@ struct FRigVMFunction_MathVectorAngle : public FRigVMFunction_MathVectorBase
 	RIGVM_METHOD()
 	RIGVM_API virtual void Execute() override;
 
+	// The first vector to measure the angle between
 	UPROPERTY(meta=(Input))
 	FVector A;
 
+	// The second vector to measure the angle between
 	UPROPERTY(meta=(Input))
 	FVector B;
 
+	// The angle between A and B in radians
 	UPROPERTY(meta=(Output))
 	float Result;
 };
@@ -879,12 +956,15 @@ struct FRigVMFunction_MathVectorParallel : public FRigVMFunction_MathVectorBase
 	RIGVM_METHOD()
 	RIGVM_API virtual void Execute() override;
 
+	// The first vector to compare
 	UPROPERTY(meta=(Input))
 	FVector A;
 
+	// The second vector to compare
 	UPROPERTY(meta=(Input))
 	FVector B;
 
+	// True if the two vectors A and B are parallel
 	UPROPERTY(meta=(Output))
 	bool Result;
 };
@@ -907,12 +987,15 @@ struct FRigVMFunction_MathVectorOrthogonal : public FRigVMFunction_MathVectorBas
 	RIGVM_METHOD()
 	RIGVM_API virtual void Execute() override;
 
+	// The first vector to compare
 	UPROPERTY(meta=(Input))
 	FVector A;
 
+	// The second vector to compare
 	UPROPERTY(meta=(Input))
 	FVector B;
 
+	// True if the two vectors A and B are orthogonal
 	UPROPERTY(meta=(Output))
 	bool Result;
 };
@@ -935,15 +1018,19 @@ struct FRigVMFunction_MathVectorBezierFourPoint : public FRigVMFunction_MathVect
 	RIGVM_METHOD()
 	RIGVM_API virtual void Execute() override;
 
+	// The input four point structure
 	UPROPERTY(meta=(Input))
 	FRigVMFourPointBezier Bezier;
 
+	// The interpolation T value along the bezier to compute
 	UPROPERTY(meta=(Input))
 	float T;
 
+	// The resulting position at value T on the bezier
 	UPROPERTY(meta=(Output))
 	FVector Result;
 
+	// The resulting tangent at value T on the bezier
 	UPROPERTY(meta=(Output))
 	FVector Tangent;
 
@@ -1000,12 +1087,15 @@ struct FRigVMFunction_MathVectorClampSpatially: public FRigVMFunction_MathVector
 	RIGVM_METHOD()
 	RIGVM_API virtual void Execute() override;
 
+	// The input position to clamp
 	UPROPERTY(meta = (Input))
 	FVector Value;
 
+	// The axis to use for the filter
 	UPROPERTY(meta = (Input))
 	TEnumAsByte<EAxis::Type> Axis;
 
+	// The type of filter / mode to apply
 	UPROPERTY(meta = (Input))
 	TEnumAsByte<ERigVMClampSpatialMode::Type> Type;
 
@@ -1027,15 +1117,19 @@ struct FRigVMFunction_MathVectorClampSpatially: public FRigVMFunction_MathVector
 	UPROPERTY(meta = (Input))
 	FTransform Space;
 
+	// Draws debug information if True
 	UPROPERTY(meta = (Input))
 	bool bDrawDebug;
 
+	// The color to use for the debug draw
 	UPROPERTY(meta = (Input))
 	FLinearColor DebugColor;
 
+	// The thickness to use for the debug draw
 	UPROPERTY(meta = (Input))
 	float DebugThickness;
 
+	// The resulting clamped position
 	UPROPERTY(meta = (Output))
 	FVector Result;
 };
@@ -1058,21 +1152,27 @@ struct FRigVMFunction_MathIntersectPlane : public FRigVMFunction_MathVectorBase
 	RIGVM_METHOD()
 	RIGVM_API virtual void Execute() override;
 
+	// The start of the ray to intersect with the plane
 	UPROPERTY(meta = (Input))
 	FVector Start;
 
+	// The direction of the ray to intersect with the plane
 	UPROPERTY(meta = (Input))
 	FVector Direction;
 
+	// The point on the plane to intersect with
 	UPROPERTY(meta = (Input))
 	FVector PlanePoint;
 
+	// The normal of the plane to intersect with
 	UPROPERTY(meta = (Input))
 	FVector PlaneNormal;
 
+	// The resulting position on the plane
 	UPROPERTY(meta = (Output))
 	FVector Result;
 
+	// The distance along the ray to the intersection 
 	UPROPERTY(meta = (Output))
 	float Distance;
 };
@@ -1095,18 +1195,23 @@ struct FRigVMFunction_MathDistanceToPlane : public FRigVMFunction_MathVectorBase
 	RIGVM_METHOD()
 	RIGVM_API virtual void Execute() override;
 
+	// The point to measure against the plane
 	UPROPERTY(meta = (Input))
 	FVector Point;
 
+	// A point on the plane to measure agains
 	UPROPERTY(meta = (Input))
 	FVector PlanePoint;
 
+	// The normal of the plane to measure agains
 	UPROPERTY(meta = (Input))
 	FVector PlaneNormal;
 
+	// The resulting closest point on the plane
 	UPROPERTY(meta = (Output))
 	FVector ClosestPointOnPlane;
 
+	// The distance to the plane as a signed (positive or negative) value
 	UPROPERTY(meta = (Output))
 	float SignedDistance;
 };
@@ -1126,13 +1231,16 @@ struct FRigVMFunction_MathVectorMakeRelative : public FRigVMFunction_MathVectorB
 
 	RIGVM_METHOD()
 	RIGVM_API virtual void Execute() override;
-
+	
+	// The global position to make relative
 	UPROPERTY(meta=(Input))
 	FVector Global;
 
+	// The parent global position to make the input relative to
 	UPROPERTY(meta=(Input))
 	FVector Parent;
 
+	// The resulting relative position
 	UPROPERTY(meta=(Output))
 	FVector Local;
 };
@@ -1153,12 +1261,15 @@ struct FRigVMFunction_MathVectorMakeAbsolute : public FRigVMFunction_MathVectorB
 	RIGVM_METHOD()
 	RIGVM_API virtual void Execute() override;
 
+	// The input local position to make absolute
 	UPROPERTY(meta = (Input))
 	FVector Local;
 
+	// The parent global position to make the local absolute to
 	UPROPERTY(meta = (Input))
 	FVector Parent;
 
+	// The resulting global position
 	UPROPERTY(meta = (Output))
 	FVector Global;
 };
@@ -1182,6 +1293,7 @@ struct FRigVMFunction_MathVectorMirrorTransform : public FRigVMFunction_MathVect
 	RIGVM_METHOD()
 	RIGVM_API virtual void Execute() override;
 
+	// The input vector to mirror
 	UPROPERTY(meta=(Input))
 	FVector Value;
 
@@ -1197,6 +1309,7 @@ struct FRigVMFunction_MathVectorMirrorTransform : public FRigVMFunction_MathVect
 	UPROPERTY(meta=(Input))
 	FTransform CentralTransform;
 
+	// The resulting mirrored vector
 	UPROPERTY(meta=(Output))
 	FVector Result;
 };
@@ -1218,9 +1331,11 @@ struct FRigVMFunction_MathVectorArraySum : public FRigVMFunction_MathVectorBase
 	RIGVM_METHOD()
 	RIGVM_API virtual void Execute() override;
 
+	// The array of vector to add up
 	UPROPERTY(meta = (Input))
 	TArray<FVector> Array;
 
+	// The sum of all vectors in the input array
 	UPROPERTY(meta = (Output))
 	FVector Sum;
 };
@@ -1243,9 +1358,11 @@ struct FRigVMFunction_MathVectorArrayAverage : public FRigVMFunction_MathVectorB
 	RIGVM_METHOD()
 	RIGVM_API virtual void Execute() override;
 
+	// The array of vectors to average
 	UPROPERTY(meta = (Input))
 	TArray<FVector> Array;
 
+	// The resulting average vector of the input array
 	UPROPERTY(meta = (Output))
 	FVector Average;
 };
\ No newline at end of file
diff --git a/Engine/Plugins/Runtime/RigVM/Source/RigVM/Public/RigVMFunctions/Math/RigVMFunction_Noise.h b/Engine/Plugins/Runtime/RigVM/Source/RigVM/Public/RigVMFunctions/Math/RigVMFunction_Noise.h
--- a/Engine/Plugins/Runtime/RigVM/Source/RigVM/Public/RigVMFunctions/Math/RigVMFunction_Noise.h
+++ b/Engine/Plugins/Runtime/RigVM/Source/RigVM/Public/RigVMFunctions/Math/RigVMFunction_Noise.h
@@ -23,21 +23,27 @@ struct FRigVMFunction_NoiseFloat : public FRigVMFunction_MathBase
 	RIGVM_METHOD()
 	RIGVM_API virtual void Execute() override;
 
+	// The input value to retrieve the noise for
 	UPROPERTY(meta = (Input))
 	float Value;
 
+	// The speed of the noise field
 	UPROPERTY(meta = (Input))
 	float Speed;
 
+	// The frequency of the noise field
 	UPROPERTY(meta = (Input))
 	float Frequency;
 
+	// The minimum range of the noise
 	UPROPERTY(meta = (Input))
 	float Minimum;
 
+	// The maximum range of the noise
 	UPROPERTY(meta = (Input))
 	float Maximum;
 
+	// The resulting value in the noise field
 	UPROPERTY(meta = (Output))
 	float Result;
 
@@ -63,21 +69,27 @@ struct FRigVMFunction_NoiseDouble : public FRigVMFunction_MathBase
 	RIGVM_METHOD()
 	RIGVM_API virtual void Execute() override;
 
+	// The input value to retrieve the noise for
 	UPROPERTY(meta = (Input))
 	double Value;
 
+	// The speed of the noise field
 	UPROPERTY(meta = (Input))
 	double Speed;
 
+	// The frequency of the noise field
 	UPROPERTY(meta = (Input))
 	double Frequency;
 
+	// The minimum range of the noise
 	UPROPERTY(meta = (Input))
 	double Minimum;
 
+	// The maximum range of the noise
 	UPROPERTY(meta = (Input))
 	double Maximum;
 
+	// The resulting value in the noise field
 	UPROPERTY(meta = (Output))
 	double Result;
 
@@ -150,21 +162,27 @@ struct FRigVMFunction_NoiseVector2 : public FRigVMFunction_MathBase
 	RIGVM_METHOD()
 	RIGVM_API virtual void Execute() override;
 
+	// The input position to retrieve the noise for
 	UPROPERTY(meta = (Input))
 	FVector Value;
 
+	// The speed of the noise field
 	UPROPERTY(meta = (Input))
 	FVector Speed;
 
+	// The frequency of the noise field
 	UPROPERTY(meta = (Input))
 	FVector Frequency;
 
+	// The minimum range of the noise
 	UPROPERTY(meta = (Input))
 	double Minimum;
 
+	// The maximum range of the noise
 	UPROPERTY(meta = (Input))
 	double Maximum;
 
+	// The resulting vector in the noise field
 	UPROPERTY(meta = (Output))
 	FVector Result;
 
diff --git a/Engine/Plugins/Runtime/RigVM/Source/RigVM/Public/RigVMFunctions/Math/RigVMFunction_Random.h b/Engine/Plugins/Runtime/RigVM/Source/RigVM/Public/RigVMFunctions/Math/RigVMFunction_Random.h
--- a/Engine/Plugins/Runtime/RigVM/Source/RigVM/Public/RigVMFunctions/Math/RigVMFunction_Random.h
+++ b/Engine/Plugins/Runtime/RigVM/Source/RigVM/Public/RigVMFunctions/Math/RigVMFunction_Random.h
@@ -25,12 +25,15 @@ struct FRigVMFunction_RandomFloat : public FRigVMFunction_MathBase
 	RIGVM_METHOD()
 	RIGVM_API virtual void Execute() override;
 
+	// The seed for the random number generator
 	UPROPERTY(meta = (Input))
 	int32 Seed;
 
+	// The minimum value for the random number range
 	UPROPERTY(meta = (Input))
 	float Minimum;
 
+	// The maximum value for the random number range
 	UPROPERTY(meta = (Input))
 	float Maximum;
 
@@ -42,6 +45,7 @@ struct FRigVMFunction_RandomFloat : public FRigVMFunction_MathBase
 	UPROPERTY(meta = (Input))
 	float Duration;
 
+	// The resulting random number
 	UPROPERTY(meta = (Output))
 	float Result;
 
@@ -79,15 +83,17 @@ struct FRigVMFunction_RandomVector: public FRigVMFunction_MathBase
 	RIGVM_METHOD()
 	RIGVM_API virtual void Execute() override;
 
+	// The seed for the random number generator
 	UPROPERTY(meta = (Input))
 	int32 Seed;
 
+	// The minimum value for the random number range for all components of the vector
 	UPROPERTY(meta = (Input))
 	float Minimum;
 
+	// The maximum value for the random number range for all components of the vector
 	UPROPERTY(meta = (Input))
 	float Maximum;
-
 	
 	/**
 	 * The duration at which the number won't change.
@@ -97,6 +103,7 @@ struct FRigVMFunction_RandomVector: public FRigVMFunction_MathBase
 	UPROPERTY(meta = (Input))
 	float Duration;
 
+	// The resulting random vector
 	UPROPERTY(meta = (Output))
 	FVector Result;
 
diff --git a/Engine/Plugins/Runtime/RigVM/Source/RigVM/Public/RigVMFunctions/Math/RigVMFunction_SpringMath.h b/Engine/Plugins/Runtime/RigVM/Source/RigVM/Public/RigVMFunctions/Math/RigVMFunction_SpringMath.h
--- a/Engine/Plugins/Runtime/RigVM/Source/RigVM/Public/RigVMFunctions/Math/RigVMFunction_SpringMath.h
+++ b/Engine/Plugins/Runtime/RigVM/Source/RigVM/Public/RigVMFunctions/Math/RigVMFunction_SpringMath.h
@@ -5,10 +5,10 @@
 #include "RigVMFunction_MathBase.h"
 #include "RigVMFunction_SpringMath.generated.h"
 
+UE_EXPERIMENTAL(5.7, "SpringMath API is experimental")
 /**
  * Damps a float value using exponential decay damping
  */
-UE_EXPERIMENTAL(5.7, "SpringMath API is experimental")
 USTRUCT(meta=(DisplayName="Damp (Float)", Category="Math|Damp|Experimental", TemplateName="Damp"))
 struct FRigVMFunction_DampFloat : public FRigVMFunction_MathBase
 {
@@ -17,23 +17,27 @@ struct FRigVMFunction_DampFloat : public FRigVMFunction_MathBase
 	RIGVM_METHOD()
 	RIGVM_API virtual void Execute() override;
 
+	// The value to damp
 	UPROPERTY(meta = (Input))
 	float Value = 0.0f;
 
+	// The target value to damp towards
 	UPROPERTY(meta = (Input))
 	float Target = 0.0f;
 
+	// The time to apply smoothing for
 	UPROPERTY(meta = (Input, ClampMin = "0", UIMin = "0", ForceUnits = "s"))
 	float SmoothingTime = 0.2f;
 
+	// Resulting damped value
 	UPROPERTY(meta = (Output))
 	float Result = 0.0f;
 };
 
+UE_EXPERIMENTAL(5.7, "SpringMath API is experimental")
 /**
  * Damps a vector value using exponential decay damping
  */
-UE_EXPERIMENTAL(5.7, "SpringMath API is experimental")
 USTRUCT(meta=(DisplayName="Damp (Vector)", Category="Math|Damp|Experimental", TemplateName="Damp"))
 struct FRigVMFunction_DampVector : public FRigVMFunction_MathBase
 {
@@ -42,23 +46,27 @@ struct FRigVMFunction_DampVector : public FRigVMFunction_MathBase
 	RIGVM_METHOD()
 	RIGVM_API virtual void Execute() override;
 
+	// The value to damp
 	UPROPERTY(meta = (Input))
 	FVector Value = FVector::ZeroVector;
 
+	// The target value to damp towards
 	UPROPERTY(meta = (Input))
 	FVector Target = FVector::ZeroVector;
 
+	// The time to apply smoothing for
 	UPROPERTY(meta = (Input, ClampMin = "0", UIMin = "0", ForceUnits = "s"))
 	float SmoothingTime = 0.2f;
 
+	// Resulting damped value
 	UPROPERTY(meta = (Output))
 	FVector Result = FVector::ZeroVector;
 };
 
+UE_EXPERIMENTAL(5.7, "SpringMath API is experimental") 
 /**
  * Damps a quaternion value using exponential decay damping
  */
-UE_EXPERIMENTAL(5.7, "SpringMath API is experimental") 
 USTRUCT(meta=(DisplayName="Damp (Quaternion)", Category="Math|Damp|Experimental", TemplateName="Damp"))
 struct FRigVMFunction_DampQuaternion : public FRigVMFunction_MathBase
 {
@@ -67,23 +75,27 @@ struct FRigVMFunction_DampQuaternion : public FRigVMFunction_MathBase
 	RIGVM_METHOD()
 	RIGVM_API virtual void Execute() override;
 
+	// The value to damp
 	UPROPERTY(meta = (Input))
 	FQuat Value = FQuat::Identity;
 
+	// The target value to damp towards
 	UPROPERTY(meta = (Input))
 	FQuat Target = FQuat::Identity;
 
+	// The time to apply smoothing for
 	UPROPERTY(meta = (Input, ClampMin = "0", UIMin = "0", ForceUnits = "s"))
 	float SmoothingTime = 0.2f;
 
+	// Resulting damped value
 	UPROPERTY(meta = (Output))
 	FQuat Result = FQuat::Identity;
 };
 
+UE_EXPERIMENTAL(5.7, "SpringMath API is experimental")
 /**
  * Damps a float using a spring damper
  */
-UE_EXPERIMENTAL(5.7, "SpringMath API is experimental")
 USTRUCT(meta=(DisplayName="Critical Spring Damp (Float)", Category="Math|Damp|Experimental", TemplateName="Critical Spring Damp"))
 struct FRigVMFunction_SpringDampFloat : public FRigVMFunction_MathMutableBase
 {
@@ -92,23 +104,27 @@ struct FRigVMFunction_SpringDampFloat : public FRigVMFunction_MathMutableBase
 	RIGVM_METHOD()
 	RIGVM_API virtual void Execute() override;
 
+	// The value to damp
 	UPROPERTY(meta = (Input, Output))
 	float Value = 0.0f;
 
+	// The velocity of the value
 	UPROPERTY(meta = (Input, Output))
 	float ValueVelocity = 0.0f;
 
+	// The target to damp towards
 	UPROPERTY(meta = (Input))
 	float Target = 0.0f;
 
+	// The time to apply smoothing for
 	UPROPERTY(meta = (Input, ClampMin = "0", UIMin = "0", ForceUnits = "s"))
 	float SmoothingTime = 0.2f;
 };
 
+UE_EXPERIMENTAL(5.7, "SpringMath API is experimental")
 /**
  * Damps a vector using a spring damper
  */
-UE_EXPERIMENTAL(5.7, "SpringMath API is experimental")
 USTRUCT(meta=(DisplayName="Critical Spring Damp (Vector)", Category="Math|Damp|Experimental", TemplateName="Critical Spring Damp"))
 struct FRigVMFunction_SpringDampVector : public FRigVMFunction_MathMutableBase
 {
@@ -117,23 +133,27 @@ struct FRigVMFunction_SpringDampVector : public FRigVMFunction_MathMutableBase
 	RIGVM_METHOD()
 	RIGVM_API virtual void Execute() override;
 
+	// The value to damp
 	UPROPERTY(meta = (Input, Output))
 	FVector Value = FVector::ZeroVector;
 
+	// The velocity of the value
 	UPROPERTY(meta = (Input, Output))
 	FVector ValueVelocity = FVector::ZeroVector;
 
+	// The target to damp towards
 	UPROPERTY(meta = (Input))
 	FVector Target = FVector::ZeroVector;
 
+	// The time to apply smoothing for
 	UPROPERTY(meta = (Input, ClampMin = "0", UIMin = "0", ForceUnits = "s"))
 	float SmoothingTime = 0.2f;
 };
 
+UE_EXPERIMENTAL(5.7, "SpringMath API is experimental")
 /**
  * Damps a quaternion using a spring damper
  */
-UE_EXPERIMENTAL(5.7, "SpringMath API is experimental")
 USTRUCT(meta=(DisplayName="Critical Spring Damp (Quat)", Category="Math|Damp|Experimental", TemplateName="Critical Spring Damp"))
 struct FRigVMFunction_SpringDampQuat : public FRigVMFunction_MathMutableBase
 {
@@ -142,15 +162,19 @@ struct FRigVMFunction_SpringDampQuat : public FRigVMFunction_MathMutableBase
 	RIGVM_METHOD()
 	RIGVM_API virtual void Execute() override;
 
+	// The value to damp
 	UPROPERTY(meta = (Input, Output))
 	FQuat Value = FQuat::Identity;
 
+	// The velocity of the value
 	UPROPERTY(meta = (Input, Output))
 	FVector ValueVelocity = FVector::ZeroVector;
 
+	// The target to damp towards
 	UPROPERTY(meta = (Input))
 	FQuat Target = FQuat::Identity;
 
+	// The time to apply smoothing for
 	UPROPERTY(meta = (Input, ClampMin = "0", UIMin = "0", ForceUnits = "s"))
 	float SmoothingTime = 0.2f;
 };
