RigVM / Control Rig: Comment all debug, name, string and animation nodes

#rnx

diff --git a/Engine/Plugins/Animation/ControlRig/Source/ControlRig/Public/Units/Debug/RigUnit_DebugHierarchy.h b/Engine/Plugins/Animation/ControlRig/Source/ControlRig/Public/Units/Debug/RigUnit_DebugHierarchy.h
--- a/Engine/Plugins/Animation/ControlRig/Source/ControlRig/Public/Units/Debug/RigUnit_DebugHierarchy.h
+++ b/Engine/Plugins/Animation/ControlRig/Source/ControlRig/Public/Units/Debug/RigUnit_DebugHierarchy.h
@@ -50,18 +50,23 @@ struct FRigUnit_DebugHierarchy : public FRigVMFunction_DebugBase
 	UPROPERTY(meta = (Input))
 	TArray<FRigElementKey> Items;
 
+	// The scale to apply to the drawn transform
 	UPROPERTY(meta = (Input))
 	float Scale;
 
+	// The color to use when drawing
 	UPROPERTY(meta = (Input))
 	FLinearColor Color;
 
+	// The line thickness to use when drawing
 	UPROPERTY(meta = (Input))
 	float Thickness;
 
+	// The world offset to pre-multiply the drawn transform with
 	UPROPERTY(meta = (Input))
 	FTransform WorldOffset;
 
+	// If set to False the debug drawing will be skipped
 	UPROPERTY(meta = (Input))
 	bool bEnabled;
 
@@ -94,6 +99,7 @@ struct FRigUnit_DebugPose : public FRigVMFunction_DebugBase
 	UPROPERTY(DisplayName = "Execute", Transient, meta = (Input, Output))
 	FRigVMExecutePin ExecutePin;
 
+	// The input pose to draw
 	UPROPERTY(meta = (Input))
 	FRigPose Pose;
 
@@ -102,18 +108,23 @@ struct FRigUnit_DebugPose : public FRigVMFunction_DebugBase
 	UPROPERTY(meta = (Input))
 	TArray<FRigElementKey> Items;
 
+	// The scale to apply to the drawn transform
 	UPROPERTY(meta = (Input))
 	float Scale;
 
+	// The color to use when drawing
 	UPROPERTY(meta = (Input))
 	FLinearColor Color;
 
+	// The line thickness to use when drawing
 	UPROPERTY(meta = (Input))
 	float Thickness;
 
+	// The world offset to pre-multiply the drawn transform with
 	UPROPERTY(meta = (Input))
 	FTransform WorldOffset;
 
+	// If set to False the debug drawing will be skipped
 	UPROPERTY(meta = (Input))
 	bool bEnabled;
 };
diff --git a/Engine/Plugins/Animation/ControlRig/Source/ControlRig/Public/Units/Debug/RigUnit_ProfilingBracket.h b/Engine/Plugins/Animation/ControlRig/Source/ControlRig/Public/Units/Debug/RigUnit_ProfilingBracket.h
--- a/Engine/Plugins/Animation/ControlRig/Source/ControlRig/Public/Units/Debug/RigUnit_ProfilingBracket.h
+++ b/Engine/Plugins/Animation/ControlRig/Source/ControlRig/Public/Units/Debug/RigUnit_ProfilingBracket.h
@@ -43,9 +43,11 @@ struct FRigUnit_EndProfilingTimer : public FRigVMFunction_DebugBaseMutable
 	RIGVM_METHOD()
 	UE_API virtual void Execute() override;
 
+	// The number of measurements to take for profiling
 	UPROPERTY(meta = (Input, Constant))
 	int32 NumberOfMeasurements;
 
+	// The prefix to use when printing to the log
 	UPROPERTY(meta = (Input, Constant))
 	FString Prefix;
 
diff --git a/Engine/Plugins/Runtime/RigVM/Source/RigVM/Public/RigVMFunctions/Animation/RigVMFunction_AnimBase.h b/Engine/Plugins/Runtime/RigVM/Source/RigVM/Public/RigVMFunctions/Animation/RigVMFunction_AnimBase.h
--- a/Engine/Plugins/Runtime/RigVM/Source/RigVM/Public/RigVMFunctions/Animation/RigVMFunction_AnimBase.h
+++ b/Engine/Plugins/Runtime/RigVM/Source/RigVM/Public/RigVMFunctions/Animation/RigVMFunction_AnimBase.h
@@ -5,7 +5,10 @@
 #include "RigVMCore/RigVMStruct.h"
 #include "RigVMFunction_AnimBase.generated.h"
 
-USTRUCT(meta=(Abstract, Category = "Animation", NodeColor = "0.05 0.05 0.25"))
+/*
+ * The base class for all animation related nodes
+ */
+USTRUCT(meta=(Abstract, Category = "Animation", NodeColor = "0.05 0.05 0.25", DocumentationPolicy = "Strict"))
 struct FRigVMFunction_AnimBase : public FRigVMStruct
 {
 	GENERATED_BODY()
diff --git a/Engine/Plugins/Runtime/RigVM/Source/RigVM/Public/RigVMFunctions/Animation/RigVMFunction_AnimEasing.h b/Engine/Plugins/Runtime/RigVM/Source/RigVM/Public/RigVMFunctions/Animation/RigVMFunction_AnimEasing.h
--- a/Engine/Plugins/Runtime/RigVM/Source/RigVM/Public/RigVMFunctions/Animation/RigVMFunction_AnimEasing.h
+++ b/Engine/Plugins/Runtime/RigVM/Source/RigVM/Public/RigVMFunctions/Animation/RigVMFunction_AnimEasing.h
@@ -22,6 +22,7 @@ struct FRigVMFunction_AnimEasingType : public FRigVMFunction_AnimBase
 	RIGVM_METHOD()
 	RIGVM_API virtual void Execute() override;
 
+	// The pass through constant easing type
 	UPROPERTY(meta = (Input, Output))
 	ERigVMAnimEasingType Type;
 };
@@ -45,24 +46,31 @@ struct FRigVMFunction_AnimEasing : public FRigVMFunction_AnimBase
 	RIGVM_METHOD()
 	RIGVM_API virtual void Execute() override;
 
+	// The value to ease using the easing functions
 	UPROPERTY(meta=(Input))
 	float Value;
 
+	// The type of easing to apply
 	UPROPERTY(meta=(Input))
 	ERigVMAnimEasingType Type;
 
+	// The minimum value of the input
 	UPROPERTY(meta=(Input))
 	float SourceMinimum;
 
+	// The maximum value of the input
 	UPROPERTY(meta=(Input))
 	float SourceMaximum;
 
+	// The minimum value of the output
 	UPROPERTY(meta=(Input))
 	float TargetMinimum;
 
+	// The maximum value of the output
 	UPROPERTY(meta=(Input))
 	float TargetMaximum;
 
+	// The resulting eased output value
 	UPROPERTY(meta=(Output))
 	float Result;
 };
diff --git a/Engine/Plugins/Runtime/RigVM/Source/RigVM/Public/RigVMFunctions/Animation/RigVMFunction_AnimEvalRichCurve.h b/Engine/Plugins/Runtime/RigVM/Source/RigVM/Public/RigVMFunctions/Animation/RigVMFunction_AnimEvalRichCurve.h
--- a/Engine/Plugins/Runtime/RigVM/Source/RigVM/Public/RigVMFunctions/Animation/RigVMFunction_AnimEvalRichCurve.h
+++ b/Engine/Plugins/Runtime/RigVM/Source/RigVM/Public/RigVMFunctions/Animation/RigVMFunction_AnimEvalRichCurve.h
@@ -29,24 +29,31 @@ struct FRigVMFunction_AnimEvalRichCurve : public FRigVMFunction_AnimBase
 	RIGVM_METHOD()
 	UE_API virtual void Execute() override;
 
+	// The input value to evaluate the curve at
 	UPROPERTY(meta=(Input))
 	float Value;
 
+	// The curve to evaluate
 	UPROPERTY(meta=(Input))
 	FRuntimeFloatCurve Curve;
 
+	// The minimum range of the input value (typically 0.0)
 	UPROPERTY(meta=(Input))
 	float SourceMinimum;
 
+	// The maximum range of the input value (typically 1.0)
 	UPROPERTY(meta=(Input))
 	float SourceMaximum;
 
+	// The minimum range of the output
 	UPROPERTY(meta=(Input))
 	float TargetMinimum;
 
+	// The maximum range of the output
 	UPROPERTY(meta=(Input))
 	float TargetMaximum;
 
+	// The evaluated value of the curve at the input value
 	UPROPERTY(meta=(Output))
 	float Result;
 };
diff --git a/Engine/Plugins/Runtime/RigVM/Source/RigVM/Public/RigVMFunctions/Animation/RigVMFunction_AnimRichCurve.h b/Engine/Plugins/Runtime/RigVM/Source/RigVM/Public/RigVMFunctions/Animation/RigVMFunction_AnimRichCurve.h
--- a/Engine/Plugins/Runtime/RigVM/Source/RigVM/Public/RigVMFunctions/Animation/RigVMFunction_AnimRichCurve.h
+++ b/Engine/Plugins/Runtime/RigVM/Source/RigVM/Public/RigVMFunctions/Animation/RigVMFunction_AnimRichCurve.h
@@ -24,6 +24,7 @@ struct FRigVMFunction_AnimRichCurve : public FRigVMFunction_AnimBase
 	RIGVM_METHOD()
 	RIGVM_API virtual void Execute() override;
 
+	// The pass through animation rich curve
 	UPROPERTY(meta=(Input, Output, Constant))
 	FRuntimeFloatCurve Curve;
 };
diff --git a/Engine/Plugins/Runtime/RigVM/Source/RigVM/Public/RigVMFunctions/Animation/RigVMFunction_GetDeltaTime.h b/Engine/Plugins/Runtime/RigVM/Source/RigVM/Public/RigVMFunctions/Animation/RigVMFunction_GetDeltaTime.h
--- a/Engine/Plugins/Runtime/RigVM/Source/RigVM/Public/RigVMFunctions/Animation/RigVMFunction_GetDeltaTime.h
+++ b/Engine/Plugins/Runtime/RigVM/Source/RigVM/Public/RigVMFunctions/Animation/RigVMFunction_GetDeltaTime.h
@@ -21,6 +21,7 @@ struct FRigVMFunction_GetDeltaTime : public FRigVMFunction_AnimBase
 	RIGVM_METHOD()
 	RIGVM_API virtual void Execute() override;
 
+	// The current delta time in seconds
 	UPROPERTY(meta=(Output))
 	float Result;
 };
diff --git a/Engine/Plugins/Runtime/RigVM/Source/RigVM/Public/RigVMFunctions/Animation/RigVMFunction_GetWorldTime.h b/Engine/Plugins/Runtime/RigVM/Source/RigVM/Public/RigVMFunctions/Animation/RigVMFunction_GetWorldTime.h
--- a/Engine/Plugins/Runtime/RigVM/Source/RigVM/Public/RigVMFunctions/Animation/RigVMFunction_GetWorldTime.h
+++ b/Engine/Plugins/Runtime/RigVM/Source/RigVM/Public/RigVMFunctions/Animation/RigVMFunction_GetWorldTime.h
@@ -21,27 +21,35 @@ struct FRigVMFunction_GetWorldTime : public FRigVMFunction_AnimBase
 	RIGVM_METHOD()
 	RIGVM_API virtual void Execute() override;
 
+	// The Year of the world time
 	UPROPERTY(meta = (Output))
 	float Year;
 
+	// The Month of the world time
 	UPROPERTY(meta = (Output))
 	float Month;
 
+	// The Day of the world time
 	UPROPERTY(meta = (Output))
 	float Day;
 
+	// The WeekDay of the world time
 	UPROPERTY(meta = (Output))
 	float WeekDay;
 
+	// The Hours of the world time
 	UPROPERTY(meta = (Output))
 	float Hours;
 
+	// The Minutes of the world time
 	UPROPERTY(meta = (Output))
 	float Minutes;
 
+	// The Seconds of the world time
 	UPROPERTY(meta = (Output))
 	float Seconds;
 
+	// The OverallSeconds of the world time
 	UPROPERTY(meta = (Output))
 	float OverallSeconds;
 };
diff --git a/Engine/Plugins/Runtime/RigVM/Source/RigVM/Public/RigVMFunctions/Animation/RigVMFunction_TimeConversion.h b/Engine/Plugins/Runtime/RigVM/Source/RigVM/Public/RigVMFunctions/Animation/RigVMFunction_TimeConversion.h
--- a/Engine/Plugins/Runtime/RigVM/Source/RigVM/Public/RigVMFunctions/Animation/RigVMFunction_TimeConversion.h
+++ b/Engine/Plugins/Runtime/RigVM/Source/RigVM/Public/RigVMFunctions/Animation/RigVMFunction_TimeConversion.h
@@ -21,9 +21,11 @@ struct FRigVMFunction_FramesToSeconds : public FRigVMFunction_AnimBase
 	RIGVM_METHOD()
 	RIGVM_API virtual void Execute() override;
 
+	// The input value representing a number of frames to convert to seconds
 	UPROPERTY(meta=(Input))
 	float Frames;
 
+	// The resulting time in seconds
 	UPROPERTY(meta=(Output))
 	float Seconds;
 };
@@ -44,9 +46,11 @@ struct  FRigVMFunction_SecondsToFrames : public FRigVMFunction_AnimBase
 	RIGVM_METHOD()
 	RIGVM_API virtual void Execute() override;
 
+	// The input time as seconds
 	UPROPERTY(meta=(Input))
 	float Seconds;
 
+	// The resulting number of frames 
 	UPROPERTY(meta=(Output))
 	float Frames;
 };
diff --git a/Engine/Plugins/Runtime/RigVM/Source/RigVM/Public/RigVMFunctions/Debug/RigVMFunction_DebugBase.h b/Engine/Plugins/Runtime/RigVM/Source/RigVM/Public/RigVMFunctions/Debug/RigVMFunction_DebugBase.h
--- a/Engine/Plugins/Runtime/RigVM/Source/RigVM/Public/RigVMFunctions/Debug/RigVMFunction_DebugBase.h
+++ b/Engine/Plugins/Runtime/RigVM/Source/RigVM/Public/RigVMFunctions/Debug/RigVMFunction_DebugBase.h
@@ -6,20 +6,28 @@
 #include "RigVMCore/RigVMDebugDrawSettings.h"
 #include "RigVMFunction_DebugBase.generated.h"
 
-USTRUCT(meta=(Abstract, Category="Debug", NodeColor = "0.83077 0.846873 0.049707"))
+/*
+ * The base class for all pure debug draw nodes
+ */
+USTRUCT(meta=(Abstract, Category="Debug", NodeColor = "0.83077 0.846873 0.049707", DocumentationPolicy = "Strict"))
 struct FRigVMFunction_DebugBase : public FRigVMStruct
 {
 	GENERATED_BODY()
 
+	// The debug draw settings for this node
 	UPROPERTY(meta = (Input, DetailsOnly, DisplayName = "Draw Settings"))
 	FRigVMDebugDrawSettings DebugDrawSettings;
 };
 
-USTRUCT(meta=(Abstract, Category="Debug", NodeColor = "0.83077 0.846873 0.049707"))
+/*
+ * The base class for all mutable debug draw nodes
+ */
+USTRUCT(meta=(Abstract, Category="Debug", NodeColor = "0.83077 0.846873 0.049707", DocumentationPolicy = "Strict"))
 struct FRigVMFunction_DebugBaseMutable : public FRigVMStructMutable
 {
 	GENERATED_BODY()
 
+	// The debug draw settings for this node
 	UPROPERTY(meta = (Input, DetailsOnly, DisplayName = "Draw Settings"))
 	FRigVMDebugDrawSettings DebugDrawSettings;
 };
diff --git a/Engine/Plugins/Runtime/RigVM/Source/RigVM/Public/RigVMFunctions/Debug/RigVMFunction_DebugLine.h b/Engine/Plugins/Runtime/RigVM/Source/RigVM/Public/RigVMFunctions/Debug/RigVMFunction_DebugLine.h
--- a/Engine/Plugins/Runtime/RigVM/Source/RigVM/Public/RigVMFunctions/Debug/RigVMFunction_DebugLine.h
+++ b/Engine/Plugins/Runtime/RigVM/Source/RigVM/Public/RigVMFunctions/Debug/RigVMFunction_DebugLine.h
@@ -27,21 +27,27 @@ struct FRigVMFunction_DebugLineNoSpace : public FRigVMFunction_DebugBaseMutable
 	RIGVM_METHOD()
 	UE_API virtual void Execute() override;
 
+	// The global start position of the line
 	UPROPERTY(meta = (Input))
 	FVector A;
 
+	// The global end position of the line
 	UPROPERTY(meta = (Input))
 	FVector B;
 
+	// The color to use for the drawing
 	UPROPERTY(meta = (Input))
 	FLinearColor Color;
 
+	// The line thickness to use for the drawing
 	UPROPERTY(meta = (Input))
 	float Thickness;
 
+	// The world offset to pre-multiply the positions with
 	UPROPERTY(meta = (Input))
 	FTransform WorldOffset;
 
+	// If set to False the debug drawing will be skipped
 	UPROPERTY(meta = (Input))
 	bool bEnabled;
 };
diff --git a/Engine/Plugins/Runtime/RigVM/Source/RigVM/Public/RigVMFunctions/Debug/RigVMFunction_DebugLineStrip.h b/Engine/Plugins/Runtime/RigVM/Source/RigVM/Public/RigVMFunctions/Debug/RigVMFunction_DebugLineStrip.h
--- a/Engine/Plugins/Runtime/RigVM/Source/RigVM/Public/RigVMFunctions/Debug/RigVMFunction_DebugLineStrip.h
+++ b/Engine/Plugins/Runtime/RigVM/Source/RigVM/Public/RigVMFunctions/Debug/RigVMFunction_DebugLineStrip.h
@@ -24,18 +24,23 @@ struct FRigVMFunction_DebugLineStripNoSpace : public FRigVMFunction_DebugBaseMut
 	RIGVM_METHOD()
 	RIGVM_API virtual void Execute() override;
 
+	// The array of input positions to treat as a line-strip
 	UPROPERTY(meta = (Input))
 	TArray<FVector> Points;
 
+	// The color to use for the drawing
 	UPROPERTY(meta = (Input))
 	FLinearColor Color;
 
+	// The line thickness to use for the drawing
 	UPROPERTY(meta = (Input))
 	float Thickness;
 
+	// The world offset to pre-multiply the positions with
 	UPROPERTY(meta = (Input))
 	FTransform WorldOffset;
 	
+	// If set to False the debug drawing will be skipped
 	UPROPERTY(meta = (Input))
 	bool bEnabled;
 };
diff --git a/Engine/Plugins/Runtime/RigVM/Source/RigVM/Public/RigVMFunctions/Debug/RigVMFunction_DebugPrimitives.h b/Engine/Plugins/Runtime/RigVM/Source/RigVM/Public/RigVMFunctions/Debug/RigVMFunction_DebugPrimitives.h
--- a/Engine/Plugins/Runtime/RigVM/Source/RigVM/Public/RigVMFunctions/Debug/RigVMFunction_DebugPrimitives.h
+++ b/Engine/Plugins/Runtime/RigVM/Source/RigVM/Public/RigVMFunctions/Debug/RigVMFunction_DebugPrimitives.h
@@ -25,24 +25,30 @@ struct FRigVMFunction_DebugRectangle : public FRigVMFunction_DebugBaseMutable
 	RIGVM_METHOD()
 	UE_API virtual void Execute() override;
 
+	// The transform of the rectangle to draw
 	UPROPERTY(meta = (Input))
 	FTransform Transform;
 
+	// The color to use for the debug draw
 	UPROPERTY(meta = (Input))
 	FLinearColor Color;
 
+	// The scale to apply to the debug draw
 	UPROPERTY(meta = (Input))
 	float Scale;
 
+	// The line thickness to use for the debug draw
 	UPROPERTY(meta = (Input))
 	float Thickness;
 
 	UPROPERTY(meta = (Input))
 	FName Space;
 
+	// The world offset to pre-multiply the positions with
 	UPROPERTY(meta = (Input))
 	FTransform WorldOffset;
 
+	// If set to False the debug drawing will be skipped
 	UPROPERTY(meta = (Input, Constant))
 	bool bEnabled;
 
@@ -71,21 +77,27 @@ struct FRigVMFunction_DebugRectangleNoSpace : public FRigVMFunction_DebugBaseMut
 	RIGVM_METHOD()
 	UE_API virtual void Execute() override;
 
+	// The transform of the rectangle to draw
 	UPROPERTY(meta = (Input))
 	FTransform Transform;
 
+	// The color to use for the debug draw
 	UPROPERTY(meta = (Input))
 	FLinearColor Color;
 
+	// The scale to apply to the debug draw
 	UPROPERTY(meta = (Input))
 	float Scale;
 
+	// The line thickness to use for the debug draw
 	UPROPERTY(meta = (Input))
 	float Thickness;
 
+	// The world offset to pre-multiply the positions with
 	UPROPERTY(meta = (Input))
 	FTransform WorldOffset;
 
+	// If set to False the debug drawing will be skipped
 	UPROPERTY(meta = (Input))
 	bool bEnabled;
 };
@@ -114,6 +126,7 @@ struct FRigVMFunction_DebugArc : public FRigVMFunction_DebugBaseMutable
 	UPROPERTY(meta = (Input))
 	FTransform Transform;
 
+	// The color to use for the debug draw
 	UPROPERTY(meta = (Input))
 	FLinearColor Color;
 
@@ -126,6 +139,7 @@ struct FRigVMFunction_DebugArc : public FRigVMFunction_DebugBaseMutable
 	UPROPERTY(meta = (Input))
 	float MaximumDegrees;
 
+	// The line thickness to use for the debug draw
 	UPROPERTY(meta = (Input))
 	float Thickness;
 
@@ -135,9 +149,11 @@ struct FRigVMFunction_DebugArc : public FRigVMFunction_DebugBaseMutable
 	UPROPERTY(meta = (Input))
 	FName Space;
 
+	// The world offset to pre-multiply the positions with
 	UPROPERTY(meta = (Input))
 	FTransform WorldOffset;
 
+	// If set to False the debug drawing will be skipped
 	UPROPERTY(meta = (Input, Constant))
 	bool bEnabled;
 
@@ -168,30 +184,39 @@ struct FRigVMFunction_DebugArcNoSpace : public FRigVMFunction_DebugBaseMutable
 	RIGVM_METHOD()
 	UE_API virtual void Execute() override;
 
+	// The transform of the arc to draw
 	UPROPERTY(meta = (Input))
 	FTransform Transform;
 
+	// The color to use for the debug draw
 	UPROPERTY(meta = (Input))
 	FLinearColor Color;
 
+	// The radius of the arc to draw
 	UPROPERTY(meta = (Input))
 	float Radius;
 
+	// The minimum angle in degress of the arc
 	UPROPERTY(meta = (Input))
 	float MinimumDegrees;
 
+	// The maximum angle in degrees of the arc
 	UPROPERTY(meta = (Input))
 	float MaximumDegrees;
 
+	// The line thickness to use for the debug draw
 	UPROPERTY(meta = (Input))
 	float Thickness;
 
+	// The detail to use when drawing the arc
 	UPROPERTY(meta = (Input))
 	int32 Detail;
 
+	// The world offset to pre-multiply the positions with
 	UPROPERTY(meta = (Input))
 	FTransform WorldOffset;
 
+	// If set to False the debug drawing will be skipped
 	UPROPERTY(meta = (Input))
 	bool bEnabled;
 };
@@ -217,18 +242,23 @@ struct FRigVMFunction_DebugBoxNoSpace : public FRigVMFunction_DebugBaseMutable
 	RIGVM_METHOD()
 	UE_API virtual void Execute() override;
 
+	// The input bounding box to draw
 	UPROPERTY(meta = (Input))
 	FBox Box;
 
+	// The color to use for the debug draw
 	UPROPERTY(meta = (Input))
 	FLinearColor Color;
 
+	// The line thickness to use for the debug draw
 	UPROPERTY(meta = (Input))
 	float Thickness;
 
+	// The world offset to pre-multiply the positions with
 	UPROPERTY(meta = (Input))
 	FTransform WorldOffset;
 
+	// If set to False the debug drawing will be skipped
 	UPROPERTY(meta = (Input))
 	bool bEnabled;
 };
diff --git a/Engine/Plugins/Runtime/RigVM/Source/RigVM/Public/RigVMFunctions/Debug/RigVMFunction_DebugTransform.h b/Engine/Plugins/Runtime/RigVM/Source/RigVM/Public/RigVMFunctions/Debug/RigVMFunction_DebugTransform.h
--- a/Engine/Plugins/Runtime/RigVM/Source/RigVM/Public/RigVMFunctions/Debug/RigVMFunction_DebugTransform.h
+++ b/Engine/Plugins/Runtime/RigVM/Source/RigVM/Public/RigVMFunctions/Debug/RigVMFunction_DebugTransform.h
@@ -42,24 +42,31 @@ struct FRigVMFunction_DebugTransformMutableNoSpace : public FRigVMFunction_Debug
 	RIGVM_METHOD()
 	RIGVM_API virtual void Execute() override;
 
+	// The transform to draw
 	UPROPERTY(meta = (Input))
 	FTransform Transform;
 
+	// The mode to use when drawing the transform
 	UPROPERTY(meta = (Input))
 	ERigUnitDebugTransformMode Mode;
 
+	// The debug color to use 
 	UPROPERTY(meta = (Input))
 	FLinearColor Color;
 
+	// The line thickness to use
 	UPROPERTY(meta = (Input))
 	float Thickness;
 
+	// The scale to scale the transform by
 	UPROPERTY(meta = (Input))
 	float Scale;
 
+	// The world offset to pre-multiply the transform with
 	UPROPERTY(meta = (Input))
 	FTransform WorldOffset;
 
+	// If False the debug drawing will be skipped
 	UPROPERTY(meta = (Input))
 	bool bEnabled;
 };
@@ -94,27 +101,37 @@ struct FRigVMFunction_DebugTransformArrayMutableNoSpace : public FRigVMFunction_
 	RIGVM_METHOD()
 	RIGVM_API virtual void Execute() override;
 
+	// An array of input transforms to draw
 	UPROPERTY(meta = (Input))
 	TArray<FTransform> Transforms;
 
+	// An array of parent indices for each transform.
+	// If this is specified lines will be drawn from each child
+	// to parent to represent a hierarchy.
 	UPROPERTY(meta = (Input))
 	TArray<int32> ParentIndices;
 
+	// The mode to use when drawing the transforms
 	UPROPERTY(meta = (Input))
 	ERigUnitDebugTransformMode Mode;
 
+	// The debug color to use
 	UPROPERTY(meta = (Input))
 	FLinearColor Color;
 
+	// The line thickness to use
 	UPROPERTY(meta = (Input))
 	float Thickness;
 
+	// The scale to scale the transforms by
 	UPROPERTY(meta = (Input))
 	float Scale;
 
+	// The world offset to pre-multiply the transforms with
 	UPROPERTY(meta = (Input))
 	FTransform WorldOffset;
 
+	// If False the debug drawing will be skipped
 	UPROPERTY(meta = (Input))
 	bool bEnabled;
 };
diff --git a/Engine/Plugins/Runtime/RigVM/Source/RigVM/Public/RigVMFunctions/Debug/RigVMFunction_VisualDebug.h b/Engine/Plugins/Runtime/RigVM/Source/RigVM/Public/RigVMFunctions/Debug/RigVMFunction_VisualDebug.h
--- a/Engine/Plugins/Runtime/RigVM/Source/RigVM/Public/RigVMFunctions/Debug/RigVMFunction_VisualDebug.h
+++ b/Engine/Plugins/Runtime/RigVM/Source/RigVM/Public/RigVMFunctions/Debug/RigVMFunction_VisualDebug.h
@@ -83,21 +83,27 @@ struct FRigVMFunction_VisualDebugVectorNoSpace : public FRigVMFunction_DebugBase
 	RIGVM_METHOD()
 	RIGVM_API virtual void Execute() override;
 
+	// The pass through vector to draw
 	UPROPERTY(meta = (Input, Output))
 	FVector Value;
 
+	// If set to False the debug drawing will be skipped
 	UPROPERTY(meta = (Input))
 	bool bEnabled;
 
+	// The mode to draw the vector with
 	UPROPERTY(meta = (Input, EditCondition = "bEnabled"))
 	ERigUnitVisualDebugPointMode Mode;
 
+	// The color to use for the debug draw
 	UPROPERTY(meta = (Input, EditCondition = "bEnabled"))
 	FLinearColor Color;
 
+	// The line thickness to use for the drawing
 	UPROPERTY(meta = (Input, EditCondition = "bEnabled"))
 	float Thickness;
 
+	// The scale to apply to the vector when drawing
 	UPROPERTY(meta = (Input, EditCondition = "bEnabled"))
 	float Scale;
 };
@@ -157,15 +163,19 @@ struct FRigVMFunction_VisualDebugQuatNoSpace : public FRigVMFunction_DebugBase
 	RIGVM_METHOD()
 	RIGVM_API virtual void Execute() override;
 
+	// The pass through rotation to draw
 	UPROPERTY(meta = (Input, Output))
 	FQuat Value;
 
+	// If set to False the debug drawing will be skipped
 	UPROPERTY(meta = (Input))
 	bool bEnabled;
 
+	// The line thickness to use for the drawing
 	UPROPERTY(meta = (Input, EditCondition = "bEnabled"))
 	float Thickness;
 
+	// The scale to apply to the axes drawn
 	UPROPERTY(meta = (Input, EditCondition = "bEnabled"))
 	float Scale;
 };
@@ -225,15 +235,19 @@ struct FRigVMFunction_VisualDebugTransformNoSpace : public FRigVMFunction_DebugB
 	RIGVM_METHOD()
 	RIGVM_API virtual void Execute() override;
 
+	// The pass through transform to draw
 	UPROPERTY(meta = (Input, Output))
 	FTransform Value;
 
+	// If set to False the debug drawing will be skipped
 	UPROPERTY(meta = (Input))
 	bool bEnabled;
 
+	// The line thickness to use for the drawing
 	UPROPERTY(meta = (Input, EditCondition = "bEnabled"))
 	float Thickness;
 
+	// The scale to apply to the axes drawn
 	UPROPERTY(meta = (Input, EditCondition = "bEnabled"))
 	float Scale;
 };
\ No newline at end of file
diff --git a/Engine/Plugins/Runtime/RigVM/Source/RigVM/Public/RigVMFunctions/RigVMFunction_ControlFlow.h b/Engine/Plugins/Runtime/RigVM/Source/RigVM/Public/RigVMFunctions/RigVMFunction_ControlFlow.h
--- a/Engine/Plugins/Runtime/RigVM/Source/RigVM/Public/RigVMFunctions/RigVMFunction_ControlFlow.h
+++ b/Engine/Plugins/Runtime/RigVM/Source/RigVM/Public/RigVMFunctions/RigVMFunction_ControlFlow.h
@@ -5,7 +5,10 @@
 #include "RigVMCore/RigVMStruct.h"
 #include "RigVMFunction_ControlFlow.generated.h"
 
-USTRUCT(meta=(Abstract, Category="Execution", NodeColor = "0, 0, 0, 1"))
+/*
+ * The base class for all control flow nodes
+ */
+USTRUCT(meta=(Abstract, Category="Execution", NodeColor = "0, 0, 0, 1", DocumentationPolicy = "Strict"))
 struct FRigVMFunction_ControlFlowBase : public FRigVMStruct
 {
 	GENERATED_BODY()
@@ -41,18 +44,23 @@ struct FRigVMFunction_ControlFlowBranch : public FRigVMFunction_ControlFlowBase
 	UPROPERTY(meta=(Input, DisplayName="Execute"))
 	FRigVMExecuteContext ExecuteContext;
 
+	// The condition based on which to pick between the True and False branches
 	UPROPERTY(meta=(Input))
 	bool Condition;
 
+	// The branch to run if the condition is True
 	UPROPERTY(meta=(Output))
 	FRigVMExecuteContext True;
 
+	// The branch to run if the condition is False
 	UPROPERTY(meta=(Output))
 	FRigVMExecuteContext False;
 
+	// The execute flow to run when the True or False branch is complete
 	UPROPERTY(meta=(Output))
 	FRigVMExecuteContext Completed;
 
+	// The internal block to run as the node progresses
 	UPROPERTY(meta=(Singleton))
 	FName BlockToRun;
 };
diff --git a/Engine/Plugins/Runtime/RigVM/Source/RigVM/Public/RigVMFunctions/RigVMFunction_Name.h b/Engine/Plugins/Runtime/RigVM/Source/RigVM/Public/RigVMFunctions/RigVMFunction_Name.h
--- a/Engine/Plugins/Runtime/RigVM/Source/RigVM/Public/RigVMFunctions/RigVMFunction_Name.h
+++ b/Engine/Plugins/Runtime/RigVM/Source/RigVM/Public/RigVMFunctions/RigVMFunction_Name.h
@@ -5,7 +5,10 @@
 #include "RigVMCore/RigVMStruct.h"
 #include "RigVMFunction_Name.generated.h"
 
-USTRUCT(meta=(Abstract, Category="Core|Name", NodeColor = "0.462745, 1,0, 0.329412"))
+/*
+ * The base class for all name operation nodes
+ */
+USTRUCT(meta=(Abstract, Category="Core|Name", NodeColor = "0.462745, 1,0, 0.329412", DocumentationPolicy = "Strict"))
 struct FRigVMFunction_NameBase : public FRigVMStruct
 {
 	GENERATED_BODY()
@@ -14,7 +17,7 @@ struct FRigVMFunction_NameBase : public FRigVMStruct
 };
 
 /**
- * Concatenates two strings together to make a new string
+ * Concatenates two names together to make a new name
  */
 USTRUCT(meta = (DisplayName = "Concat", TemplateName = "Concat", Keywords = "Add,+,Combine,Merge,Append"))
 struct FRigVMFunction_NameConcat : public FRigVMFunction_NameBase
@@ -29,18 +32,21 @@ struct FRigVMFunction_NameConcat : public FRigVMFunction_NameBase
 	RIGVM_METHOD()
 	RIGVM_API virtual void Execute() override;
 
+	// The first / left name to concatenate
 	UPROPERTY(meta=(Input, Aggregate))
 	FName A;
 
+	// The second / right name to concatenate
 	UPROPERTY(meta=(Input, Aggregate))
 	FName B;
 
+	// The resulting concatenated name
 	UPROPERTY(meta=(Output, Aggregate))
 	FName Result;
 };
 
 /**
- * Returns the left or right most characters from the string chopping the given number of characters from the start or the end
+ * Returns the left or right most characters from the name chopping the given number of characters from the start or the end
  */
 USTRUCT(meta = (DisplayName = "Chop", TemplateName = "Chop", Keywords = "Truncate,-,Remove,Subtract,Split"))
 struct FRigVMFunction_NameTruncate : public FRigVMFunction_NameBase
@@ -57,6 +63,7 @@ struct FRigVMFunction_NameTruncate : public FRigVMFunction_NameBase
 	RIGVM_METHOD()
 	RIGVM_API virtual void Execute() override;
 
+	// The input name to chop
 	UPROPERTY(meta=(Input))
 	FName Name;
 
@@ -68,7 +75,7 @@ struct FRigVMFunction_NameTruncate : public FRigVMFunction_NameBase
 	UPROPERTY(meta=(Input))
 	bool FromEnd;
 
-	// the part of the string without the chopped characters
+	// the part of the name without the chopped characters
 	UPROPERTY(meta=(Output))
 	FName Remainder;
 
@@ -78,7 +85,7 @@ struct FRigVMFunction_NameTruncate : public FRigVMFunction_NameBase
 };
 
 /**
- * Replace all occurrences of a substring in this string
+ * Replace all occurrences of a subname in this name
  */
 USTRUCT(meta = (DisplayName = "Replace", TemplateName = "Replace", Keywords = "Search,Emplace,Find"))
 struct FRigVMFunction_NameReplace : public FRigVMFunction_NameBase
@@ -93,21 +100,25 @@ struct FRigVMFunction_NameReplace : public FRigVMFunction_NameBase
 	RIGVM_METHOD()
 	RIGVM_API virtual void Execute() override;
 
+	// The input name to search within
 	UPROPERTY(meta = (Input))
 	FName Name;
 
+	// The old name to search for
 	UPROPERTY(meta = (Input))
 	FName Old;
 
+	// The new name to replace the old name with
 	UPROPERTY(meta = (Input))
 	FName New;
 
+	// The resulting replaced name
 	UPROPERTY(meta = (Output))
 	FName Result;
 };
 
 /**
- * Tests whether this string ends with given string
+ * Tests whether this name ends with given name
  */
 USTRUCT(meta = (DisplayName = "Ends With", TemplateName = "EndsWith", Keywords = "Right"))
 struct FRigVMFunction_EndsWith : public FRigVMFunction_NameBase
@@ -123,18 +134,21 @@ struct FRigVMFunction_EndsWith : public FRigVMFunction_NameBase
 	RIGVM_METHOD()
 		RIGVM_API virtual void Execute() override;
 
+	// The input name to search within
 	UPROPERTY(meta = (Input))
 	FName Name;
 
+	// The ending name to look for
 	UPROPERTY(meta = (Input))
 	FName Ending;
 
+	// True if the given input name ends with the given ending
 	UPROPERTY(meta = (Output))
 	bool Result;
 };
 
 /**
- * Tests whether this string starts with given string
+ * Tests whether this name starts with given name
  */
 USTRUCT(meta = (DisplayName = "Starts With", TemplateName = "StartsWith", Keywords = "Left"))
 struct FRigVMFunction_StartsWith : public FRigVMFunction_NameBase
@@ -150,12 +164,15 @@ struct FRigVMFunction_StartsWith : public FRigVMFunction_NameBase
 	RIGVM_METHOD()
 	RIGVM_API virtual void Execute() override;
 
+	// The input name to search within
 	UPROPERTY(meta = (Input))
 	FName Name;
 
+	// The start name to look for
 	UPROPERTY(meta = (Input))
 	FName Start;
 
+	// True if the given input name starts with the given start string
 	UPROPERTY(meta = (Output))
 	bool Result;
 };
@@ -177,12 +194,15 @@ struct FRigVMFunction_Contains : public FRigVMFunction_NameBase
 	RIGVM_METHOD()
 	RIGVM_API virtual void Execute() override;
 
+	// The input name to search within
 	UPROPERTY(meta = (Input))
 	FName Name;
 
+	// The name to search for
 	UPROPERTY(meta = (Input))
 	FName Search;
 
+	// True if a given name exists in another given name
 	UPROPERTY(meta = (Output))
 	bool Result;
 };
@@ -204,15 +224,17 @@ struct FRigVMFunction_IsNone : public FRigVMFunction_NameBase
 	RIGVM_METHOD()
 		RIGVM_API virtual void Execute() override;
 
+	// The input name to check
 	UPROPERTY(meta = (Input))
 	FName Name;
 
+	// True if the input name is None
 	UPROPERTY(meta = (Output))
 	bool Result;
 };
 
 /**
- * Returns true if this name is not valid / is none
+ * Returns true if this name is valid / is not none
  */
 USTRUCT(meta = (DisplayName = "Is Valid", TemplateName = "IsValid", Keywords = "Empty,Valid"))
 struct FRigVMFunction_IsNameValid : public FRigVMFunction_NameBase
@@ -228,9 +250,11 @@ struct FRigVMFunction_IsNameValid : public FRigVMFunction_NameBase
 	RIGVM_METHOD()
 	RIGVM_API virtual void Execute() override;
 
+	// The input name to check
 	UPROPERTY(meta = (Input))
 	FName Value;
 
+	// True if the name is valid/ is not None
 	UPROPERTY(meta = (Output))
 	bool Result;
 };
@@ -253,12 +277,15 @@ struct FRigVMFunction_GetNameNumericSuffix : public FRigVMFunction_NameBase
 	RIGVM_METHOD()
 	RIGVM_API virtual void Execute() override;
 
+	// The input name to check for a numeric suffix
 	UPROPERTY(meta = (Input))
 	FName Name;
 
+	// The resulting suffix (if it exists)
 	UPROPERTY(meta = (Output))
 	int32 Suffix;
 
+	// True if the name ends with the given numeric suffix
 	UPROPERTY(meta = (Output))
 	bool Success;
 };
\ No newline at end of file
diff --git a/Engine/Plugins/Runtime/RigVM/Source/RigVM/Public/RigVMFunctions/RigVMFunction_String.h b/Engine/Plugins/Runtime/RigVM/Source/RigVM/Public/RigVMFunctions/RigVMFunction_String.h
--- a/Engine/Plugins/Runtime/RigVM/Source/RigVM/Public/RigVMFunctions/RigVMFunction_String.h
+++ b/Engine/Plugins/Runtime/RigVM/Source/RigVM/Public/RigVMFunctions/RigVMFunction_String.h
@@ -6,7 +6,10 @@
 #include "RigVMCore/RigVMDispatchFactory.h"
 #include "RigVMFunction_String.generated.h"
 
-USTRUCT(meta=(Abstract, Category="Core|String", NodeColor = "0.462745, 1,0, 0.329412"))
+/*
+ * The base class for all string operation nodes
+ */
+USTRUCT(meta=(Abstract, Category="Core|String", NodeColor = "0.462745, 1,0, 0.329412", DocumentationPolicy = "Strict"))
 struct FRigVMFunction_StringBase : public FRigVMStruct
 {
 	GENERATED_BODY()
@@ -28,12 +31,15 @@ struct FRigVMFunction_StringConcat : public FRigVMFunction_StringBase
 	RIGVM_METHOD()
 	RIGVM_API virtual void Execute();
 
+	// The first / left string to concatenate
 	UPROPERTY(meta=(Input, Aggregate))
 	FString A;
 
+	// The second / right string to concatenate
 	UPROPERTY(meta=(Input, Aggregate))
 	FString B;
 
+	// The resulting concatenated string
 	UPROPERTY(meta=(Output, Aggregate))
 	FString Result;
 };
@@ -56,6 +62,7 @@ struct FRigVMFunction_StringTruncate : public FRigVMFunction_StringBase
 	RIGVM_METHOD()
 	RIGVM_API virtual void Execute();
 
+	// The input string to chop
 	UPROPERTY(meta=(Input))
 	FString Name;
 
@@ -92,15 +99,19 @@ struct FRigVMFunction_StringReplace : public FRigVMFunction_StringBase
 	RIGVM_METHOD()
 	RIGVM_API virtual void Execute();
 
+	// The input string to search within
 	UPROPERTY(meta = (Input))
 	FString Name;
 
+	// The old string to search for
 	UPROPERTY(meta = (Input))
 	FString Old;
 
+	// The new string to replace the old string with
 	UPROPERTY(meta = (Input))
 	FString New;
 
+	// The resulting replaced string
 	UPROPERTY(meta = (Output))
 	FString Result;
 };
@@ -122,12 +133,15 @@ struct FRigVMFunction_StringEndsWith : public FRigVMFunction_StringBase
 	RIGVM_METHOD()
 	RIGVM_API virtual void Execute();
 
+	// The input string to search within
 	UPROPERTY(meta = (Input))
 	FString Name;
 
+	// The ending string to look for
 	UPROPERTY(meta = (Input))
 	FString Ending;
 
+	// True if the given input string ends with the given ending
 	UPROPERTY(meta = (Output))
 	bool Result;
 };
@@ -149,12 +163,15 @@ struct FRigVMFunction_StringStartsWith : public FRigVMFunction_StringBase
 	RIGVM_METHOD()
 	RIGVM_API virtual void Execute();
 
+	// The input string to search within
 	UPROPERTY(meta = (Input))
 	FString Name;
 
+	// The start string to look for
 	UPROPERTY(meta = (Input))
 	FString Start;
 
+	// True if the given input string starts with the given start string
 	UPROPERTY(meta = (Output))
 	bool Result;
 };
@@ -176,12 +193,15 @@ struct FRigVMFunction_StringContains : public FRigVMFunction_StringBase
 	RIGVM_METHOD()
 	RIGVM_API virtual void Execute();
 
+	// The input string to search within
 	UPROPERTY(meta = (Input))
 	FString Name;
 
+	// The string to search for
 	UPROPERTY(meta = (Input))
 	FString Search;
 
+	// True if a given string exists in another given string
 	UPROPERTY(meta = (Output))
 	bool Result;
 };
@@ -203,9 +223,11 @@ struct FRigVMFunction_StringLength : public FRigVMFunction_StringBase
 	RIGVM_METHOD()
 	RIGVM_API virtual void Execute();
 
+	// The input string to retrieve the length for
 	UPROPERTY(meta = (Input))
 	FString Value;
 
+	// The resulting length of the input string
 	UPROPERTY(meta = (Output))
 	int32 Length;
 };
@@ -226,9 +248,11 @@ struct FRigVMFunction_StringTrimWhitespace : public FRigVMFunction_StringBase
 	RIGVM_METHOD()
 	RIGVM_API virtual void Execute();
 
+	// The input string to trim
 	UPROPERTY(meta = (Input))
 	FString Value;
 
+	// The resulting trimmed string
 	UPROPERTY(meta = (Output))
 	FString Result;
 };
@@ -249,9 +273,11 @@ struct FRigVMFunction_StringToUppercase : public FRigVMFunction_StringBase
 	RIGVM_METHOD()
 	RIGVM_API virtual void Execute();
 
+	// The input string to convert to uppercase
 	UPROPERTY(meta = (Input))
 	FString Value;
 
+	// The resulting uppercased string
 	UPROPERTY(meta = (Output))
 	FString Result;
 };
@@ -272,9 +298,11 @@ struct FRigVMFunction_StringToLowercase : public FRigVMFunction_StringBase
 	RIGVM_METHOD()
 	RIGVM_API virtual void Execute();
 
+	// The input string to convert to lowercase
 	UPROPERTY(meta = (Input))
 	FString Value;
 
+	// The resulting lowercased string
 	UPROPERTY(meta = (Output))
 	FString Result;
 };
@@ -295,9 +323,11 @@ struct FRigVMFunction_StringReverse : public FRigVMFunction_StringBase
 	RIGVM_METHOD()
 	RIGVM_API virtual void Execute();
 
+	// The input string to reverse
 	UPROPERTY(meta = (Input))
 	FString Value;
 
+	// The resulting reversed string
 	UPROPERTY(meta = (Output))
 	FString Reverse;
 };
@@ -319,12 +349,15 @@ struct FRigVMFunction_StringLeft : public FRigVMFunction_StringBase
 	RIGVM_METHOD()
 	RIGVM_API virtual void Execute();
 
+	// The input string to retrieve the left most characters of
 	UPROPERTY(meta = (Input))
 	FString Value;
 
+	// The number of characters to retrieve
 	UPROPERTY(meta = (Input))
 	int32 Count;
 
+	// The resulting left most characters of the input string
 	UPROPERTY(meta = (Output))
 	FString Result;
 };
@@ -346,12 +379,15 @@ struct FRigVMFunction_StringRight : public FRigVMFunction_StringBase
 	RIGVM_METHOD()
 	RIGVM_API virtual void Execute();
 
+	// The input string to retrieve the right most characters of
 	UPROPERTY(meta = (Input))
 	FString Value;
 
+	// The number of characters to retrieve
 	UPROPERTY(meta = (Input))
 	int32 Count;
 
+	// The resulting right most characters of the input string
 	UPROPERTY(meta = (Output))
 	FString Result;
 };
@@ -374,6 +410,7 @@ struct FRigVMFunction_StringMiddle : public FRigVMFunction_StringBase
 	RIGVM_METHOD()
 	RIGVM_API virtual void Execute();
 
+	// The input string to retrieve a section of
 	UPROPERTY(meta = (Input))
 	FString Value;
 
@@ -385,6 +422,7 @@ struct FRigVMFunction_StringMiddle : public FRigVMFunction_StringBase
 	UPROPERTY(meta = (Input))
 	int32 Count;
 
+	// The resulting section of the input string given start and count
 	UPROPERTY(meta = (Output))
 	FString Result;
 };
@@ -407,15 +445,19 @@ struct FRigVMFunction_StringFind : public FRigVMFunction_StringBase
 	RIGVM_METHOD()
 	RIGVM_API virtual void Execute();
 
+	// The input string to search within
 	UPROPERTY(meta = (Input))
 	FString Value;
 
+	// The string to search for
 	UPROPERTY(meta = (Input))
 	FString Search;
 
+	// True if the search string was found in the input string
 	UPROPERTY(meta = (Output))
     bool Found;
 
+	// The index of the search string within the input string (or INDEX_NONE)
 	UPROPERTY(meta = (Output))
 	int32 Index;
 };
@@ -436,12 +478,15 @@ struct FRigVMFunction_StringSplit : public FRigVMFunction_StringBase
 	RIGVM_METHOD()
 	RIGVM_API virtual void Execute();
 
+	// The input string to split up
 	UPROPERTY(meta = (Input))
 	FString Value;
 
+	// The separator to use to split up the input string
 	UPROPERTY(meta = (Input))
 	FString Separator;
 
+	// The array of split up parts of the input string 
 	UPROPERTY(meta = (Output))
 	TArray<FString> Result;
 };
@@ -462,12 +507,15 @@ struct FRigVMFunction_StringJoin : public FRigVMFunction_StringBase
 	RIGVM_METHOD()
 	RIGVM_API virtual void Execute();
 
+	// An array of input strings to join
 	UPROPERTY(meta = (Input))
 	TArray<FString> Values;
 
+	// The separator to use when joining the parts
 	UPROPERTY(meta = (Input))
 	FString Separator;
 
+	// The resulting joined string
 	UPROPERTY(meta = (Output))
 	FString Result;
 };
@@ -490,12 +538,15 @@ struct FRigVMFunction_StringPadInteger : public FRigVMFunction_StringBase
 	RIGVM_METHOD()
 	RIGVM_API virtual void Execute();
 
+	// The input number to convert to a string (for example 4)
 	UPROPERTY(meta = (Input))
 	int32 Value;
 
+	// The number of digits in total to use (for example 3)
 	UPROPERTY(meta = (Input))
 	int32 Digits;
-
+ 
+	// The resulting string representing the padded integer (for example '004')
 	UPROPERTY(meta = (Output))
 	FString Result;
 };
@@ -518,6 +569,7 @@ struct FRigVMFunction_StringToInteger : public FRigVMFunction_StringBase
 	RIGVM_METHOD()
 	RIGVM_API virtual void Execute();
 
+	// The input string to convert to integer
 	UPROPERTY(meta = (Input))
 	FString Value;
 
@@ -529,9 +581,11 @@ struct FRigVMFunction_StringToInteger : public FRigVMFunction_StringBase
 	UPROPERTY(meta = (Input))
 	bool ChopRight;
 
+	// The resulting integer number (or INDEX_NONE)
 	UPROPERTY(meta = (Output))
 	int32 Result;
 
+	// True if the string was successfully converted to integer
 	UPROPERTY(meta = (Output))
 	bool Success;
 };
