RigVM: Comment all simulation nodes

#rnx

diff --git a/Engine/Plugins/Animation/ControlRig/Source/ControlRig/Public/Units/Simulation/RigUnit_SpringInterp.h b/Engine/Plugins/Animation/ControlRig/Source/ControlRig/Public/Units/Simulation/RigUnit_SpringInterp.h
--- a/Engine/Plugins/Animation/ControlRig/Source/ControlRig/Public/Units/Simulation/RigUnit_SpringInterp.h
+++ b/Engine/Plugins/Animation/ControlRig/Source/ControlRig/Public/Units/Simulation/RigUnit_SpringInterp.h
@@ -177,7 +177,7 @@ struct FRigUnit_SpringInterpV2 : public FRigVMFunction_SimBase
 	UPROPERTY(meta=(Output))
 	float Result;
 
-	/** Velocity */
+	/** The resulting velocity of this spring */
 	UPROPERTY(meta = (Output))
 	float Velocity;
 
@@ -267,7 +267,7 @@ struct FRigUnit_SpringInterpVectorV2 : public FRigVMFunction_SimBase
 	UPROPERTY(meta=(Output))
 	FVector Result;
 
-	/** Velocity */
+	/** The resulting velocity of this spring */
 	UPROPERTY(meta = (Output))
 	FVector Velocity;
 
diff --git a/Engine/Plugins/Runtime/RigVM/Source/RigVM/Public/RigVMFunctions/Execution/RigVMFunction_ForLoop.h b/Engine/Plugins/Runtime/RigVM/Source/RigVM/Public/RigVMFunctions/Execution/RigVMFunction_ForLoop.h
--- a/Engine/Plugins/Runtime/RigVM/Source/RigVM/Public/RigVMFunctions/Execution/RigVMFunction_ForLoop.h
+++ b/Engine/Plugins/Runtime/RigVM/Source/RigVM/Public/RigVMFunctions/Execution/RigVMFunction_ForLoop.h
@@ -8,7 +8,7 @@
 /**
  * Given a count, execute iteratively until the count is up
  */
-USTRUCT(meta=(DisplayName="For Loop", Category="Execution", TitleColor="1 0 0", NodeColor="1 1 1", Keywords="Iterate", Icon="EditorStyle|GraphEditor.Macro.Loop_16x"))
+USTRUCT(meta=(DisplayName="For Loop", Category="Execution", TitleColor="1 0 0", NodeColor="1 1 1", Keywords="Iterate", Icon="EditorStyle|GraphEditor.Macro.Loop_16x", DocumentationPolicy="Strict"))
 struct FRigVMFunction_ForLoopCount : public FRigVMStructMutable
 {
 	GENERATED_BODY()
@@ -33,18 +33,23 @@ struct FRigVMFunction_ForLoopCount : public FRigVMStructMutable
 	RIGVM_METHOD()
 	RIGVM_API virtual void Execute() override;
 
+	// The internal block to run as this node progressed
 	UPROPERTY(meta = (Singleton))
 	FName BlockToRun;
 
+	// The number of iterations for this loop
 	UPROPERTY(meta = (Singleton, Input))
 	int32 Count;
 
+	// The index of the current iteration
 	UPROPERTY(meta = (Singleton, Output))
 	int32 Index;
 
+	// The ratio of the current iteration (from 0.0 for the first to 1.0 for the last)
 	UPROPERTY(meta = (Singleton, Output))
 	float Ratio;
 
+	// The completed branch to run once the loop has finished all iterations
 	UPROPERTY(EditAnywhere, Transient, Category = "ForLoop", meta = (Output))
 	FRigVMExecuteContext Completed;
 };
diff --git a/Engine/Plugins/Runtime/RigVM/Source/RigVM/Public/RigVMFunctions/Simulation/RigVMFunction_Accumulate.h b/Engine/Plugins/Runtime/RigVM/Source/RigVM/Public/RigVMFunctions/Simulation/RigVMFunction_Accumulate.h
--- a/Engine/Plugins/Runtime/RigVM/Source/RigVM/Public/RigVMFunctions/Simulation/RigVMFunction_Accumulate.h
+++ b/Engine/Plugins/Runtime/RigVM/Source/RigVM/Public/RigVMFunctions/Simulation/RigVMFunction_Accumulate.h
@@ -5,6 +5,9 @@
 #include "RigVMFunctions/Simulation/RigVMFunction_SimBase.h"
 #include "RigVMFunction_Accumulate.generated.h"
 
+/*
+ * The base class for all accumulation nodes
+ */
 USTRUCT(meta=(Category="Simulation|Accumulate"))
 struct FRigVMFunction_AccumulateBase : public FRigVMFunction_SimBase
 {
@@ -31,15 +34,20 @@ struct FRigVMFunction_AccumulateFloatAdd : public FRigVMFunction_AccumulateBase
 	RIGVM_METHOD()
 	RIGVM_API virtual void Execute() override;
 
+	// The increment to add for every iteration
 	UPROPERTY(meta=(Input))
 	float Increment;
 
+	// The initial value to start at
 	UPROPERTY(meta = (Input))
 	float InitialValue;
 
+	// If True the integration will be relying on the delta time
+	// to create more deterministic results with varying framerates.
 	UPROPERTY(meta=(Input))
 	bool bIntegrateDeltaTime;
 
+	// The resulting accumulated value
 	UPROPERTY(meta=(Output))
 	float Result;
 
@@ -70,15 +78,20 @@ struct FRigVMFunction_AccumulateVectorAdd : public FRigVMFunction_AccumulateBase
 	RIGVM_METHOD()
 	RIGVM_API virtual void Execute() override;
 
+	// The increment to add for every iteration
 	UPROPERTY(meta = (Input))
 	FVector Increment;
 
+	// The initial value to start at
 	UPROPERTY(meta = (Input))
 	FVector InitialValue;
 
+	// If True the integration will be relying on the delta time
+	// to create more deterministic results with varying framerates.
 	UPROPERTY(meta = (Input))
 	bool bIntegrateDeltaTime;
 
+	// The resulting accumulated value
 	UPROPERTY(meta = (Output))
 	FVector Result;
 
@@ -109,15 +122,20 @@ struct FRigVMFunction_AccumulateFloatMul : public FRigVMFunction_AccumulateBase
 	RIGVM_METHOD()
 	RIGVM_API virtual void Execute() override;
 
+	// The multiplier to apply for every iteration
 	UPROPERTY(meta=(Input))
 	float Multiplier;
 
+	// The initial value to start at
 	UPROPERTY(meta = (Input))
 	float InitialValue;
 
+	// If True the integration will be relying on the delta time
+	// to create more deterministic results with varying framerates.
 	UPROPERTY(meta=(Input))
 	bool bIntegrateDeltaTime;
 
+	// The resulting accumulated value
 	UPROPERTY(meta=(Output))
 	float Result;
 
@@ -148,15 +166,20 @@ struct FRigVMFunction_AccumulateVectorMul : public FRigVMFunction_AccumulateBase
 	RIGVM_METHOD()
 	RIGVM_API virtual void Execute() override;
 
+	// The multiplier to apply for every iteration
 	UPROPERTY(meta = (Input))
 	FVector Multiplier;
 
+	// The initial value to start at
 	UPROPERTY(meta = (Input))
 	FVector InitialValue;
 
+	// If True the integration will be relying on the delta time
+	// to create more deterministic results with varying framerates.
 	UPROPERTY(meta = (Input))
 	bool bIntegrateDeltaTime;
 
+	// The resulting accumulated value
 	UPROPERTY(meta = (Output))
 	FVector Result;
 
@@ -187,18 +210,24 @@ struct FRigVMFunction_AccumulateQuatMul : public FRigVMFunction_AccumulateBase
 	RIGVM_METHOD()
 	RIGVM_API virtual void Execute() override;
 
+	// The multiplier to apply for every iteration
 	UPROPERTY(meta=(Input))
 	FQuat Multiplier;
 
+	// The initial value to start at
 	UPROPERTY(meta = (Input))
 	FQuat InitialValue;
 
+	// If True the multiplier will be pre-multiplied, otherwise post-multiplied
 	UPROPERTY(meta = (Input))
 	bool bFlipOrder;
 
+	// If True the integration will be relying on the delta time
+	// to create more deterministic results with varying framerates.
 	UPROPERTY(meta=(Input))
 	bool bIntegrateDeltaTime;
 
+	// The resulting accumulated value
 	UPROPERTY(meta=(Output))
 	FQuat Result;
 
@@ -229,18 +258,24 @@ struct FRigVMFunction_AccumulateTransformMul : public FRigVMFunction_AccumulateB
 	RIGVM_METHOD()
 	RIGVM_API virtual void Execute() override;
 
+	// The multiplier to apply for every iteration
 	UPROPERTY(meta = (Input))
 	FTransform Multiplier;
 
+	// The initial value to start at
 	UPROPERTY(meta = (Input))
 	FTransform InitialValue;
 
+	// If True the multiplier will be pre-multiplied, otherwise post-multiplied
 	UPROPERTY(meta = (Input))
 	bool bFlipOrder;
 
+	// If True the integration will be relying on the delta time
+	// to create more deterministic results with varying framerates.
 	UPROPERTY(meta = (Input))
 	bool bIntegrateDeltaTime;
 
+	// The resulting accumulated value
 	UPROPERTY(meta = (Output))
 	FTransform Result;
 
@@ -271,18 +306,25 @@ struct FRigVMFunction_AccumulateFloatLerp : public FRigVMFunction_AccumulateBase
 	RIGVM_METHOD()
 	RIGVM_API virtual void Execute() override;
 
+	// The target value to interpolate towards
 	UPROPERTY(meta=(Input))
 	float TargetValue;
 
+	// The initial value to start at
 	UPROPERTY(meta = (Input))
 	float InitialValue;
 
+	// The blend to use for the interpolation.
+	// This value may be scaled down based on the Integrate Delta Time setting
 	UPROPERTY(meta = (Input))
 	float Blend;
 
+	// If True the integration will be relying on the delta time
+	// to create more deterministic results with varying framerates.
 	UPROPERTY(meta=(Input))
 	bool bIntegrateDeltaTime;
 
+	// The resulting accumulated value
 	UPROPERTY(meta=(Output))
 	float Result;
 
@@ -314,18 +356,25 @@ struct FRigVMFunction_AccumulateVectorLerp : public FRigVMFunction_AccumulateBas
 	RIGVM_METHOD()
 	RIGVM_API virtual void Execute() override;
 
+	// The target value to interpolate towards
 	UPROPERTY(meta = (Input))
 	FVector TargetValue;
 
+	// The initial value to start at
 	UPROPERTY(meta = (Input))
 	FVector InitialValue;
 
+	// The blend to use for the interpolation.
+	// This value may be scaled down based on the Integrate Delta Time setting
 	UPROPERTY(meta = (Input))
 	float Blend;
 
+	// If True the integration will be relying on the delta time
+	// to create more deterministic results with varying framerates.
 	UPROPERTY(meta = (Input))
 	bool bIntegrateDeltaTime;
 
+	// The resulting accumulated value
 	UPROPERTY(meta = (Output))
 	FVector Result;
 
@@ -357,18 +406,25 @@ struct FRigVMFunction_AccumulateQuatLerp : public FRigVMFunction_AccumulateBase
 	RIGVM_METHOD()
 	RIGVM_API virtual void Execute() override;
 
+	// The target value to interpolate towards
 	UPROPERTY(meta=(Input))
 	FQuat TargetValue;
 
+	// The initial value to start at
 	UPROPERTY(meta = (Input))
 	FQuat InitialValue;
 
+	// The blend to use for the interpolation.
+	// This value may be scaled down based on the Integrate Delta Time setting
 	UPROPERTY(meta = (Input))
 	float Blend;
 
+	// If True the integration will be relying on the delta time
+	// to create more deterministic results with varying framerates.
 	UPROPERTY(meta=(Input))
 	bool bIntegrateDeltaTime;
 
+	// The resulting accumulated value
 	UPROPERTY(meta=(Output))
 	FQuat Result;
 
@@ -400,18 +456,25 @@ struct FRigVMFunction_AccumulateTransformLerp : public FRigVMFunction_Accumulate
 	RIGVM_METHOD()
 	RIGVM_API virtual void Execute() override;
 
+	// The target value to interpolate towards
 	UPROPERTY(meta = (Input))
 	FTransform TargetValue;
 
+	// The initial value to start at
 	UPROPERTY(meta = (Input))
 	FTransform InitialValue;
 
+	// The blend to use for the interpolation.
+	// This value may be scaled down based on the Integrate Delta Time setting
 	UPROPERTY(meta = (Input))
 	float Blend;
 
+	// If True the integration will be relying on the delta time
+	// to create more deterministic results with varying framerates.
 	UPROPERTY(meta = (Input))
 	bool bIntegrateDeltaTime;
 
+	// The resulting accumulated value
 	UPROPERTY(meta = (Output))
 	FTransform Result;
 
@@ -441,12 +504,15 @@ struct FRigVMFunction_AccumulateFloatRange : public FRigVMFunction_AccumulateBas
 	RIGVM_METHOD()
 	RIGVM_API virtual void Execute() override;
 
+	// The value to accumulate the min and max value of
 	UPROPERTY(meta = (Input))
 	float Value;
 
+	// The accumulated minimum value of the input
 	UPROPERTY(meta= (Output))
 	float Minimum;
 
+	// The accumulated maximum value of the input
 	UPROPERTY(meta= (Output))
 	float Maximum;
 
@@ -479,12 +545,15 @@ struct FRigVMFunction_AccumulateVectorRange : public FRigVMFunction_AccumulateBa
 	RIGVM_METHOD()
 	RIGVM_API virtual void Execute() override;
 
+	// The value to accumulate the min and max value of
 	UPROPERTY(meta = (Input))
 	FVector Value;
 
+	// The accumulated minimum value of the input (for each component)
 	UPROPERTY(meta= (Output))
 	FVector Minimum;
 
+	// The accumulated maximum value of the input (for each component)
 	UPROPERTY(meta= (Output))
 	FVector Maximum;
 
diff --git a/Engine/Plugins/Runtime/RigVM/Source/RigVM/Public/RigVMFunctions/Simulation/RigVMFunction_AlphaInterp.h b/Engine/Plugins/Runtime/RigVM/Source/RigVM/Public/RigVMFunctions/Simulation/RigVMFunction_AlphaInterp.h
--- a/Engine/Plugins/Runtime/RigVM/Source/RigVM/Public/RigVMFunctions/Simulation/RigVMFunction_AlphaInterp.h
+++ b/Engine/Plugins/Runtime/RigVM/Source/RigVM/Public/RigVMFunctions/Simulation/RigVMFunction_AlphaInterp.h
@@ -36,42 +36,55 @@ struct FRigVMFunction_AlphaInterp : public FRigVMFunction_SimBase
 
 	RIGVM_API virtual FString ProcessPinLabelForInjection(const FString& InLabel) const override;
 
+	// The input value to interpolate
 	UPROPERTY(meta=(Input))
 	float Value;
 
+	// The scale to apply to the interpolation
 	UPROPERTY(meta = (Input))
 	float Scale;
 
+	// The bias to use for the interpolation  
 	UPROPERTY(meta = (Input))
 	float Bias;
 
+	// If True the input value will be mapped using the min and max range 
 	UPROPERTY(meta=(Input, Constant))
 	bool bMapRange;
 
+	// The minimum and maximum for the input range
 	UPROPERTY(meta=(Input, EditCondition = "bMapRange"))
 	FInputRange InRange;
 
+	// The minimum and maximum for the output range
 	UPROPERTY(meta=(Input, EditCondition = "bMapRange"))
 	FInputRange OutRange;
 
+	// If True the output value will be clamped by the min and max 
 	UPROPERTY(meta = (Input, Constant))
 	bool bClampResult;
 
+	// The minimum for the output's clamp range 
 	UPROPERTY(meta=(Input, EditCondition = "bClampResult"))
 	float ClampMin;
 
+	// The maximum for the output's clamp range 
 	UPROPERTY(meta=(Input, EditCondition = "bClampResult"))
 	float ClampMax;
 
+	// If True to the output result will be further intepolated
 	UPROPERTY(meta = (Input, Constant))
 	bool bInterpResult;
 
+	// The output interpolation increasing speed
 	UPROPERTY(meta=(Input, EditCondition = "bInterpResult"))
 	float InterpSpeedIncreasing;
 
+	// The output interpolation decreasing speed
 	UPROPERTY(meta=(Input, EditCondition = "bInterpResult"))
 	float InterpSpeedDecreasing;
 
+	// The resulting interpolated value
 	UPROPERTY(meta=(Output))
 	float Result;
 
@@ -109,42 +122,55 @@ struct FRigVMFunction_AlphaInterpVector : public FRigVMFunction_SimBase
 
 	RIGVM_API virtual FString ProcessPinLabelForInjection(const FString& InLabel) const override;
 
+	// The input value to interpolate
 	UPROPERTY(meta=(Input))
 	FVector Value;
 
+	// The scale to apply to the interpolation
 	UPROPERTY(meta = (Input))
 	float Scale;
 
+	// The bias to use for the interpolation  
 	UPROPERTY(meta = (Input))
 	float Bias;
 
+	// If True the input value will be mapped using the min and max range 
 	UPROPERTY(meta = (Input, Constant))
 	bool bMapRange;
 
+	// The minimum and maximum for the input range
 	UPROPERTY(meta = (Input, EditCondition = "bMapRange"))
 	FInputRange InRange;
 
+	// The minimum and maximum for the output range
 	UPROPERTY(meta = (Input, EditCondition = "bMapRange"))
 	FInputRange OutRange;
 
+	// If True the output value will be clamped by the min and max 
 	UPROPERTY(meta = (Input, Constant))
 	bool bClampResult;
 
+	// The minimum for the output's clamp range 
 	UPROPERTY(meta = (Input, EditCondition = "bClampResult"))
 	float ClampMin;
 
+	// The maximum for the output's clamp range 
 	UPROPERTY(meta = (Input, EditCondition = "bClampResult"))
 	float ClampMax;
 
+	// If True to the output result will be further intepolated
 	UPROPERTY(meta = (Input, Constant))
 	bool bInterpResult;
 
+	// The output interpolation increasing speed
 	UPROPERTY(meta = (Input, EditCondition = "bInterpResult"))
 	float InterpSpeedIncreasing;
 
+	// The output interpolation decreasing speed
 	UPROPERTY(meta = (Input, EditCondition = "bInterpResult"))
 	float InterpSpeedDecreasing;
 
+	// The resulting interpolated value
 	UPROPERTY(meta=(Output))
 	FVector Result;
 
@@ -182,42 +208,55 @@ struct FRigVMFunction_AlphaInterpQuat : public FRigVMFunction_SimBase
 
 	RIGVM_API virtual FString ProcessPinLabelForInjection(const FString& InLabel) const override;
 
+	// The input value to interpolate
 	UPROPERTY(meta=(Input))
 	FQuat Value;
 
+	// The scale to apply to the interpolation
 	UPROPERTY(meta = (Input))
 	float Scale;
 
+	// The bias to use for the interpolation  
 	UPROPERTY(meta = (Input))
 	float Bias;
 
+	// If True the input value will be mapped using the min and max range 
 	UPROPERTY(meta = (Input, Constant))
 	bool bMapRange;
 
+	// The minimum and maximum for the input range
 	UPROPERTY(meta = (Input, EditCondition = "bMapRange"))
 	FInputRange InRange;
 
+	// The minimum and maximum for the output range
 	UPROPERTY(meta = (Input, EditCondition = "bMapRange"))
 	FInputRange OutRange;
 
+	// If True the output value will be clamped by the min and max 
 	UPROPERTY(meta = (Input, Constant))
 	bool bClampResult;
 
+	// The minimum for the output's clamp range 
 	UPROPERTY(meta = (Input, EditCondition = "bClampResult"))
 	float ClampMin;
 
+	// The maximum for the output's clamp range 
 	UPROPERTY(meta = (Input, EditCondition = "bClampResult"))
 	float ClampMax;
 
+	// If True to the output result will be further intepolated
 	UPROPERTY(meta = (Input, Constant))
 	bool bInterpResult;
 
+	// The output interpolation increasing speed
 	UPROPERTY(meta = (Input, EditCondition = "bInterpResult"))
 	float InterpSpeedIncreasing;
 
+	// The output interpolation decreasing speed
 	UPROPERTY(meta = (Input, EditCondition = "bInterpResult"))
 	float InterpSpeedDecreasing;
 
+	// The resulting interpolated value
 	UPROPERTY(meta=(Output))
 	FQuat Result;
 
diff --git a/Engine/Plugins/Runtime/RigVM/Source/RigVM/Public/RigVMFunctions/Simulation/RigVMFunction_DeltaFromPrevious.h b/Engine/Plugins/Runtime/RigVM/Source/RigVM/Public/RigVMFunctions/Simulation/RigVMFunction_DeltaFromPrevious.h
--- a/Engine/Plugins/Runtime/RigVM/Source/RigVM/Public/RigVMFunctions/Simulation/RigVMFunction_DeltaFromPrevious.h
+++ b/Engine/Plugins/Runtime/RigVM/Source/RigVM/Public/RigVMFunctions/Simulation/RigVMFunction_DeltaFromPrevious.h
@@ -24,12 +24,15 @@ struct FRigVMFunction_DeltaFromPreviousFloat : public FRigVMFunction_SimBase
 	RIGVM_METHOD()
 	UE_API virtual void Execute() override;
 
+	// The input value to compute the delta of
 	UPROPERTY(meta = (Input))
 	float Value;
 
+	// The resulting delta from the previous iteration
 	UPROPERTY(meta=(Output))
 	float Delta;
 
+	// The value from the previous iteration
 	UPROPERTY(meta=(Output))
 	float PreviousValue;
 
@@ -57,12 +60,15 @@ struct FRigVMFunction_DeltaFromPreviousVector : public FRigVMFunction_SimBase
 	RIGVM_METHOD()
 	UE_API virtual void Execute() override;
 
+	// The input value to compute the delta of
 	UPROPERTY(meta = (Input))
 	FVector Value;
 
+	// The resulting delta from the previous iteration
 	UPROPERTY(meta=(Output))
 	FVector Delta;
 
+	// The value from the previous iteration
 	UPROPERTY(meta=(Output))
 	FVector PreviousValue;
 
@@ -90,12 +96,15 @@ struct FRigVMFunction_DeltaFromPreviousQuat : public FRigVMFunction_SimBase
 	RIGVM_METHOD()
 	UE_API virtual void Execute() override;
 
+	// The input value to compute the delta of
 	UPROPERTY(meta = (Input))
 	FQuat Value;
 
+	// The resulting delta from the previous iteration
 	UPROPERTY(meta=(Output))
 	FQuat Delta;
 
+	// The value from the previous iteration
 	UPROPERTY(meta=(Output))
 	FQuat PreviousValue;
 
@@ -123,12 +132,15 @@ struct FRigVMFunction_DeltaFromPreviousTransform : public FRigVMFunction_SimBase
 	RIGVM_METHOD()
 	UE_API virtual void Execute() override;
 
+	// The input value to compute the delta of
 	UPROPERTY(meta = (Input))
 	FTransform Value;
 
+	// The resulting delta from the previous iteration
 	UPROPERTY(meta=(Output))
 	FTransform Delta;
 
+	// The value from the previous iteration
 	UPROPERTY(meta=(Output))
 	FTransform PreviousValue;
 
diff --git a/Engine/Plugins/Runtime/RigVM/Source/RigVM/Public/RigVMFunctions/Simulation/RigVMFunction_Kalman.h b/Engine/Plugins/Runtime/RigVM/Source/RigVM/Public/RigVMFunctions/Simulation/RigVMFunction_Kalman.h
--- a/Engine/Plugins/Runtime/RigVM/Source/RigVM/Public/RigVMFunctions/Simulation/RigVMFunction_Kalman.h
+++ b/Engine/Plugins/Runtime/RigVM/Source/RigVM/Public/RigVMFunctions/Simulation/RigVMFunction_Kalman.h
@@ -26,12 +26,15 @@ struct FRigVMFunction_KalmanFloat : public FRigVMFunction_SimBase
 	RIGVM_METHOD()
 	RIGVM_API virtual void Execute() override;
 
+	// The value to filter
 	UPROPERTY(meta = (Input))
 	float Value;
 
+	// The buffer size to use (the higher the more precise)
 	UPROPERTY(meta=(Input, Constant))
 	int32 BufferSize;
 
+	// The resulting averaged / filtered value
 	UPROPERTY(meta=(Output))
 	float Result;
 
@@ -63,12 +66,15 @@ struct FRigVMFunction_KalmanVector : public FRigVMFunction_SimBase
 	RIGVM_METHOD()
 	RIGVM_API virtual void Execute() override;
 
+	// The value to filter
 	UPROPERTY(meta = (Input))
 	FVector Value;
 
+	// The buffer size to use (the higher the more precise)
 	UPROPERTY(meta=(Input, Constant))
 	int32 BufferSize;
 
+	// The resulting averaged / filtered value
 	UPROPERTY(meta=(Output))
 	FVector Result;
 
@@ -100,12 +106,15 @@ struct FRigVMFunction_KalmanTransform : public FRigVMFunction_SimBase
 	RIGVM_METHOD()
 	RIGVM_API virtual void Execute() override;
 
+	// The value to filter
 	UPROPERTY(meta = (Input))
 	FTransform Value;
 
+	// The buffer size to use (the higher the more precise)
 	UPROPERTY(meta=(Input, Constant))
 	int32 BufferSize;
 
+	// The resulting averaged / filtered value
 	UPROPERTY(meta=(Output))
 	FTransform Result;
 
diff --git a/Engine/Plugins/Runtime/RigVM/Source/RigVM/Public/RigVMFunctions/Simulation/RigVMFunction_SimBase.h b/Engine/Plugins/Runtime/RigVM/Source/RigVM/Public/RigVMFunctions/Simulation/RigVMFunction_SimBase.h
--- a/Engine/Plugins/Runtime/RigVM/Source/RigVM/Public/RigVMFunctions/Simulation/RigVMFunction_SimBase.h
+++ b/Engine/Plugins/Runtime/RigVM/Source/RigVM/Public/RigVMFunctions/Simulation/RigVMFunction_SimBase.h
@@ -5,13 +5,19 @@
 #include "RigVMCore/RigVMStruct.h"
 #include "RigVMFunction_SimBase.generated.h"
 
-USTRUCT(meta=(Abstract, Category = "Simulation", NodeColor = "0.25 0.05 0.05"))
+/*
+ * The base class for all pure simulation nodes
+ */
+USTRUCT(meta=(Abstract, Category = "Simulation", NodeColor = "0.25 0.05 0.05", DocumentationPolicy = "Strict"))
 struct FRigVMFunction_SimBase : public FRigVMStruct
 {
 	GENERATED_BODY()
 };
 
-USTRUCT(meta=(Abstract, Category = "Simulation", NodeColor = "0.25 0.05 0.05"))
+/*
+ * The base class for all mutable simulation nodes
+ */
+USTRUCT(meta=(Abstract, Category = "Simulation", NodeColor = "0.25 0.05 0.05", DocumentationPolicy = "Strict"))
 struct FRigVMFunction_SimBaseMutable : public FRigVMStructMutable
 {
 	GENERATED_BODY()
diff --git a/Engine/Plugins/Runtime/RigVM/Source/RigVM/Public/RigVMFunctions/Simulation/RigVMFunction_TimeOffset.h b/Engine/Plugins/Runtime/RigVM/Source/RigVM/Public/RigVMFunctions/Simulation/RigVMFunction_TimeOffset.h
--- a/Engine/Plugins/Runtime/RigVM/Source/RigVM/Public/RigVMFunctions/Simulation/RigVMFunction_TimeOffset.h
+++ b/Engine/Plugins/Runtime/RigVM/Source/RigVM/Public/RigVMFunctions/Simulation/RigVMFunction_TimeOffset.h
@@ -44,6 +44,7 @@ struct FRigVMFunction_TimeOffsetFloat : public FRigVMFunction_SimBase
 	UPROPERTY(meta=(Input, Constant))
 	float TimeRange;
 
+	// The resulting offset value
 	UPROPERTY(meta=(Output))
 	float Result;
 
@@ -99,6 +100,7 @@ struct FRigVMFunction_TimeOffsetVector : public FRigVMFunction_SimBase
 	UPROPERTY(meta=(Input, Constant))
 	float TimeRange;
 
+	// The resulting offset value
 	UPROPERTY(meta=(Output))
 	FVector Result;
 
@@ -154,6 +156,7 @@ struct FRigVMFunction_TimeOffsetTransform : public FRigVMFunction_SimBase
 	UPROPERTY(meta=(Input, Constant))
 	float TimeRange;
 
+	// The resulting offset value
 	UPROPERTY(meta=(Output))
 	FTransform Result;
 
diff --git a/Engine/Plugins/Runtime/RigVM/Source/RigVM/Public/RigVMFunctions/Simulation/RigVMFunction_Timeline.h b/Engine/Plugins/Runtime/RigVM/Source/RigVM/Public/RigVMFunctions/Simulation/RigVMFunction_Timeline.h
--- a/Engine/Plugins/Runtime/RigVM/Source/RigVM/Public/RigVMFunctions/Simulation/RigVMFunction_Timeline.h
+++ b/Engine/Plugins/Runtime/RigVM/Source/RigVM/Public/RigVMFunctions/Simulation/RigVMFunction_Timeline.h
@@ -25,9 +25,11 @@ struct FRigVMFunction_Timeline : public FRigVMFunction_SimBase
 	RIGVM_METHOD()
 	RIGVM_API virtual void Execute() override;
 
+	// The speed to apply to the played back time
 	UPROPERTY(meta = (Input))
 	float Speed;
 
+	// The played back / accumulated time in seconds
 	UPROPERTY(meta=(Output))
 	float Time;
 
@@ -62,6 +64,7 @@ struct FRigVMFunction_TimeLoop : public FRigVMFunction_SimBase
 	RIGVM_METHOD()
 	RIGVM_API virtual void Execute() override;
 
+	// The speed to play-back the time at
 	UPROPERTY(meta = (Input))
 	float Speed;
 
diff --git a/Engine/Plugins/Runtime/RigVM/Source/RigVM/Public/RigVMFunctions/Simulation/RigVMFunction_Verlet.h b/Engine/Plugins/Runtime/RigVM/Source/RigVM/Public/RigVMFunctions/Simulation/RigVMFunction_Verlet.h
--- a/Engine/Plugins/Runtime/RigVM/Source/RigVM/Public/RigVMFunctions/Simulation/RigVMFunction_Verlet.h
+++ b/Engine/Plugins/Runtime/RigVM/Source/RigVM/Public/RigVMFunctions/Simulation/RigVMFunction_Verlet.h
@@ -28,6 +28,7 @@ struct FRigVMFunction_VerletIntegrateVector : public FRigVMFunction_SimBase
 	RIGVM_METHOD()
 	RIGVM_API virtual void Execute() override;
 
+	// The target value to interpolate / integrate towards
 	UPROPERTY(meta = (Input))
 	FVector Target;
 
@@ -47,12 +48,15 @@ struct FRigVMFunction_VerletIntegrateVector : public FRigVMFunction_SimBase
 	UPROPERTY(meta = (Input))
 	FVector Force;
 
+	// The resulting simulated position
 	UPROPERTY(meta = (Output))
 	FVector Position;
 
+	// The resulting simulated velocity
 	UPROPERTY(meta = (Output))
 	FVector Velocity;
 
+	// The resulting simulated acceleration
 	UPROPERTY(meta = (Output))
 	FVector Acceleration;
 
