RigVM: Guard localized registry creation during serialization

Full cook 68ee0fb7d33288774355592f

#rnx

diff --git a/Engine/Plugins/Runtime/RigVM/Source/RigVM/Private/RigVMCore/RigVM.cpp b/Engine/Plugins/Runtime/RigVM/Source/RigVM/Private/RigVMCore/RigVM.cpp
--- a/Engine/Plugins/Runtime/RigVM/Source/RigVM/Private/RigVMCore/RigVM.cpp
+++ b/Engine/Plugins/Runtime/RigVM/Source/RigVM/Private/RigVMCore/RigVM.cpp
@@ -112,6 +112,8 @@ URigVM::~URigVM()
 
 void URigVM::Serialize(FArchive& Ar)
 {
+	const TGuardValue<bool> _(bIsSerializing, true);
+	
 	Ar.UsingCustomVersion(FAnimObjectVersion::GUID);
 	Ar.UsingCustomVersion(FUE5MainStreamObjectVersion::GUID);
 	Ar.UsingCustomVersion(FRigVMObjectVersion::GUID);
@@ -143,6 +145,8 @@ void URigVM::Serialize(FArchive& Ar)
 
 void URigVM::Save(FArchive& Ar)
 {
+	const bool bIsSerializingToDisk = IsSerializingToDisk(Ar);
+	
 	UE_RIGVM_ARCHIVETRACE_SCOPE(Ar, FString::Printf(TEXT("URigVM(%s)"), *GetName()));
 
 	// The Save function has to be in sync with CopyDataForSerialization
@@ -166,7 +170,7 @@ void URigVM::Save(FArchive& Ar)
 	Ar << DefaultWorkMemoryStorage;
 	UE_RIGVM_ARCHIVETRACE_ENTRY(Ar, TEXT("DefaultWorkMemoryStorage"));
 
-	if(Ar.IsPersistent() && CVarRigVMEnableLocalizedRegistry.GetValueOnAnyThread() && LocalizedRegistry.IsValid())
+	if(bIsSerializingToDisk && LocalizedRegistry.IsValid())
 	{
 		FRigVMMemoryStorageStruct EmptyDebugMemoryStruct;
 		Ar << EmptyDebugMemoryStruct;
@@ -359,6 +363,7 @@ void URigVM::CopyDataForSerialization(URigVM* InVM)
 	LiteralMemoryStorage = InVM->LiteralMemoryStorage;
 	DefaultWorkMemoryStorage = InVM->DefaultWorkMemoryStorage;
 	DefaultDebugMemoryStorage = InVM->DefaultDebugMemoryStorage;
+	CopyExternalVariableDefs(InVM);
 
 	LocalizedRegistry.Reset();
 	if (InVM->LocalizedRegistry.IsValid())
@@ -367,6 +372,17 @@ void URigVM::CopyDataForSerialization(URigVM* InVM)
 	}
 }
 
+void URigVM::CopyExternalVariableDefs(URigVM* InVM)
+{
+	ExternalVariables = InVM->ExternalVariables;
+}
+
+bool URigVM::IsSerializingToDisk(const FArchive& Ar)
+{
+	return Ar.IsPersistent()
+		&& !Ar.HasAnyPortFlags(PPF_Duplicate | PPF_DuplicateForPIE);
+}
+
 void URigVM::PostLoad()
 {
 	Super::PostLoad();
@@ -1337,6 +1353,15 @@ void URigVM::PrepareMemoryForExecution(FRigVMExtendedExecuteContext& Context)
 
 	InvalidateCachedMemory(Context);
 
+	if (URigVMHost* CDO = Cast<URigVMHost>(GetOuter()))
+	{
+		FRigVMExtendedExecuteContext& ExtendedExecuteContext = CDO->GetRigVMExtendedExecuteContext();
+
+		// update the VM's external variables
+		ClearExternalVariables(ExtendedExecuteContext);
+		SetExternalVariableDefs(CDO->GetExternalVariablesImpl(false));
+	}
+
 	if (Instructions.Num() == 0)
 	{
 		return;
@@ -2288,6 +2313,12 @@ ERigVMExecuteResult URigVM::ExecuteBranch(FRigVMExtendedExecuteContext& Context,
 	return ExecuteInstructions(Context, InBranchToRun.FirstInstruction, InBranchToRun.LastInstruction);
 }
 
+void URigVM::ClearExternalVariables(FRigVMExtendedExecuteContext& Context)
+{
+	ExternalVariables.Reset();
+	Context.ExternalVariableRuntimeData.Reset();
+}
+
 FRigVMExternalVariableDef URigVM::GetExternalVariableDefByName(const FName& InExternalVariableName)
 {
 	for (const FRigVMExternalVariableDef& ExternalVariable : ExternalVariables)
@@ -2315,6 +2346,17 @@ FRigVMExternalVariable URigVM::GetExternalVariableByName(const FRigVMExtendedExe
 	return FRigVMExternalVariable();
 }
 
+void URigVM::SetExternalVariableDefs(const TArray<FRigVMExternalVariable>& InExternalVariables)
+{
+	ExternalVariables.Reset(InExternalVariables.Num());
+	for (const FRigVMExternalVariableDef& ExternalVariable : InExternalVariables)
+	{
+		ExternalVariables.Add(ExternalVariable);
+	}
+
+	RefreshExternalPropertyPaths();
+}
+
 void URigVM::SetPropertyValueFromString(FRigVMExtendedExecuteContext& Context, const FRigVMOperand& InOperand, const FString& InDefaultValue)
 {
 	FRigVMMemoryStorageStruct* Memory = GetMemoryByType(Context, InOperand.GetMemoryType());
@@ -2340,7 +2382,11 @@ FRigVMRegistry_NoLock* URigVM::GetLocalizedRegistry() const
 bool URigVM::CreateLocalizedRegistryIfRequired(bool bForce)
 {
 	UE::TScopeLock CreateLocalizedRegistryScopeLock(CreateLocalizedRegistryMutex);
+	return CreateLocalizedRegistryIfRequired_NoLock(bForce);
+}
 
+bool URigVM::CreateLocalizedRegistryIfRequired_NoLock(bool bForce)
+{
 	if(LocalizedRegistry.IsValid())
 	{
 		return true;
@@ -2350,6 +2396,15 @@ bool URigVM::CreateLocalizedRegistryIfRequired(bool bForce)
 	{
 		return false;
 	}
+
+	auto ResetLocalizedRegistry = [this]() -> bool
+	{
+		LocalizedRegistry.Reset();
+		FunctionsStorage.Reset();
+		FactoriesStorage.Reset();
+		verify(ResolveFunctionsIfRequired());
+		return false;
+	};
 		
 	// create a new registry
 	TSharedRef<FRigVMRegistry_NoLock> TemporaryRegistry = FRigVMRegistry_NoLock::CreateLocalizedRegistry();
@@ -2377,10 +2432,20 @@ bool URigVM::CreateLocalizedRegistryIfRequired(bool bForce)
 				{
 					if(Operand.GetRegisterOffset() != INDEX_NONE)
 					{
+						if(!LiteralMemoryStorage.GetPropertyPaths().IsValidIndex(Operand.GetRegisterOffset()))
+						{
+							UE_LOG(LogRigVM, Error, TEXT("%s: Cannot resolve literal memory property path '%d'."), *GetPathName(), Operand.GetRegisterOffset());
+							return ResetLocalizedRegistry();
+						}
 						Property = LiteralMemoryStorage.GetPropertyPaths()[Operand.GetRegisterOffset()].GetTailProperty();
 					}
 					else
 					{
+						if(!LiteralMemoryStorage.IsValidIndex(Operand.GetRegisterIndex()))
+						{
+							UE_LOG(LogRigVM, Error, TEXT("%s: Cannot resolve literal memory property '%d'."), *GetPathName(), Operand.GetRegisterIndex());
+							return ResetLocalizedRegistry();
+						}
 						Property = LiteralMemoryStorage.GetProperty(Operand.GetRegisterIndex());
 					}
 					break;
@@ -2389,10 +2454,20 @@ bool URigVM::CreateLocalizedRegistryIfRequired(bool bForce)
 				{
 					if(Operand.GetRegisterOffset() != INDEX_NONE)
 					{
+						if(!DefaultWorkMemoryStorage.GetPropertyPaths().IsValidIndex(Operand.GetRegisterOffset()))
+						{
+							UE_LOG(LogRigVM, Error, TEXT("%s: Cannot resolve work memory property path '%d'."), *GetPathName(), Operand.GetRegisterOffset());
+							return ResetLocalizedRegistry();
+						}
 						Property = DefaultWorkMemoryStorage.GetPropertyPaths()[Operand.GetRegisterOffset()].GetTailProperty();
 					}
 					else
 					{
+						if(!DefaultWorkMemoryStorage.IsValidIndex(Operand.GetRegisterIndex()))
+						{
+							UE_LOG(LogRigVM, Error, TEXT("%s: Cannot resolve work memory property '%d'."), *GetPathName(), Operand.GetRegisterIndex());
+							return ResetLocalizedRegistry();
+						}
 						Property = DefaultWorkMemoryStorage.GetProperty(Operand.GetRegisterIndex());
 					}
 					break;
@@ -2405,10 +2480,20 @@ bool URigVM::CreateLocalizedRegistryIfRequired(bool bForce)
 						{
 							RefreshExternalPropertyPaths();
 						}
+						if(!ExternalPropertyPaths.IsValidIndex(Operand.GetRegisterOffset()))
+						{
+							UE_LOG(LogRigVM, Error, TEXT("%s: Cannot resolve external variable property path '%d'."), *GetPathName(), Operand.GetRegisterOffset());
+							return ResetLocalizedRegistry();
+						}
 						Property = ExternalPropertyPaths[Operand.GetRegisterOffset()].GetTailProperty();
 					}
 					else
 					{
+						if(!ExternalVariables.IsValidIndex(Operand.GetRegisterIndex()))
+						{
+							UE_LOG(LogRigVM, Error, TEXT("%s: Cannot resolve external variable '%d'."), *GetPathName(), Operand.GetRegisterIndex());
+							return ResetLocalizedRegistry();
+						}
 						Property = ExternalVariables[Operand.GetRegisterIndex()].Property;
 					}
 					break;
@@ -2425,8 +2510,11 @@ bool URigVM::CreateLocalizedRegistryIfRequired(bool bForce)
 				continue;
 			}
 
-			// just creating the argument will register the type as needed
-			FRigVMTemplateArgument::Make_NoLock(const_cast<FProperty*>(Property), TemporaryHandle);
+			// set up the types - but try to avoid loading types relying on StaticFindObjectFast
+			// since this may be happening during GC or serialization
+			const FName CPPTypeName = *RigVMTypeUtils::GetCPPTypeFromProperty(Property);
+			UObject* CPPTypeObject = RigVMTypeUtils::GetCPPTypeObjectFromProperty(Property);
+			TemporaryHandle->FindOrAddType_NoLock(FRigVMTemplateArgumentType(CPPTypeName, CPPTypeObject, bIsSerializing));
 		}
 	}
 
@@ -2446,7 +2534,15 @@ bool URigVM::CreateLocalizedRegistryIfRequired(bool bForce)
 					if (ArgumentTypeIndex != INDEX_NONE)
 					{
 						const FRigVMTemplateArgumentType& Type = ReadLock->GetType_NoLock(ArgumentTypeIndex);
-						TemporaryRegistry->FindOrAddType_NoLock(Type);
+						if (Type.IsValid())
+						{
+							TemporaryRegistry->FindOrAddType_NoLock(Type);
+						}
+						else
+						{
+							UE_LOG(LogRigVM, Error, TEXT("%s: Cannot resolve Argument Type %d for function '%s'."), *GetPathName(), (int32)ArgumentTypeIndex, *Function->Name);
+							return ResetLocalizedRegistry();
+						}
 					}
 				}
 			}
@@ -2472,14 +2568,14 @@ bool URigVM::CreateLocalizedRegistryIfRequired(bool bForce)
 		if (Factory == nullptr)
 		{
 			UE_LOG(LogRigVM, Error, TEXT("Found an unexpected nullptr for required dispatch factories in VM '%s'."), *GetPathName());
-			return false;
+			return ResetLocalizedRegistry();
 		}
 
 		UScriptStruct* FactoryStruct = Factory->GetScriptStruct();
 		if (FactoryStruct == nullptr)
 		{
 			UE_LOG(LogRigVM, Error, TEXT("Dispatch factory '%s' has an unexpected nullptr for its script struct in VM '%s'."), *Factory->GetFactoryName().ToString(), *GetPathName());
-			return false;
+			return ResetLocalizedRegistry();
 		}
 
 		const TArray<const FRigVMFunction*>& OriginalFunctionsForFactory = Pair.Value; 
@@ -2503,21 +2599,21 @@ bool URigVM::CreateLocalizedRegistryIfRequired(bool bForce)
 					if (!OriginalType.IsValid())
 					{
 						UE_LOG(LogRigVM, Error, TEXT("Cannot register dispatch dispatch '%s' for VM '%s'. Invalid types found.\nCannot localize registry."), *Factory->GetFactoryName().ToString(), *GetPathName());
-						return false;
+						return ResetLocalizedRegistry();
 					}
 
 					TemporaryTypeIndices.Add(TemporaryRegistry->FindOrAddType_NoLock(OriginalType));
 					if (TemporaryTypeIndices.Last() == INDEX_NONE)
 					{
 						UE_LOG(LogRigVM, Error, TEXT("Cannot register type '%s' in dispatch '%s' for VM '%s'\nCannot localize registry."), *OriginalType.CPPType.ToString(), *Factory->GetFactoryName().ToString(), *GetPathName());
-						return false;
+						return ResetLocalizedRegistry();
 					}
 
 					const FRigVMTemplateArgumentType TemporaryType = TemporaryRegistry->GetType_NoLock(TemporaryTypeIndices.Last());
 					if(TemporaryType != OriginalType)
 					{
 						UE_LOG(LogRigVM, Error, TEXT("Registered type '%s' in dispatch '%s' does not match ('%s') for VM '%s'\nCannot localize registry."), *OriginalType.CPPType.ToString(), *Factory->GetFactoryName().ToString(), *TemporaryType.CPPType.ToString(), *GetPathName());
-						return false;
+						return ResetLocalizedRegistry();
 					}
 				}
 				TemporaryArgumentInfos.Add(FRigVMTemplateArgumentInfo::MakeFlattenedFromTypeIndices(ArgumentInfo.Name, ArgumentInfo.Direction, TemporaryTypeIndices));
@@ -2544,7 +2640,7 @@ bool URigVM::CreateLocalizedRegistryIfRequired(bool bForce)
 					FRigVMTemplateTypeMap Types = Template->GetTypesForPermutation_NoLock(PermutationIndex, TemporaryHandle);
 					const FString TypesString = Template->GetStringFromArgumentTypes(Types, TemporaryHandle);
 					UE_LOG(LogRigVM, Error, TEXT("Cannot create permutation [%03d] of template '%s' types '%s' for VM '%s'\nCannot localize registry."), PermutationIndex, *Template->GetNotation().ToString(), *TypesString, *GetPathName());
-					return false;
+					return ResetLocalizedRegistry();
 				}
 
 				TemporaryFunctionsForFactory.Add(TemporaryPermutation);
@@ -2568,7 +2664,7 @@ bool URigVM::CreateLocalizedRegistryIfRequired(bool bForce)
 						UE_LOG(LogRigVM, Error, TEXT("Temporary function [%03d] '%s' of template '%s' for VM '%s'."), TemporaryFunctionIndex, *TemporaryFunctionsForFactory[TemporaryFunctionIndex]->Name, *Template->GetNotation().ToString(), *GetPathName());
 					}
 					UE_LOG(LogRigVM, Error, TEXT("Missing expected permutation '%s' of template '%s' for VM '%s'.\nCannot localize registry."), *OriginalFunction->Name, *Template->GetNotation().ToString(), *GetPathName());
-					return false;
+					return ResetLocalizedRegistry();
 				}
 			}
 		}
@@ -2598,7 +2694,7 @@ bool URigVM::CreateLocalizedRegistryIfRequired(bool bForce)
 				}
 
 				UE_LOG(LogRigVM, Error, TEXT("Cannot find expected permutation '%s' for VM '%s'\nCannot localize registry."), *Function->Name, *GetPathName());
-				return false;
+				return ResetLocalizedRegistry();
 			}
 		}
 		else
@@ -2619,15 +2715,9 @@ bool URigVM::CreateLocalizedRegistryIfRequired(bool bForce)
 	if(!ResolveFunctionsIfRequired())
 	{
 		UE_LOG(LogRigVM, Error, TEXT("Cannot localize registry due to missing function."));
-
-		LocalizedRegistry.Reset();
-
-		FunctionsStorage.Reset();
-		FactoriesStorage.Reset();
-		verify(ResolveFunctionsIfRequired());
-		return false;
+		return ResetLocalizedRegistry();
 	}
-	
+
 	return true;
 }
 
diff --git a/Engine/Plugins/Runtime/RigVM/Source/RigVM/Private/RigVMCore/RigVMTemplate.cpp b/Engine/Plugins/Runtime/RigVM/Source/RigVM/Private/RigVMCore/RigVMTemplate.cpp
--- a/Engine/Plugins/Runtime/RigVM/Source/RigVM/Private/RigVMCore/RigVMTemplate.cpp
+++ b/Engine/Plugins/Runtime/RigVM/Source/RigVM/Private/RigVMCore/RigVMTemplate.cpp
@@ -16,18 +16,18 @@
 
 ////////////////////////////////////////////////////////////////////////////////////////////////
 
-FRigVMTemplateArgumentType::FRigVMTemplateArgumentType(const FName& InCPPType, UObject* InCPPTypeObject)
+FRigVMTemplateArgumentType::FRigVMTemplateArgumentType(const FName& InCPPType, UObject* InCPPTypeObject, bool bIsSerializing)
 		: CPPType(InCPPType)
 		, CPPTypeObject(InCPPTypeObject)
 {
 	// InCppType is unreliable because not all caller knows that
 	// we use generated unique names for user defined structs
 	// so here we override the CppType name with the actual name used in the registry
-	const FString InCPPTypeString = CPPType.ToString();
-	CPPType = *RigVMTypeUtils::PostProcessCPPType(InCPPTypeString, CPPTypeObject);
+	CPPType = *RigVMTypeUtils::PostProcessCPPType(CPPType.ToString(), CPPTypeObject, nullptr, bIsSerializing);
 #if WITH_EDITOR
 	if (CPPType.IsNone())
 	{
+		const FString InCPPTypeString = CPPType.ToString();
 		if (const UClass* ObjectClass = InCPPTypeObject ? InCPPTypeObject->GetClass() : nullptr)
 		{
 			UE_LOG(LogRigVM, Warning, TEXT("FRigVMTemplateArgumentType(): Input CPPType '%s' (Input Object '%s') could not be resolved."),
@@ -63,39 +63,7 @@ FRigVMTemplateArgument::FRigVMTemplateArgument(FProperty* InProperty, FRigVMRegi
 #endif
 
 	const FName CPPTypeName = *RigVMTypeUtils::GetCPPTypeFromProperty(InProperty);
-	UObject* CPPTypeObject = nullptr;
-
-	if (FArrayProperty* ArrayProperty = CastField<FArrayProperty>(InProperty))
-	{
-		InProperty = ArrayProperty->Inner;
-	}
-
-	if (FStructProperty* StructProperty = CastField<FStructProperty>(InProperty))
-	{
-		CPPTypeObject = StructProperty->Struct;
-	}
-	else if (FEnumProperty* EnumProperty = CastField<FEnumProperty>(InProperty))
-	{
-		CPPTypeObject = EnumProperty->GetEnum();
-	}
-	else if (FByteProperty* ByteProperty = CastField<FByteProperty>(InProperty))
-	{
-		CPPTypeObject = ByteProperty->Enum;
-	}
-	else if (FClassProperty* ClassProperty = CastField<FClassProperty>(InProperty))
-	{
-		if(RigVMCore::SupportsUObjects())
-		{
-			CPPTypeObject = ClassProperty->MetaClass;
-		}
-	}
-	else if (FObjectProperty* ObjectProperty = CastField<FObjectProperty>(InProperty))
-	{
-		if(RigVMCore::SupportsUObjects())
-		{
-			CPPTypeObject = ObjectProperty->PropertyClass;
-		}
-	}
+	UObject* CPPTypeObject = RigVMTypeUtils::GetCPPTypeObjectFromProperty(InProperty);
 	
 	const FRigVMTemplateArgumentType Type(CPPTypeName, CPPTypeObject);
 	const TRigVMTypeIndex TypeIndex = InRegistry->FindOrAddType_NoLock(Type, true); 
diff --git a/Engine/Plugins/Runtime/RigVM/Source/RigVM/Private/RigVMTypeUtils.cpp b/Engine/Plugins/Runtime/RigVM/Source/RigVM/Private/RigVMTypeUtils.cpp
--- a/Engine/Plugins/Runtime/RigVM/Source/RigVM/Private/RigVMTypeUtils.cpp
+++ b/Engine/Plugins/Runtime/RigVM/Source/RigVM/Private/RigVMTypeUtils.cpp
@@ -5,6 +5,7 @@
 =============================================================================*/
 
 #include "RigVMTypeUtils.h"
+#include "RigVMModule.h"
 #include "AssetRegistry/AssetData.h"
 #include "AssetRegistry/AssetRegistryModule.h"
 #include "Engine/UserDefinedEnum.h"
@@ -415,7 +416,46 @@ UObject* RigVMTypeUtils::ObjectFromCPPType(FString& InOutCPPType, bool bUseRedir
 	return CPPTypeObject;
 }
 
-FString RigVMTypeUtils::PostProcessCPPType(const FString& InCPPType, UObject* InCPPTypeObject, const FRigVMUserDefinedTypeResolver* InResolvalInfo)
+UObject* RigVMTypeUtils::GetCPPTypeObjectFromProperty(const FProperty* InProperty)
+{
+	UObject* CPPTypeObject = nullptr;
+	
+	if (const FArrayProperty* ArrayProperty = CastField<FArrayProperty>(InProperty))
+	{
+		InProperty = ArrayProperty->Inner;
+	}
+
+	if (const FStructProperty* StructProperty = CastField<FStructProperty>(InProperty))
+	{
+		CPPTypeObject = StructProperty->Struct;
+	}
+	else if (const FEnumProperty* EnumProperty = CastField<FEnumProperty>(InProperty))
+	{
+		CPPTypeObject = EnumProperty->GetEnum();
+	}
+	else if (const FByteProperty* ByteProperty = CastField<FByteProperty>(InProperty))
+	{
+		CPPTypeObject = ByteProperty->Enum;
+	}
+	else if (const FClassProperty* ClassProperty = CastField<FClassProperty>(InProperty))
+	{
+		if(RigVMCore::SupportsUObjects())
+		{
+			CPPTypeObject = ClassProperty->MetaClass;
+		}
+	}
+	else if (const FObjectProperty* ObjectProperty = CastField<FObjectProperty>(InProperty))
+	{
+		if(RigVMCore::SupportsUObjects())
+		{
+			CPPTypeObject = ObjectProperty->PropertyClass;
+		}
+	}
+
+	return CPPTypeObject;
+}
+
+FString RigVMTypeUtils::PostProcessCPPType(const FString& InCPPType, UObject* InCPPTypeObject, const FRigVMUserDefinedTypeResolver* InResolvalInfo, bool bIsSerializing)
 {
 	FString CPPType = InCPPType;
 	CleanupCPPType(CPPType);
@@ -436,8 +476,11 @@ FString RigVMTypeUtils::PostProcessCPPType(const FString& InCPPType, UObject* In
 	}
 	else if (RequiresCPPTypeObject(CPPType))
 	{
-		// Uses redirectors and updates the CPPType if necessary
-		ObjectFromCPPType(CPPType, true, InResolvalInfo);
+		if (!bIsSerializing)
+		{
+			// Uses redirectors and updates the CPPType if necessary
+			ObjectFromCPPType(CPPType, true, InResolvalInfo);
+		}
 	}
 
 	return CPPType;
diff --git a/Engine/Plugins/Runtime/RigVM/Source/RigVM/Public/RigVMCore/RigVM.h b/Engine/Plugins/Runtime/RigVM/Source/RigVM/Public/RigVMCore/RigVM.h
--- a/Engine/Plugins/Runtime/RigVM/Source/RigVM/Public/RigVMCore/RigVM.h
+++ b/Engine/Plugins/Runtime/RigVM/Source/RigVM/Public/RigVMCore/RigVM.h
@@ -165,6 +165,9 @@ public:
 	RIGVM_API virtual void Save(FArchive& Ar);
 	RIGVM_API virtual void Load(FArchive& Ar);
 	RIGVM_API void CopyDataForSerialization(URigVM* InVM);
+	RIGVM_API void CopyExternalVariableDefs(URigVM* InVM);
+
+	RIGVM_API static bool IsSerializingToDisk(const FArchive& Ar);
 
 	RIGVM_API virtual void PostLoad() override;
 
@@ -752,11 +755,7 @@ public:
 	void ClearExternalVariables() {}
 
 	// Clears the external variables of the VM
-	void ClearExternalVariables(FRigVMExtendedExecuteContext& Context)
-	{
-		ExternalVariables.Reset();
-		Context.ExternalVariableRuntimeData.Reset();
-	}
+	RIGVM_API void ClearExternalVariables(FRigVMExtendedExecuteContext& Context);
 
 	// Returns the external variables of the VM (without variable memory, as it is stored in Context)
 	const TArray<FRigVMExternalVariableDef>& GetExternalVariableDefs() const
@@ -797,16 +796,7 @@ public:
 	FRigVMOperand AddExternalVariable(const FRigVMExternalVariable& InExternalVariable) { return FRigVMOperand(); }
 
 	// Sets the external variables without the instance data
-	void SetExternalVariableDefs(const TArray<FRigVMExternalVariable>& InExternalVariables)
-	{
-		ExternalVariables.Reset(InExternalVariables.Num());
-		for (const FRigVMExternalVariableDef& ExternalVariable : InExternalVariables)
-		{
-			ExternalVariables.Add(ExternalVariable);
-		}
-
-		RefreshExternalPropertyPaths();
-	}
+	RIGVM_API void SetExternalVariableDefs(const TArray<FRigVMExternalVariable>& InExternalVariables);
 
 	// Sets the external variables instance data required for execution
 	void SetExternalVariablesInstanceData(FRigVMExtendedExecuteContext& Context, const TArray<FRigVMExternalVariable>& InExternalVariables, bool bAllowNullMemory = false)
@@ -880,6 +870,7 @@ public:
 
 	// Creates a localized registry for this VM
 	RIGVM_API bool CreateLocalizedRegistryIfRequired(bool bForce = false);
+	RIGVM_API bool CreateLocalizedRegistryIfRequired_NoLock(bool bForce = false);
 
 private:
 	//  The local registry to use by the VM
@@ -1130,6 +1121,8 @@ private:
 
 	mutable TArray<FName> EntryNames;
 
+	bool bIsSerializing = false;
+
 protected:
 	RIGVM_API const URigVMHost* GetHostCDO() const;
 
diff --git a/Engine/Plugins/Runtime/RigVM/Source/RigVM/Public/RigVMCore/RigVMTemplate.h b/Engine/Plugins/Runtime/RigVM/Source/RigVM/Public/RigVMCore/RigVMTemplate.h
--- a/Engine/Plugins/Runtime/RigVM/Source/RigVM/Public/RigVMCore/RigVMTemplate.h
+++ b/Engine/Plugins/Runtime/RigVM/Source/RigVM/Public/RigVMCore/RigVMTemplate.h
@@ -62,7 +62,7 @@ struct FRigVMTemplateArgumentType
 		CPPTypeObject = RigVMTypeUtils::GetWildCardCPPTypeObject();
 	}
 
-	RIGVM_API FRigVMTemplateArgumentType(const FName& InCPPType, UObject* InCPPTypeObject = nullptr);
+	RIGVM_API FRigVMTemplateArgumentType(const FName& InCPPType, UObject* InCPPTypeObject = nullptr, bool bIsSerializing = false);
 	
 	FRigVMTemplateArgumentType(UClass* InClass, RigVMTypeUtils::EClassArgType InClassArgType = RigVMTypeUtils::EClassArgType::AsObject)
 	: CPPType(*RigVMTypeUtils::CPPTypeFromObject(InClass, InClassArgType))
diff --git a/Engine/Plugins/Runtime/RigVM/Source/RigVM/Public/RigVMHost.h b/Engine/Plugins/Runtime/RigVM/Source/RigVM/Public/RigVMHost.h
--- a/Engine/Plugins/Runtime/RigVM/Source/RigVM/Public/RigVMHost.h
+++ b/Engine/Plugins/Runtime/RigVM/Source/RigVM/Public/RigVMHost.h
@@ -597,6 +597,7 @@ protected:
 	friend class URigVMBlueprintGeneratedClass;
 	friend class FRigVMCompileSettingsDetails;
 	friend class FRigVMTrace;
+	friend class URigVM;
 };
 
 class FRigVMBracketScope
diff --git a/Engine/Plugins/Runtime/RigVM/Source/RigVM/Public/RigVMTypeUtils.h b/Engine/Plugins/Runtime/RigVM/Source/RigVM/Public/RigVMTypeUtils.h
--- a/Engine/Plugins/Runtime/RigVM/Source/RigVM/Public/RigVMTypeUtils.h
+++ b/Engine/Plugins/Runtime/RigVM/Source/RigVM/Public/RigVMTypeUtils.h
@@ -260,7 +260,9 @@ namespace RigVMTypeUtils
 		return CPPType;
 	}
 
-	RIGVM_API FString PostProcessCPPType(const FString& InCPPType, UObject* InCPPTypeObject = nullptr, const FRigVMUserDefinedTypeResolver* InResolvalInfo = nullptr);
+	RIGVM_API UObject* GetCPPTypeObjectFromProperty(const FProperty* InProperty);
+
+	RIGVM_API FString PostProcessCPPType(const FString& InCPPType, UObject* InCPPTypeObject = nullptr, const FRigVMUserDefinedTypeResolver* InResolvalInfo = nullptr, bool bIsSerializing = false);
 
 	RIGVM_API bool FixCPPTypeAndObject(FString& InOutCPPType, TObjectPtr<UObject>& InOutCPPTypeObject, const FRigVMUserDefinedTypeResolver* InResolvalInfo = nullptr);
 
