RigVM: Always serialize localized registry

#rnx

full cook preflight: https://horde.devtools.epicgames.com/job/68ee4df0f902144c701382ed

diff --git a/Engine/Plugins/Animation/ControlRig/Source/ControlRig/Private/ControlRigBlueprintGeneratedClass.cpp b/Engine/Plugins/Animation/ControlRig/Source/ControlRig/Private/ControlRigBlueprintGeneratedClass.cpp
--- a/Engine/Plugins/Animation/ControlRig/Source/ControlRig/Private/ControlRigBlueprintGeneratedClass.cpp
+++ b/Engine/Plugins/Animation/ControlRig/Source/ControlRig/Private/ControlRigBlueprintGeneratedClass.cpp
@@ -39,10 +39,11 @@ void UControlRigBlueprintGeneratedClass::Serialize(FArchive& Ar)
 
 	if (UControlRig* CDO = Cast<UControlRig>(GetDefaultObject(true)))
 	{
-		if (Ar.IsSaving() && CDO->VM)
+		const bool bIsSerializingToDisk = URigVM::IsSerializingToDisk(Ar);
+		if (Ar.IsSaving() && bIsSerializingToDisk && CDO->VM)
 		{
-			VM->CreateLocalizedRegistryIfRequired();
 			VM->CopyDataForSerialization(CDO->VM);
+			VM->CreateLocalizedRegistryIfRequired();
 		}
 	}
 	
diff --git a/Engine/Plugins/Runtime/RigVM/Source/RigVM/Private/RigVMBlueprintGeneratedClass.cpp b/Engine/Plugins/Runtime/RigVM/Source/RigVM/Private/RigVMBlueprintGeneratedClass.cpp
--- a/Engine/Plugins/Runtime/RigVM/Source/RigVM/Private/RigVMBlueprintGeneratedClass.cpp
+++ b/Engine/Plugins/Runtime/RigVM/Source/RigVM/Private/RigVMBlueprintGeneratedClass.cpp
@@ -72,12 +72,13 @@ void URigVMBlueprintGeneratedClass::Serialize(FArchive& Ar)
 
 	if (!IsSkeletonClass(this))
 	{
-		if (const URigVMHost* CDO = Cast<URigVMHost>(GetDefaultObject(true)))
+		if (URigVMHost* CDO = Cast<URigVMHost>(GetDefaultObject(true)))
 		{
-			if (Ar.IsSaving() && Ar.IsPersistent() && CDO->VM)
+			const bool bIsSerializingToDisk = URigVM::IsSerializingToDisk(Ar);
+			if (Ar.IsSaving() && bIsSerializingToDisk && CDO->VM)
 			{
-				VM->CreateLocalizedRegistryIfRequired();
 				VM->CopyDataForSerialization(CDO->VM);
+				VM->CreateLocalizedRegistryIfRequired();
 			}
 		}
 	}
diff --git a/Engine/Plugins/Runtime/RigVM/Source/RigVM/Private/RigVMCore/RigVM.cpp b/Engine/Plugins/Runtime/RigVM/Source/RigVM/Private/RigVMCore/RigVM.cpp
--- a/Engine/Plugins/Runtime/RigVM/Source/RigVM/Private/RigVMCore/RigVM.cpp
+++ b/Engine/Plugins/Runtime/RigVM/Source/RigVM/Private/RigVMCore/RigVM.cpp
@@ -170,6 +170,27 @@ void URigVM::Save(FArchive& Ar)
 	Ar << DefaultWorkMemoryStorage;
 	UE_RIGVM_ARCHIVETRACE_ENTRY(Ar, TEXT("DefaultWorkMemoryStorage"));
 
+	TArray<const FRigVMFunction*> TempFunctionsStorage;
+	TArray<FName> TempFunctionNamesStorage; 
+	TArray<const FRigVMDispatchFactory*> TempFactoriesStorage;
+	bool bUsingTemporaryLocalizedRegistry = false;
+
+	UE::TScopeLock CreateLocalizedRegistryScopeLock(CreateLocalizedRegistryMutex);
+
+	// if we are serializing to disk we should create the localized registry temporarily
+	if (bIsSerializingToDisk && !LocalizedRegistry.IsValid())
+	{
+		TempFunctionsStorage = FunctionsStorage;
+		TempFunctionNamesStorage = FunctionNamesStorage;
+		TempFactoriesStorage = FactoriesStorage;
+
+		// we can call the no-lock flavor here since we've established the scope lock above
+		if (CreateLocalizedRegistryIfRequired_NoLock(true))
+		{
+			bUsingTemporaryLocalizedRegistry = true;
+		}
+	}
+
 	if(bIsSerializingToDisk && LocalizedRegistry.IsValid())
 	{
 		FRigVMMemoryStorageStruct EmptyDebugMemoryStruct;
@@ -203,6 +224,16 @@ void URigVM::Save(FArchive& Ar)
 		Ar << bStoredLocalizedRegistry;
 		UE_RIGVM_ARCHIVETRACE_ENTRY(Ar, TEXT("LocalizedRegistry"));
 	}
+
+	if (bUsingTemporaryLocalizedRegistry)
+	{
+		// remove the localized registry
+		LocalizedRegistry.Reset();
+		// swap the function / factory storage back
+		FunctionsStorage = TempFunctionsStorage;
+		FunctionNamesStorage = TempFunctionNamesStorage;
+		FactoriesStorage = TempFactoriesStorage;
+	}
 }
 
 void URigVM::Load(FArchive& Ar)
diff --git a/Engine/Plugins/Runtime/RigVM/Source/RigVM/Public/RigVMObjectVersion.h b/Engine/Plugins/Runtime/RigVM/Source/RigVM/Public/RigVMObjectVersion.h
--- a/Engine/Plugins/Runtime/RigVM/Source/RigVM/Public/RigVMObjectVersion.h
+++ b/Engine/Plugins/Runtime/RigVM/Source/RigVM/Public/RigVMObjectVersion.h
@@ -84,4 +84,17 @@ private:
 	FRigVMObjectVersion() {}
 };
 
+// This is used to version up cooked assets independent of the object version. 
+namespace ERigVMCookVersion
+{
+	enum Type : int
+	{
+		// even when the localized registry is turned off,
+		// make sure to always serialize it so that we can decide
+		// in code if we want to utilize it.
+		AlwaysSerializeLocalizedRegistry,
+		
+		LatestVersion
+	};
+}
 #undef UE_API
diff --git a/Engine/Plugins/Runtime/RigVM/Source/RigVMDeveloper/Private/RigVMBlueprintLegacy.cpp b/Engine/Plugins/Runtime/RigVM/Source/RigVMDeveloper/Private/RigVMBlueprintLegacy.cpp
--- a/Engine/Plugins/Runtime/RigVM/Source/RigVMDeveloper/Private/RigVMBlueprintLegacy.cpp
+++ b/Engine/Plugins/Runtime/RigVM/Source/RigVMDeveloper/Private/RigVMBlueprintLegacy.cpp
@@ -444,6 +444,23 @@ const TArray<FString>& URigVMBlueprint::GetRequiredPlugins(bool bRefresh) const
 	return RequiredPlugins;
 }
 
+#if WITH_EDITORONLY_DATA
+void URigVMBlueprint::AppendToClassSchema(FAppendToClassSchemaContext& Context)
+{
+	Super::AppendToClassSchema(Context);
+
+	// we append the object version to the hash so that we invalidate
+	// cooked assets based on this.
+	int32 ObjectVersion = (int32)FRigVMObjectVersion::LatestVersion;
+	Context.Update(&ObjectVersion, sizeof(ObjectVersion));
+
+	// we also append the cook version to the hash so that we invalidate
+	// cooked assets based on the enum independent from the object version.
+	int32 CookVersion = (int32)ERigVMCookVersion::LatestVersion;
+	Context.Update(&CookVersion, sizeof(CookVersion));
+}
+#endif
+
 void URigVMBlueprint::PostEditChangeChainProperty(struct FPropertyChangedChainEvent& PropertyChangedEvent)
 {
 	UBlueprint::PostEditChangeChainProperty(PropertyChangedEvent);
diff --git a/Engine/Plugins/Runtime/RigVM/Source/RigVMDeveloper/Public/RigVMBlueprintLegacy.h b/Engine/Plugins/Runtime/RigVM/Source/RigVMDeveloper/Public/RigVMBlueprintLegacy.h
--- a/Engine/Plugins/Runtime/RigVM/Source/RigVMDeveloper/Public/RigVMBlueprintLegacy.h
+++ b/Engine/Plugins/Runtime/RigVM/Source/RigVMDeveloper/Public/RigVMBlueprintLegacy.h
@@ -110,6 +110,10 @@ public:
 	virtual TArray<struct FEditedDocumentInfo>& GetLastEditedDocuments() override { return LastEditedDocuments;}
 	virtual const TArray<FString>& GetRequiredPlugins(bool bRefresh = true) const override;
 
+#if WITH_EDITORONLY_DATA
+	static void AppendToClassSchema(FAppendToClassSchemaContext& Context);
+#endif
+
 protected:
 	virtual UClass* GetRigVMGeneratedClass() override { return GeneratedClass; }
 	virtual TArray<FRigVMExternalVariable> GetExternalVariables(bool bFallbackToBlueprint) const override;
