RigVM: disable cooking localized registry by default to address regressions

#rnx

diff --git a/Engine/Plugins/Runtime/RigVM/Source/RigVM/Private/RigVMCore/RigVM.cpp b/Engine/Plugins/Runtime/RigVM/Source/RigVM/Private/RigVMCore/RigVM.cpp
--- a/Engine/Plugins/Runtime/RigVM/Source/RigVM/Private/RigVMCore/RigVM.cpp
+++ b/Engine/Plugins/Runtime/RigVM/Source/RigVM/Private/RigVMCore/RigVM.cpp
@@ -177,20 +177,6 @@ void URigVM::Save(FArchive& Ar)
 
 	UE::TScopeLock CreateLocalizedRegistryScopeLock(CreateLocalizedRegistryMutex);
 
-	// if we are serializing to disk we should create the localized registry temporarily
-	if (bIsSerializingToDisk && !LocalizedRegistry.IsValid())
-	{
-		TempFunctionsStorage = FunctionsStorage;
-		TempFunctionNamesStorage = FunctionNamesStorage;
-		TempFactoriesStorage = FactoriesStorage;
-
-		// we can call the no-lock flavor here since we've established the scope lock above
-		if (CreateLocalizedRegistryIfRequired_NoLock(true))
-		{
-			bUsingTemporaryLocalizedRegistry = true;
-		}
-	}
-
 	if(bIsSerializingToDisk && LocalizedRegistry.IsValid())
 	{
 		FRigVMMemoryStorageStruct EmptyDebugMemoryStruct;
diff --git a/Engine/Plugins/Runtime/RigVM/Source/RigVM/Public/RigVMObjectVersion.h b/Engine/Plugins/Runtime/RigVM/Source/RigVM/Public/RigVMObjectVersion.h
--- a/Engine/Plugins/Runtime/RigVM/Source/RigVM/Public/RigVMObjectVersion.h
+++ b/Engine/Plugins/Runtime/RigVM/Source/RigVM/Public/RigVMObjectVersion.h
@@ -89,11 +89,6 @@ namespace ERigVMCookVersion
 {
 	enum Type : int
 	{
-		// even when the localized registry is turned off,
-		// make sure to always serialize it so that we can decide
-		// in code if we want to utilize it.
-		AlwaysSerializeLocalizedRegistry,
-		
 		LatestVersion
 	};
 }
