RigVM: Avoid using the entry / return pin hash for collapse nodes

#rnx

diff --git a/Engine/Plugins/Runtime/RigVM/Source/RigVMDeveloper/Private/RigVMCompiler/RigVMCompiler.cpp b/Engine/Plugins/Runtime/RigVM/Source/RigVMDeveloper/Private/RigVMCompiler/RigVMCompiler.cpp
--- a/Engine/Plugins/Runtime/RigVM/Source/RigVMDeveloper/Private/RigVMCompiler/RigVMCompiler.cpp
+++ b/Engine/Plugins/Runtime/RigVM/Source/RigVMDeveloper/Private/RigVMCompiler/RigVMCompiler.cpp
@@ -3809,7 +3809,30 @@ FString URigVMCompiler::GetPinHashImpl(const URigVMPin* InPin, const FRigVMVarEx
 			}
 			else if(Node->IsA<URigVMFunctionInterfaceNode>())
 			{
-				const FString FullPath = InPinProxy.GetCallstack().GetCallPath(true);
+				FString FullPath = InPinProxy.GetCallstack().GetCallPath(true);
+
+				// if this is a pin on an entry / return node within a collapse node
+				if (const URigVMGraph* Graph = Node->GetGraph())
+				{
+					if (const URigVMCollapseNode* CollapseNode = Graph->GetTypedOuter<URigVMCollapseNode>())
+					{
+						// if this collapse node is not on a function library
+						if (Cast<URigVMFunctionLibrary>(CollapseNode->GetGraph()) == nullptr)
+						{
+							static const FString EntryPrefix = FString::Printf(TEXT("|%s."), FRigVMGraphFunctionData::EntryString); 
+							static const FString ReturnPrefix = FString::Printf(TEXT("|%s."), FRigVMGraphFunctionData::ReturnString);
+							static const FString CollapseEntryPrefix = FString::Printf(TEXT("|Collapse%s."), FRigVMGraphFunctionData::EntryString); 
+							static const FString CollapseReturnPrefix = FString::Printf(TEXT("|Collapse%s."), FRigVMGraphFunctionData::ReturnString);
+
+							// adapt the path to use CollapseEntry instead of Entry
+							if (FullPath.ReplaceInline(*EntryPrefix, *CollapseEntryPrefix) == 0)
+							{
+								// adapt the path to use CollapseReturn instead of Return
+								(void)FullPath.ReplaceInline(*ReturnPrefix, *CollapseReturnPrefix);
+							}
+						}
+					}
+				}
 				return FString::Printf(TEXT("%s%s%s"), *Prefix, *FullPath, *Suffix);
 			}
 		}
