RigVM: Avoid maintaining template controller / graph on rigvm asset

#rnx

diff --git a/Engine/Plugins/Runtime/RigVM/Source/RigVMDeveloper/Private/RigVMAsset.cpp b/Engine/Plugins/Runtime/RigVM/Source/RigVMDeveloper/Private/RigVMAsset.cpp
--- a/Engine/Plugins/Runtime/RigVM/Source/RigVMDeveloper/Private/RigVMAsset.cpp
+++ b/Engine/Plugins/Runtime/RigVM/Source/RigVMDeveloper/Private/RigVMAsset.cpp
@@ -2628,20 +2628,18 @@ void IRigVMAssetInterface::AddVariableSearchMetaDataInfo(const FName InVariableN
 URigVMGraph* IRigVMAssetInterface::GetTemplateModel(bool bIsFunctionLibrary)
 {
 #if WITH_EDITORONLY_DATA
-	if (TemplateModel == nullptr)
+	URigVMGraph* Model = nullptr;
+	if (bIsFunctionLibrary)
 	{
-		if (bIsFunctionLibrary)
-		{
-			TemplateModel = NewObject<URigVMFunctionLibrary>(GetObject(), TEXT("TemplateFunctionLibrary"));
-		}
-		else
-		{
-			TemplateModel = NewObject<URigVMGraph>(GetObject(), TEXT("TemplateModel"));
-		}
-		TemplateModel->SetFlags(RF_Transient);
-		TemplateModel->SetExecuteContextStruct(GetRigVMClient()->GetDefaultExecuteContextStruct());
+		Model = NewObject<URigVMFunctionLibrary>(GetObject(), TEXT("TemplateFunctionLibrary"));
 	}
-	return TemplateModel;
+	else
+	{
+		Model = NewObject<URigVMGraph>(GetObject(), TEXT("TemplateModel"));
+	}
+	Model->SetFlags(RF_Transient);
+	Model->SetExecuteContextStruct(GetRigVMClient()->GetDefaultExecuteContextStruct());
+	return Model;
 #else
 	return nullptr;
 #endif
@@ -2650,15 +2648,20 @@ URigVMGraph* IRigVMAssetInterface::GetTemplateModel(bool bIsFunctionLibrary)
 URigVMController* IRigVMAssetInterface::GetTemplateController(bool bIsFunctionLibrary)
 {
 #if WITH_EDITORONLY_DATA
-	if (TemplateController == nullptr)
+	URigVMController* Controller = nullptr; 
+	if (bIsFunctionLibrary)
+	{
+		Controller = NewObject<URigVMController>(GetObject(), TEXT("FunctionLibraryTemplateController"));
+	}
+	else
 	{
-		TemplateController = NewObject<URigVMController>(GetObject(), TEXT("TemplateController"));
-		TemplateController->SetGraph(GetTemplateModel(bIsFunctionLibrary));
-		TemplateController->EnableReporting(false);
-		TemplateController->SetFlags(RF_Transient);
-		TemplateController->SetSchemaClass(GetRigVMClient()->GetDefaultSchemaClass());
+		Controller = NewObject<URigVMController>(GetObject(), TEXT("TemplateController"));
 	}
-	return TemplateController;
+	Controller->EnableReporting(false);
+	Controller->SetFlags(RF_Transient);
+	Controller->SetSchemaClass(GetRigVMClient()->GetDefaultSchemaClass());
+	Controller->SetGraph(GetTemplateModel(bIsFunctionLibrary));
+	return Controller;
 #else
 	return nullptr;
 #endif
diff --git a/Engine/Plugins/Runtime/RigVM/Source/RigVMDeveloper/Public/RigVMAsset.h b/Engine/Plugins/Runtime/RigVM/Source/RigVMDeveloper/Public/RigVMAsset.h
--- a/Engine/Plugins/Runtime/RigVM/Source/RigVMDeveloper/Public/RigVMAsset.h
+++ b/Engine/Plugins/Runtime/RigVM/Source/RigVMDeveloper/Public/RigVMAsset.h
@@ -558,8 +558,6 @@ protected:
 
 	
 
-	TObjectPtr<URigVMGraph> TemplateModel;
-	TObjectPtr<URigVMController> TemplateController;
 	mutable TArray<UObject::FAssetRegistryTag> CachedAssetTags;
 	TArray<TPair<TWeakObjectPtr<URigVMNode>,int32>> LastPreviewHereNodes;
 	
