RigVM: re-enable always cooking localized registry

#rnx

full cook is here https://horde.devtools.epicgames.com/job/68efab27c275c15401683409

diff --git a/Engine/Plugins/Runtime/RigVM/Source/RigVM/Private/RigVMCore/RigVM.cpp b/Engine/Plugins/Runtime/RigVM/Source/RigVM/Private/RigVMCore/RigVM.cpp
--- a/Engine/Plugins/Runtime/RigVM/Source/RigVM/Private/RigVMCore/RigVM.cpp
+++ b/Engine/Plugins/Runtime/RigVM/Source/RigVM/Private/RigVMCore/RigVM.cpp
@@ -177,6 +177,20 @@ void URigVM::Save(FArchive& Ar)
 
 	UE::TScopeLock CreateLocalizedRegistryScopeLock(CreateLocalizedRegistryMutex);
 
+	// if we are serializing to disk we should create the localized registry temporarily
+	if (bIsSerializingToDisk && !LocalizedRegistry.IsValid())
+	{
+		TempFunctionsStorage = FunctionsStorage;
+		TempFunctionNamesStorage = FunctionNamesStorage;
+		TempFactoriesStorage = FactoriesStorage;
+
+		// we can call the no-lock flavor here since we've established the scope lock above
+		if (CreateLocalizedRegistryIfRequired_NoLock(true))
+		{
+			bUsingTemporaryLocalizedRegistry = true;
+		}
+	}
+
 	if(bIsSerializingToDisk && LocalizedRegistry.IsValid())
 	{
 		FRigVMMemoryStorageStruct EmptyDebugMemoryStruct;
@@ -2437,6 +2451,9 @@ bool URigVM::CreateLocalizedRegistryIfRequired_NoLock(bool bForce)
 	const FRigVMByteCode& ByteCode = GetByteCode();
 
 	// add all required types to it
+	{
+	FRigVMRegistryReadLock ReadLock;
+
 	for (int32 InstructionIndex = 0; InstructionIndex < Instructions.Num(); InstructionIndex++)
 	{
 		FRigVMOperandArray Operands = ByteCode.GetOperandsForOp(Instructions[InstructionIndex]);
@@ -2531,6 +2548,11 @@ bool URigVM::CreateLocalizedRegistryIfRequired_NoLock(bool bForce)
 			// since this may be happening during GC or serialization
 			const FName CPPTypeName = *RigVMTypeUtils::GetCPPTypeFromProperty(Property);
 			UObject* CPPTypeObject = RigVMTypeUtils::GetCPPTypeObjectFromProperty(Property);
+
+			if (CPPTypeObject && ReadLock->IsSpecificallyAllowedType_NoLock(CPPTypeObject))
+			{
+				TemporaryHandle->AddSpecificallyAllowedType_NoLock(CPPTypeObject);
+			}
 			TemporaryHandle->FindOrAddType_NoLock(FRigVMTemplateArgumentType(CPPTypeName, CPPTypeObject, bIsSerializing));
 		}
 	}
@@ -2538,7 +2560,6 @@ bool URigVM::CreateLocalizedRegistryIfRequired_NoLock(bool bForce)
 	// for dispatch functions we need to register all of the argument types for the given permutations
 	// otherwise we won't be able to retrieve the permutation from string again later.
 	{
-		FRigVMRegistryReadLock ReadLock;
 		for(const FRigVMFunction* Function : FunctionsStorage)
 		{
 			if(Function && Function->Factory)
@@ -2553,6 +2574,10 @@ bool URigVM::CreateLocalizedRegistryIfRequired_NoLock(bool bForce)
 						const FRigVMTemplateArgumentType& Type = ReadLock->GetType_NoLock(ArgumentTypeIndex);
 						if (Type.IsValid())
 						{
+							if (Type.CPPTypeObject && ReadLock->IsSpecificallyAllowedType_NoLock(Type.CPPTypeObject))
+							{
+								TemporaryHandle->AddSpecificallyAllowedType_NoLock(Type.CPPTypeObject);
+							}
 							TemporaryRegistry->FindOrAddType_NoLock(Type);
 						}
 						else
@@ -2600,7 +2625,6 @@ bool URigVM::CreateLocalizedRegistryIfRequired_NoLock(bool bForce)
 		// flatten the argument infos to the permutations used only
 		TArray<FRigVMTemplateArgumentInfo> TemporaryArgumentInfos;
 		{
-			FRigVMRegistryReadLock ReadLock;
 			const TArray<FRigVMTemplateArgumentInfo> FlattenedArgumentInfos = Factory->GetTemplate_NoLock(ReadLock)->GetFlattenedArgumentInfosForFunctions_NoLock(OriginalFunctionsForFactory, ReadLock);
 			TemporaryArgumentInfos.Reserve(FlattenedArgumentInfos.Num());
 
@@ -2619,6 +2643,10 @@ bool URigVM::CreateLocalizedRegistryIfRequired_NoLock(bool bForce)
 						return ResetLocalizedRegistry();
 					}
 
+					if (OriginalType.CPPTypeObject && ReadLock->IsSpecificallyAllowedType_NoLock(OriginalType.CPPTypeObject))
+					{
+						TemporaryRegistry->AddSpecificallyAllowedType_NoLock(OriginalType.CPPTypeObject);
+					}
 					TemporaryTypeIndices.Add(TemporaryRegistry->FindOrAddType_NoLock(OriginalType));
 					if (TemporaryTypeIndices.Last() == INDEX_NONE)
 					{
@@ -2686,6 +2714,7 @@ bool URigVM::CreateLocalizedRegistryIfRequired_NoLock(bool bForce)
 			}
 		}
 	}
+	}
 
 	// register all functions
 	for(const FRigVMFunction* Function : FunctionsStorage)
diff --git a/Engine/Plugins/Runtime/RigVM/Source/RigVM/Private/RigVMCore/RigVMRegistry.cpp b/Engine/Plugins/Runtime/RigVM/Source/RigVM/Private/RigVMCore/RigVMRegistry.cpp
--- a/Engine/Plugins/Runtime/RigVM/Source/RigVM/Private/RigVMCore/RigVMRegistry.cpp
+++ b/Engine/Plugins/Runtime/RigVM/Source/RigVM/Private/RigVMCore/RigVMRegistry.cpp
@@ -2018,6 +2018,31 @@ TRigVMTypeIndex FRigVMRegistry_NoLock::GetBaseTypeFromArrayTypeIndex_NoLock(TRig
 	return INDEX_NONE;
 }
 
+bool FRigVMRegistry_NoLock::IsSpecificallyAllowedType_NoLock(const UObject* InObject) const
+{
+	if (const UClass* Class = Cast<UClass>(InObject))
+	{
+		return AllowedClasses.Contains(Class);
+	}
+	if (const UScriptStruct* Struct = Cast<UScriptStruct>(InObject))
+	{
+		return AllowedStructs.Contains(Struct);
+	}
+	return false;
+}
+
+void FRigVMRegistry_NoLock::AddSpecificallyAllowedType_NoLock(const UObject* InObject)
+{
+	if (const UClass* Class = Cast<UClass>(InObject))
+	{
+		AllowedClasses.Add(Class);
+	}
+	if (const UScriptStruct* Struct = Cast<UScriptStruct>(InObject))
+	{
+		AllowedStructs.Add(Struct);
+	}
+}
+
 bool FRigVMRegistry_NoLock::IsAllowedType_NoLock(const FProperty* InProperty) const
 {
 	if(InProperty->IsA<FBoolProperty>() ||
diff --git a/Engine/Plugins/Runtime/RigVM/Source/RigVM/Public/RigVMCore/RigVMRegistry.h b/Engine/Plugins/Runtime/RigVM/Source/RigVM/Public/RigVMCore/RigVMRegistry.h
--- a/Engine/Plugins/Runtime/RigVM/Source/RigVM/Public/RigVMCore/RigVMRegistry.h
+++ b/Engine/Plugins/Runtime/RigVM/Source/RigVM/Public/RigVMCore/RigVMRegistry.h
@@ -287,6 +287,12 @@ public:
 	// Returns the type index of the element matching the given array type index
 	UE_API virtual TRigVMTypeIndex GetBaseTypeFromArrayTypeIndex_NoLock(TRigVMTypeIndex InTypeIndex) const;
 
+	// Returns true if the given type is specifically allowed 
+	UE_API virtual bool IsSpecificallyAllowedType_NoLock(const UObject* InObject) const;
+
+	// Sets a given type as specifically allowed 
+	UE_API virtual void AddSpecificallyAllowedType_NoLock(const UObject* InObject);
+
 	// Returns the function given its name (or nullptr)
 	UE_API virtual const FRigVMFunction* FindFunction_NoLock(const TCHAR* InName, const FRigVMUserDefinedTypeResolver& InTypeResolver = FRigVMUserDefinedTypeResolver()) const;
 
@@ -955,6 +961,20 @@ public:
 		return Super::GetBaseTypeFromArrayTypeIndex_NoLock(InTypeIndex);
 	}
 
+	// Returns true if the given type is specifically allowed 
+	bool IsSpecificallyAllowedType(const UObject* InObject) const
+	{
+		FConditionalReadScopeLock _(*this);
+		return Super::IsSpecificallyAllowedType_NoLock(InObject);
+	}
+
+	// Sets a given type as specifically allowed 
+	void AddSpecificallyAllowedType(const UObject* InObject)
+	{
+		FConditionalWriteScopeLock _(*this);
+		Super::AddSpecificallyAllowedType_NoLock(InObject);
+	}
+	
 	// Returns the function given its name (or nullptr)
 	const FRigVMFunction* FindFunction(const TCHAR* InName, const FRigVMUserDefinedTypeResolver& InTypeResolver = FRigVMUserDefinedTypeResolver()) const
 	{
@@ -1370,6 +1390,20 @@ protected:
 		return Super::GetBaseTypeFromArrayTypeIndex_NoLock(InTypeIndex);
 	}
 
+	// Returns true if the given type is specifically allowed 
+	// Note: Only call NoLock methods on the FRigVMRegistry_NoLock
+	virtual bool IsSpecificallyAllowedType_NoLock(const UObject* InObject) const override
+	{
+		return Super::IsSpecificallyAllowedType_NoLock(InObject);
+	}
+
+	// Sets a given type as specifically allowed 
+	// Note: Only call NoLock methods on the FRigVMRegistry_NoLock
+	virtual void AddSpecificallyAllowedType_NoLock(const UObject* InObject) override
+	{
+		Super::AddSpecificallyAllowedType_NoLock(InObject);
+	}
+
 	// Returns the function given its name (or nullptr)
 	// Note: Only call NoLock methods on the FRigVMRegistry_NoLock
 	virtual const FRigVMFunction* FindFunction_NoLock(const TCHAR* InName, const FRigVMUserDefinedTypeResolver& InTypeResolver = FRigVMUserDefinedTypeResolver()) const override
diff --git a/Engine/Plugins/Runtime/RigVM/Source/RigVM/Public/RigVMObjectVersion.h b/Engine/Plugins/Runtime/RigVM/Source/RigVM/Public/RigVMObjectVersion.h
--- a/Engine/Plugins/Runtime/RigVM/Source/RigVM/Public/RigVMObjectVersion.h
+++ b/Engine/Plugins/Runtime/RigVM/Source/RigVM/Public/RigVMObjectVersion.h
@@ -89,6 +89,11 @@ namespace ERigVMCookVersion
 {
 	enum Type : int
 	{
+		// even when the localized registry is turned off,
+		// make sure to always serialize it so that we can decide
+		// in code if we want to utilize it.
+		AlwaysSerializeLocalizedRegistry,
+		
 		LatestVersion
 	};
 }
