RigVM: Localized registry skipping during load should use relative offset

#rnx

Full cook preflight https://horde.devtools.epicgames.com/job/68f77ca52657f7e47ef74018

diff --git a/Engine/Plugins/Runtime/RigVM/Source/RigVM/Private/RigVMCore/RigVM.cpp b/Engine/Plugins/Runtime/RigVM/Source/RigVM/Private/RigVMCore/RigVM.cpp
--- a/Engine/Plugins/Runtime/RigVM/Source/RigVM/Private/RigVMCore/RigVM.cpp
+++ b/Engine/Plugins/Runtime/RigVM/Source/RigVM/Private/RigVMCore/RigVM.cpp
@@ -202,17 +202,19 @@ void URigVM::Save(FArchive& Ar)
 
 		const int64 ArchivePosBeforeSerializedRegistry = Ar.Tell();
 
-		int64 ArchivePosAfterSerializedRegistry = 0;
-		Ar << ArchivePosAfterSerializedRegistry;
+		int64 ArchiveOffsetForSerializedRegistry = 0;
+		Ar << ArchiveOffsetForSerializedRegistry;
 
 		LocalizedRegistry->Serialize_NoLock(Ar);
 
 		// store the archive positions after the registry so we
 		// can skip it during load
-		ArchivePosAfterSerializedRegistry = Ar.Tell();
+		const int64 ArchivePosAfterSerializedRegistry = Ar.Tell();
+		ArchiveOffsetForSerializedRegistry = ArchivePosAfterSerializedRegistry - ArchivePosBeforeSerializedRegistry;
 		Ar.Seek(ArchivePosBeforeSerializedRegistry);
-		Ar << ArchivePosAfterSerializedRegistry;
+		Ar << ArchiveOffsetForSerializedRegistry;
 		Ar.Seek(ArchivePosAfterSerializedRegistry);
+
 		UE_RIGVM_ARCHIVETRACE_ENTRY(Ar, TEXT("LocalizedRegistry"));
 	}
 	else
@@ -358,14 +360,22 @@ void URigVM::Load(FArchive& Ar)
 		LocalizedRegistry.Reset();
 		if(bStoredLocalizedRegistry)
 		{
-			int64 ArchivePosAfterSerializedRegistry = 0;
-			Ar << ArchivePosAfterSerializedRegistry;
+			const int64 ArchivePosBeforeSerializedRegistry = Ar.Tell();
+
+			int64 ArchiveOffsetForSerializedRegistry = 0;
+			Ar << ArchiveOffsetForSerializedRegistry;
 
 			if(CVarRigVMEnableLocalizedRegistry.GetValueOnAnyThread() == false)
 			{
-				// disabling warning since it is causing cooking errors.
-				// UE_LOG(LogRigVM, Warning, TEXT("VM '%s' was serialized with a localized registry, but CVarRigVMEnableLocalizedRegistry is disabled."), *GetPathName());
-				Ar.Seek(ArchivePosAfterSerializedRegistry);
+				if (Ar.CustomVer(FRigVMObjectVersion::GUID) >= FRigVMObjectVersion::LocalizedRegistryWithRelativeSeekOffset)
+				{
+					Ar.Seek(ArchivePosBeforeSerializedRegistry + ArchiveOffsetForSerializedRegistry);
+				}
+				else
+				{
+					// in older versions the offset is considered absolute
+					Ar.Seek(ArchiveOffsetForSerializedRegistry);
+				}
 			}
 			else
 			{
@@ -3337,4 +3347,4 @@ void URigVM::StopProfiling(FRigVMExtendedExecuteContext& Context)
 	TRACE_RIGVM_EVALUATION(Context);
 }
 
-#undef UE_RIGVM_GET_READ_REGISTRY_CONTEXT
\ No newline at end of file
+#undef UE_RIGVM_GET_READ_REGISTRY_CONTEXT
diff --git a/Engine/Plugins/Runtime/RigVM/Source/RigVM/Public/RigVMObjectVersion.h b/Engine/Plugins/Runtime/RigVM/Source/RigVM/Public/RigVMObjectVersion.h
--- a/Engine/Plugins/Runtime/RigVM/Source/RigVM/Public/RigVMObjectVersion.h
+++ b/Engine/Plugins/Runtime/RigVM/Source/RigVM/Public/RigVMObjectVersion.h
@@ -72,6 +72,9 @@ struct FRigVMObjectVersion
 		// The VM stores a local snapshot registry to use in cooked environments instead of the shared global registry
 		LocalizedRegistry,
 
+		// The VM stores a relative seek offset to be able to skip the registry during load
+		LocalizedRegistryWithRelativeSeekOffset,
+
 		// -----<new versions can be added above this line>-------------------------------------------------
 		VersionPlusOne,
 		LatestVersion = VersionPlusOne - 1
