RigVM: Allow to inject a get node for "promote to exposed pin"

#rnx

diff --git a/Engine/Plugins/Runtime/RigVM/Source/RigVMDeveloper/Private/EdGraph/RigVMEdGraphSchema.cpp b/Engine/Plugins/Runtime/RigVM/Source/RigVMDeveloper/Private/EdGraph/RigVMEdGraphSchema.cpp
--- a/Engine/Plugins/Runtime/RigVM/Source/RigVMDeveloper/Private/EdGraph/RigVMEdGraphSchema.cpp
+++ b/Engine/Plugins/Runtime/RigVM/Source/RigVMDeveloper/Private/EdGraph/RigVMEdGraphSchema.cpp
@@ -444,7 +444,7 @@ UEdGraphNode* FRigVMEdGraphSchemaAction_PromoteToVariable::PerformAction(UEdGrap
 FRigVMEdGraphSchemaAction_PromoteToExposedPin::FRigVMEdGraphSchemaAction_PromoteToExposedPin(UEdGraphPin* InEdGraphPin)
 : FEdGraphSchemaAction(	FText(), 
 						LOCTEXT("PromoteToExposedPin", "Promote to exposed pin"),
-						LOCTEXT("PromoteToExposedPin", "Promote to exposed pin"),
+						LOCTEXT("PromoteToExposedPin", "Adds an exposed pin and wires it to this pin. For inputs if you hold shift this will insert a Get node."),
 						1)
 , EdGraphPin(InEdGraphPin)
 {
@@ -483,7 +483,7 @@ UEdGraphNode* FRigVMEdGraphSchemaAction_PromoteToExposedPin::PerformAction(UEdGr
 
 	Controller->OpenUndoBracket(TEXT("Promote to Exposed Pin"));
 
-	const UObject* CPPTypeObject = ModelPin->GetCPPTypeObject();
+	UObject* CPPTypeObject = ModelPin->GetCPPTypeObject();
 	const FString PinName = Controller->AddExposedPin(
 		*ModelPin->GetName(),
 		ModelPin->GetDirection(),
@@ -499,8 +499,21 @@ UEdGraphNode* FRigVMEdGraphSchemaAction_PromoteToExposedPin::PerformAction(UEdGr
 	{
 		if (ModelPin->GetDirection() == ERigVMPinDirection::Input)
 		{
-			Controller->AddLink(Model->GetEntryNode()->FindPin(PinName), ModelPin, true);
-			Result = RigGraph->FindNodeForModelNodeName(Model->GetEntryNode()->GetFName());
+			if (FSlateApplication::Get().GetModifierKeys().IsShiftDown())
+			{
+				URigVMVariableNode* VariableNode = Controller->AddVariableNode(*PinName, ModelPin->GetCPPType(), CPPTypeObject, true, ModelPin->GetDefaultValue(), FDeprecateSlateVector2D(Location), FString(), true, true);
+				if (!VariableNode)
+				{
+					return nullptr;
+				}
+				Controller->AddLink(VariableNode->GetValuePin(), ModelPin, true);
+				Result = RigGraph->FindNodeForModelNodeName(VariableNode->GetFName());
+			}
+			else
+			{
+				Controller->AddLink(Model->GetEntryNode()->FindPin(PinName), ModelPin, true);
+				Result = RigGraph->FindNodeForModelNodeName(Model->GetEntryNode()->GetFName());
+			}
 		}
 		else if(ModelPin->GetDirection() == ERigVMPinDirection::Output)
 		{
