RigVM now preprocess handles using offsets at CDO creation
Added cvars to allo tweak / disable RigVM parallel_for on some platforms

diff --git a/Engine/Plugins/Runtime/RigVM/Source/RigVM/Private/RigVMCore/RigVM.cpp b/Engine/Plugins/Runtime/RigVM/Source/RigVM/Private/RigVMCore/RigVM.cpp
--- a/Engine/Plugins/Runtime/RigVM/Source/RigVM/Private/RigVMCore/RigVM.cpp
+++ b/Engine/Plugins/Runtime/RigVM/Source/RigVM/Private/RigVMCore/RigVM.cpp
@@ -33,6 +33,18 @@ static TAutoConsoleVariable<int32> CVarControlRigDebugAllVMExecutions(
 	ECVF_Default);
 #endif
 
+static TAutoConsoleVariable<int32> CVarControlRigCacheMemoryHandlesBatchSize(
+	TEXT("ControlRig.CacheMemoryHandlesBatchSize"),
+	10,
+	TEXT("If greather than zero, the cache of memory handles will be processed using a parallel_for of the specified batch size, else single thread processing."),
+	ECVF_Default);
+
+static TAutoConsoleVariable<int32> CVarControlRigCacheMemoryHandlesParallelMinSize(
+	TEXT("ControlRig.CacheMemoryHandlesParallelMinSize"),
+	700,
+	TEXT("If greather than zero, the cache of memory handles will be processed using a parallel_for if the amount of handles is greather than the specified number, else single thread processing."),
+	ECVF_Default);
+
 void FRigVMParameter::Serialize(FArchive& Ar)
 {
 	Ar.UsingCustomVersion(FAnimObjectVersion::GUID);
@@ -1237,6 +1249,11 @@ void URigVM::InvalidateCachedMemory_Internal()
 	MemoryHandleCount = 0;
 	ExternalPropertyPaths.Reset();
 	LazyBranches.Reset();
+	PreCachedMemoryHandles.Reset();
+	PreCachedMemoryHandlesMemoryType.Reset();
+	InstructionsIndexesWithExternalMemory.Reset();
+	NumHandlesWithExternalMemory = 0;
+	CurrenntInstructionWithExternalMemory = 0;
 }
 
 void URigVM::InvalidateCachedMemory(FRigVMExtendedExecuteContext& Context)
@@ -1253,6 +1270,7 @@ void URigVM::InstructionOpEval(
 	const TFunctionRef<void(
 		FRigVMExtendedExecuteContext& /*Context*/,
 		FRigVMRegistryHandle& /*RegistryHandle*/,
+		int32 /*InstructionIndex*/,
 		int32 /*InHandleIndex*/,
 		const FRigVMBranchInfoKey& /*InBranchInfoKey*/,
 		const FRigVMOperand& /*InArg*/
@@ -1278,6 +1296,7 @@ void URigVM::InstructionOpEval(
 			InOpFunc(
 				Context,
 				RegistryHandle,
+				InstructionIndex,
 				InHandleBaseIndex++,
 				{ InstructionIndex, ArgIndex, Function->GetArgumentNameForOperandIndex(ArgIndex, Operands.Num(), RegistryHandle) },
 				Operands[ArgIndex]);
@@ -1294,14 +1313,14 @@ void URigVM::InstructionOpEval(
 		case ERigVMOpCode::Decrement:
 		{
 			const FRigVMUnaryOp& Op = ByteCode.GetOpAt<FRigVMUnaryOp>(Instructions[InstructionIndex]);
-			InOpFunc(Context, RegistryHandle, InHandleBaseIndex, {}, Op.Arg);
+			InOpFunc(Context, RegistryHandle, InstructionIndex, InHandleBaseIndex, {}, Op.Arg);
 			break;
 		}
 		case ERigVMOpCode::Copy:
 		{
 			const FRigVMCopyOp& Op = ByteCode.GetOpAt<FRigVMCopyOp>(Instructions[InstructionIndex]);
-			InOpFunc(Context, RegistryHandle, InHandleBaseIndex + 0, {}, Op.Source);
-			InOpFunc(Context, RegistryHandle, InHandleBaseIndex + 1, {}, Op.Target);
+			InOpFunc(Context, RegistryHandle, InstructionIndex, InHandleBaseIndex + 0, {}, Op.Source);
+			InOpFunc(Context, RegistryHandle, InstructionIndex, InHandleBaseIndex + 1, {}, Op.Target);
 			break;
 		}
 		case ERigVMOpCode::Equals:
@@ -1309,11 +1328,11 @@ void URigVM::InstructionOpEval(
 		{
 			const FRigVMComparisonOp& Op = ByteCode.GetOpAt<FRigVMComparisonOp>(Instructions[InstructionIndex]);
 			FRigVMOperand Arg = Op.A;
-			InOpFunc(Context, RegistryHandle, InHandleBaseIndex + 0, {}, Arg);
+			InOpFunc(Context, RegistryHandle, InstructionIndex, InHandleBaseIndex + 0, {}, Arg);
 			Arg = Op.B;
-			InOpFunc(Context, RegistryHandle, InHandleBaseIndex + 1, {}, Arg);
+			InOpFunc(Context, RegistryHandle, InstructionIndex, InHandleBaseIndex + 1, {}, Arg);
 			Arg = Op.Result;
-			InOpFunc(Context, RegistryHandle, InHandleBaseIndex + 2, {}, Arg);
+			InOpFunc(Context, RegistryHandle, InstructionIndex, InHandleBaseIndex + 2, {}, Arg);
 			break;
 		}
 		case ERigVMOpCode::JumpAbsolute:
@@ -1328,14 +1347,14 @@ void URigVM::InstructionOpEval(
 		{
 			const FRigVMJumpIfOp& Op = ByteCode.GetOpAt<FRigVMJumpIfOp>(Instructions[InstructionIndex]);
 			const FRigVMOperand& Arg = Op.Arg;
-			InOpFunc(Context, RegistryHandle, InHandleBaseIndex, {}, Arg);
+			InOpFunc(Context, RegistryHandle, InstructionIndex, InHandleBaseIndex, {}, Arg);
 			break;
 		}
 		case ERigVMOpCode::ChangeType:
 		{
 			const FRigVMChangeTypeOp& Op = ByteCode.GetOpAt<FRigVMChangeTypeOp>(Instructions[InstructionIndex]);
 			const FRigVMOperand& Arg = Op.Arg;
-			InOpFunc(Context, RegistryHandle, InHandleBaseIndex, {}, Arg);
+			InOpFunc(Context, RegistryHandle, InstructionIndex, InHandleBaseIndex, {}, Arg);
 			break;
 		}
 		case ERigVMOpCode::Exit:
@@ -1345,8 +1364,8 @@ void URigVM::InstructionOpEval(
 		case ERigVMOpCode::BeginBlock:
 		{
 			const FRigVMBinaryOp& Op = ByteCode.GetOpAt<FRigVMBinaryOp>(Instructions[InstructionIndex]);
-			InOpFunc(Context, RegistryHandle, InHandleBaseIndex + 0, {}, Op.ArgA);
-			InOpFunc(Context, RegistryHandle, InHandleBaseIndex + 1, {}, Op.ArgB);
+			InOpFunc(Context, RegistryHandle, InstructionIndex, InHandleBaseIndex + 0, {}, Op.ArgA);
+			InOpFunc(Context, RegistryHandle, InstructionIndex, InHandleBaseIndex + 1, {}, Op.ArgB);
 			break;
 		}
 		case ERigVMOpCode::EndBlock:
@@ -1361,21 +1380,21 @@ void URigVM::InstructionOpEval(
 		{
 			const FRigVMJumpToBranchOp& Op = ByteCode.GetOpAt<FRigVMJumpToBranchOp>(Instructions[InstructionIndex]);
 			const FRigVMOperand& Arg = Op.Arg;
-			InOpFunc(Context, RegistryHandle, InHandleBaseIndex, {}, Arg);
+			InOpFunc(Context, RegistryHandle, InstructionIndex, InHandleBaseIndex, {}, Arg);
 			break;
 		}
 		case ERigVMOpCode::RunInstructions:
 		{
 			const FRigVMRunInstructionsOp& Op = ByteCode.GetOpAt<FRigVMRunInstructionsOp>(Instructions[InstructionIndex]);
 			const FRigVMOperand& Arg = Op.Arg;
-			InOpFunc(Context, RegistryHandle, InHandleBaseIndex, {}, Arg);
+			InOpFunc(Context, RegistryHandle, InstructionIndex, InHandleBaseIndex, {}, Arg);
 			break;
 		}
 		case ERigVMOpCode::SetupTraits:
 		{
 			const FRigVMSetupTraitsOp& Op = ByteCode.GetOpAt<FRigVMSetupTraitsOp>(Instructions[InstructionIndex]);
 			const FRigVMOperand& Arg = Op.Arg;
-			InOpFunc(Context, RegistryHandle, InHandleBaseIndex, {}, Arg);
+			InOpFunc(Context, RegistryHandle, InstructionIndex, InHandleBaseIndex, {}, Arg);
 			break;
 		}
 		case ERigVMOpCode::Invalid:
@@ -1390,7 +1409,7 @@ void URigVM::InstructionOpEval(
 
 void URigVM::PrepareMemoryForExecution(FRigVMExtendedExecuteContext& Context)
 {
-	ensureMsgf(Context.ExecutingThreadId == FPlatformTLS::GetCurrentThreadId(), TEXT("RigVM::CacheMemoryHandlesIfRequired from multiple threads (%d and %d)"), Context.ExecutingThreadId, (int32)FPlatformTLS::GetCurrentThreadId());
+	ensureMsgf(Context.ExecutingThreadId == FPlatformTLS::GetCurrentThreadId(), TEXT("RigVM::PrepareMemoryForExecution from multiple threads (%d and %d)"), Context.ExecutingThreadId, (int32)FPlatformTLS::GetCurrentThreadId());
 
 	InvalidateCachedMemory(Context);
 
@@ -1431,53 +1450,126 @@ void URigVM::PrepareMemoryForExecution(FRigVMExtendedExecuteContext& Context)
 
 	// Count how many handles we need and set up the indirection offsets for the handles.
 	int32 HandleCount = 0;
+	int32& NumHandlesWithExternalMemoryRef = NumHandlesWithExternalMemory;
 	FirstHandleForInstruction.Add(0);
 	for (int32 InstructionIndex = 0; InstructionIndex < Instructions.Num(); InstructionIndex++)
 	{
 		InstructionOpEval(Context, RegistryHandle, InstructionIndex, INDEX_NONE,
-			[&HandleCount](FRigVMExtendedExecuteContext&, const FRigVMRegistryHandle&, int32, const FRigVMBranchInfoKey&, const FRigVMOperand&)
+			[&HandleCount, &NumHandlesWithExternalMemoryRef](FRigVMExtendedExecuteContext&, const FRigVMRegistryHandle&, int32, int32, const FRigVMBranchInfoKey&, const FRigVMOperand& InArg)
 			{
 				HandleCount++;
+				if (InArg.GetMemoryType() == ERigVMMemoryType::External)
+				{
+					NumHandlesWithExternalMemoryRef++;
+				}
 			});
 		FirstHandleForInstruction.Add(HandleCount);
 	}
 
 	MemoryHandleCount = HandleCount;
 	Context.CachedMemoryHandles.Reset(MemoryHandleCount);
+	PreCachedMemoryHandles.Reset(MemoryHandleCount);
+	PreCachedMemoryHandlesMemoryType.Reset(MemoryHandleCount);
+	InstructionsIndexesWithExternalMemory.SetNumUninitialized(NumHandlesWithExternalMemory);
 }
 
-void URigVM::CacheMemoryHandlesIfRequired(FRigVMExtendedExecuteContext& Context)
+void URigVM::CacheMemoryHandles(FRigVMExtendedExecuteContext& Context)
 {
-	if (Instructions.Num() == 0)
+	const int32 NumInstructions = Instructions.Num();
+	if (NumInstructions == 0)
 	{
 		return;
 	}
 
-	check((Instructions.Num() + 1) == FirstHandleForInstruction.Num());
-	
+	check((NumInstructions + 1) == FirstHandleForInstruction.Num());
+
 	if (Context.CachedMemoryHandles.Num() == MemoryHandleCount)
 	{
 		return;
 	}
 
 	// Allocate all the space and zero it out to ensure all pages required for it are paged in immediately.
-	Context.CachedMemoryHandles.SetNumUninitialized(MemoryHandleCount);
+	PreCachedMemoryHandles.SetNumUninitialized(MemoryHandleCount);
+	PreCachedMemoryHandlesMemoryType.SetNumUninitialized(MemoryHandleCount);
 
 	UE_RIGVM_GET_READ_REGISTRY_CONTEXT();
 
 	// Now cache the handles as needed, in several batches as needed
-	const int32 TaskBatchSize = 10;
-	ParallelFor(TEXT("CacheMemoryHandlesIfRequired"), Instructions.Num(), TaskBatchSize, [this, &Context, &RegistryHandle](int32 InstructionIndex)
+	const int32 MinHandlesForParallel = CVarControlRigCacheMemoryHandlesParallelMinSize->GetInt();
+	const bool bUseParallel = CVarControlRigCacheMemoryHandlesBatchSize->GetInt() > 0 && MinHandlesForParallel > 0 && NumInstructions > MinHandlesForParallel;
+	const int32 TaskBatchSize = bUseParallel ? CVarControlRigCacheMemoryHandlesBatchSize->GetInt() : INT32_MAX;
+
+	ParallelFor(TEXT("CacheMemoryHandles"), NumInstructions, TaskBatchSize, [this, &Context, &RegistryHandle](int32 InstructionIndex)
 		{
 			InstructionOpEval(Context, RegistryHandle, InstructionIndex, FirstHandleForInstruction[InstructionIndex],
-				[&](FRigVMExtendedExecuteContext& InContext, const FRigVMRegistryHandle& InRegistry, int32 InHandleIndex, const FRigVMBranchInfoKey& InBranchInfoKey, const FRigVMOperand& InOp)
+				[&](FRigVMExtendedExecuteContext& InContext, const FRigVMRegistryHandle& InRegistry, int32 InInstructionIndex, int32 InHandleIndex, const FRigVMBranchInfoKey& InBranchInfoKey, const FRigVMOperand& InOp)
 				{
-					CacheSingleMemoryHandle(InContext, InRegistry, InHandleIndex, InBranchInfoKey, InOp);
+					PreCacheSingleMemoryHandle(InRegistry, InInstructionIndex, InHandleIndex, InBranchInfoKey, InOp);
 				});
-		}
+		}, bUseParallel ? EParallelForFlags::None : EParallelForFlags::ForceSingleThread
 	);
 }
 
+void URigVM::PatchMemoryHandles(FRigVMExtendedExecuteContext& Context)
+{
+	const int32 NumInstructions = Instructions.Num();
+	if (NumInstructions == 0)
+	{
+		return;
+	}
+
+	check((NumInstructions + 1) == FirstHandleForInstruction.Num());
+
+	if (Context.CachedMemoryHandles.Num() == MemoryHandleCount)
+	{
+		return;
+	}
+
+	// Copy the pre-cached handles from the VM into the instance context
+	Context.CachedMemoryHandles = PreCachedMemoryHandles;
+
+	UE_RIGVM_GET_READ_REGISTRY_CONTEXT();
+
+	FRigVMByteCode& ByteCode = GetByteCode();
+
+	// Now patch the handles memory, which contain offsets at this point
+	static TStaticArray<ERigVMMemoryType, (int32)ERigVMMemoryType::Num> MemoryTypes = { ERigVMMemoryType::Work, ERigVMMemoryType::Literal, ERigVMMemoryType::External, ERigVMMemoryType::Debug };
+	static_assert((int32)ERigVMMemoryType::Num == 4);
+
+	TStaticArray<FRigVMMemoryStorageStruct*, (int32)ERigVMMemoryType::Num> MemoryPtrs;
+	for (int32 MemoryTypeIndex = 0; MemoryTypeIndex < MemoryTypes.Num(); MemoryTypeIndex++)
+	{
+		MemoryPtrs[MemoryTypeIndex] = GetMemoryByType(Context, MemoryTypes[MemoryTypeIndex]);
+	}
+
+	check(PreCachedMemoryHandlesMemoryType.Num() == Context.CachedMemoryHandles.Num());
+
+	// Fix all pre cached handles memory ptr, by adding the corresponding memory start (different by memory type)
+	FRigVMMemoryHandle* RESTRICT MemoryHandle = Context.CachedMemoryHandles.GetData();
+	const FRigVMMemoryHandle* RESTRICT MemoryHandlesEnd = Context.CachedMemoryHandles.GetData() + Context.CachedMemoryHandles.Num();
+	ERigVMMemoryType* RESTRICT HandleMemoryType = PreCachedMemoryHandlesMemoryType.GetData();
+	for (; MemoryHandle < MemoryHandlesEnd; MemoryHandle++, HandleMemoryType++)
+	{
+		if (const FRigVMMemoryStorageStruct* Memory = MemoryPtrs[(int32)*HandleMemoryType])
+		{
+			const PTRINT BaseAddres = (PTRINT)Memory->GetContainerPtr();
+			const PTRINT Offset = (PTRINT)MemoryHandle->Ptr;
+			MemoryHandle->Ptr = (uint8*)(BaseAddres + Offset);
+		}
+	}
+
+	// Fix handles using external memory, these can not be fixed just with the base memory address, have to be fully processed
+	for (const int32 InstructionIndex : InstructionsIndexesWithExternalMemory)
+	{
+		InstructionOpEval(Context, RegistryHandle, InstructionIndex, FirstHandleForInstruction[InstructionIndex],
+			[&](FRigVMExtendedExecuteContext& InContext, const FRigVMRegistryHandle& InRegistry, int32 InInstructionIndex, int32 InHandleIndex, const FRigVMBranchInfoKey& InBranchInfoKey, const FRigVMOperand& InOp)
+			{
+				// This will only process the handles of the instruction that use external memory, the rest will be ignored (these have been already patched)
+				PatchExternalMemoryHandles(Context, InRegistry, InHandleIndex, InOp);
+			});
+	}
+}
+
 void URigVM::RebuildByteCodeOnLoad()
 {
 	Instructions = GetByteCode().GetInstructions();
@@ -1534,6 +1626,8 @@ bool URigVM::Initialize(FRigVMExtendedExecuteContext& Context)
 
 	PrepareMemoryForExecution(Context);
 
+	CacheMemoryHandles(Context);
+
 	return true;
 }
 
@@ -1557,7 +1651,7 @@ bool URigVM::InitializeInstance(FRigVMExtendedExecuteContext& Context, bool bCop
 	}
 
 	UE::TScopeLock CreateLocalizedRegistryScopeLock(CreateLocalizedRegistryMutex);
-	CacheMemoryHandlesIfRequired(Context);
+	PatchMemoryHandles(Context);
 
 	return true;
 }
@@ -3077,14 +3171,14 @@ void URigVM::ClearDebugMemory(FRigVMExtendedExecuteContext& Context)
 #endif
 }
 
-void URigVM::CacheSingleMemoryHandle(FRigVMExtendedExecuteContext& Context, const FRigVMRegistryHandle& InRegistry, int32 InHandleIndex, const FRigVMBranchInfoKey& InBranchInfoKey, const FRigVMOperand& InArg, bool bForExecute)
+void URigVM::PreCacheSingleMemoryHandle(const FRigVMRegistryHandle& InRegistry, int32 InInstructionIndex, int32 InHandleIndex, const FRigVMBranchInfoKey& InBranchInfoKey, const FRigVMOperand& InArg, bool bForExecute)
 {
-	FRigVMMemoryStorageStruct* Memory = GetMemoryByType(Context, InArg.GetMemoryType()/*, false*/);
+	FRigVMMemoryStorageStruct* Memory = GetDefaultMemoryByType(InArg.GetMemoryType());
 
 	if (InArg.GetMemoryType() == ERigVMMemoryType::External)
 	{
 		const FRigVMPropertyPath* PropertyPath = nullptr;
-		if(InArg.GetRegisterOffset() != INDEX_NONE)
+		if (InArg.GetRegisterOffset() != INDEX_NONE)
 		{
 			check(ExternalPropertyPaths.IsValidIndex(InArg.GetRegisterOffset()));
 			PropertyPath = &ExternalPropertyPaths[InArg.GetRegisterOffset()];
@@ -3096,19 +3190,17 @@ void URigVM::CacheSingleMemoryHandle(FRigVMExtendedExecuteContext& Context, cons
 		FRigVMExternalVariableDef& ExternalVariable = ExternalVariables[ExternalVariableIndex];
 		check(ExternalVariable.IsValid());
 
-		FRigVMExternalVariableRuntimeData& ExternalVariableRuntimeData = Context.ExternalVariableRuntimeData[ExternalVariableIndex];
-		check(ExternalVariableRuntimeData.Memory != nullptr);
-
-		Context.CachedMemoryHandles[InHandleIndex] = {ExternalVariableRuntimeData.Memory, ExternalVariable.Property, PropertyPath};
-#if UE_BUILD_DEBUG
-		// make sure the handle points to valid memory
-		check(Context.CachedMemoryHandles[InHandleIndex].GetOutputData() != nullptr);
-#endif
+		PreCachedMemoryHandles[InHandleIndex] = { nullptr, ExternalVariable.Property, PropertyPath };
+		PreCachedMemoryHandlesMemoryType[InHandleIndex] = ERigVMMemoryType::External;
+		{
+			const int32 Index = CurrenntInstructionWithExternalMemory.fetch_add(1); // deal with parallel processing
+			InstructionsIndexesWithExternalMemory[Index] = InInstructionIndex; // These have to be patched on each instance
+		}
 		return;
 	}
 
 	const FRigVMPropertyPath* PropertyPath = nullptr;
-	if(InArg.GetRegisterOffset() != INDEX_NONE)
+	if (InArg.GetRegisterOffset() != INDEX_NONE)
 	{
 		check(Memory->GetPropertyPaths().IsValidIndex(InArg.GetRegisterOffset()));
 		PropertyPath = &Memory->GetPropertyPaths()[InArg.GetRegisterOffset()];
@@ -3123,14 +3215,17 @@ void URigVM::CacheSingleMemoryHandle(FRigVMExtendedExecuteContext& Context, cons
 	// package. the compiler bases the memory class construction on the package the VM
 	// is in - so a VM under GetTransientPackage() can be created - but not run.
 	uint8* Data = Memory->GetData<uint8>(InArg.GetRegisterIndex());
+	PTRINT DataOffset = ((uint8*)Data - (uint8*)Memory->GetContainerPtr());
 	const FProperty* Property = Memory->GetProperties()[InArg.GetRegisterIndex()];
-	FRigVMMemoryHandle& Handle = Context.CachedMemoryHandles[InHandleIndex];
-	Handle = {Data, Property, PropertyPath};
+	FRigVMMemoryHandle& Handle = PreCachedMemoryHandles[InHandleIndex];
+
+	Handle = { (uint8*)DataOffset, Property, PropertyPath };
+	PreCachedMemoryHandlesMemoryType[InHandleIndex] = InArg.GetMemoryType();
 
 	// if we are lazy executing update the handle to point to a lazy branch
-	if(InBranchInfoKey.IsValid())
+	if (InBranchInfoKey.IsValid())
 	{
-		if(const FRigVMBranchInfo* BranchInfo = GetByteCode().GetBranchInfo(InBranchInfoKey))
+		if (const FRigVMBranchInfo* BranchInfo = GetByteCode().GetBranchInfo(InBranchInfoKey))
 		{
 			FRigVMLazyBranch* LazyBranch = &LazyBranches[BranchInfo->Index];
 			LazyBranch->VM = this;
@@ -3141,10 +3236,9 @@ void URigVM::CacheSingleMemoryHandle(FRigVMExtendedExecuteContext& Context, cons
 
 #if UE_BUILD_DEBUG
 	// make sure the handle points to valid memory
-	check(Handle.GetOutputData() != nullptr);
-	if(PropertyPath)
+	if (PropertyPath)
 	{
-		uint8* MemoryPtr = Handle.GetOutputData();
+		uint8* MemoryPtr = (uint8*)Memory->GetContainerPtr() + (PTRINT)Handle.GetOutputData();
 		// don't check the result - since it may be an array element
 		// that doesn't exist yet. 
 		PropertyPath->GetData<uint8>(MemoryPtr, Property);
@@ -3152,6 +3246,42 @@ void URigVM::CacheSingleMemoryHandle(FRigVMExtendedExecuteContext& Context, cons
 #endif
 }
 
+void URigVM::PatchExternalMemoryHandles(FRigVMExtendedExecuteContext& Context, const FRigVMRegistryHandle& InRegistry, int32 InHandleIndex, const FRigVMOperand& InArg)
+{
+	FRigVMMemoryStorageStruct* Memory = GetMemoryByType(Context, InArg.GetMemoryType()/*, false*/);
+
+	FRigVMMemoryHandle& Handle = Context.CachedMemoryHandles[InHandleIndex];
+
+	if (InArg.GetMemoryType() == ERigVMMemoryType::External)
+	{
+		const FRigVMPropertyPath* PropertyPath = nullptr;
+		if (InArg.GetRegisterOffset() != INDEX_NONE)
+		{
+			check(ExternalPropertyPaths.IsValidIndex(InArg.GetRegisterOffset()));
+			PropertyPath = &ExternalPropertyPaths[InArg.GetRegisterOffset()];
+		}
+
+		check(ExternalVariables.IsValidIndex(InArg.GetRegisterIndex()));
+
+		const int32 ExternalVariableIndex = InArg.GetRegisterIndex();
+		FRigVMExternalVariableDef& ExternalVariable = ExternalVariables[ExternalVariableIndex];
+		check(ExternalVariable.IsValid());
+
+		FRigVMExternalVariableRuntimeData& ExternalVariableRuntimeData = Context.ExternalVariableRuntimeData[ExternalVariableIndex];
+		check(ExternalVariableRuntimeData.Memory != nullptr);
+		check(Handle.PropertyPath == PropertyPath);
+
+		Handle.Ptr = ExternalVariableRuntimeData.Memory;
+
+#if UE_BUILD_DEBUG
+		// make sure the handle points to valid memory
+		check(Context.CachedMemoryHandles[InHandleIndex].GetOutputData() != nullptr);
+#endif
+	}
+
+	// if handle is not using external memory, just ignore it, as it has been already patched
+}
+
 void URigVM::CopyOperandForDebuggingImpl(FRigVMExtendedExecuteContext& Context, const FRigVMOperand& InArg, const FRigVMMemoryHandle& InHandle, const FRigVMOperand& InDebugOperand)
 {
 #if WITH_EDITOR
diff --git a/Engine/Plugins/Runtime/RigVM/Source/RigVM/Public/RigVMCore/RigVM.h b/Engine/Plugins/Runtime/RigVM/Source/RigVM/Public/RigVMCore/RigVM.h
--- a/Engine/Plugins/Runtime/RigVM/Source/RigVM/Public/RigVMCore/RigVM.h
+++ b/Engine/Plugins/Runtime/RigVM/Source/RigVM/Public/RigVMCore/RigVM.h
@@ -548,6 +548,10 @@ public:
 		}
 
 		const FProperty* Property = GetWorkMemory(Context)->FindPropertyByName(InWorkMemoryPropertyName);
+		if (Property == nullptr)
+		{
+			return FRigVMParameter();
+		}
 		const int32 PropertyIndex = GetWorkMemory(Context)->GetPropertyIndex(Property);
 
 		UScriptStruct* Struct = nullptr;
@@ -568,12 +572,15 @@ public:
 	int32 GetParameterArraySize(FRigVMExtendedExecuteContext& Context, const FRigVMParameter& InParameter)
 	{
 		const int32 PropertyIndex = InParameter.GetRegisterIndex();
-		const FProperty* Property = GetWorkMemory(Context)->GetProperties()[PropertyIndex];
-		const FArrayProperty* ArrayProperty = CastField<FArrayProperty>(Property);
-		if(ArrayProperty)
+		if (PropertyIndex != INDEX_NONE)
 		{
-			FScriptArrayHelper ArrayHelper(ArrayProperty, GetWorkMemory(Context)->GetData<uint8>(PropertyIndex));
-			return ArrayHelper.Num();
+			const FProperty* Property = GetWorkMemory(Context)->GetProperties()[PropertyIndex];
+			const FArrayProperty* ArrayProperty = CastField<FArrayProperty>(Property);
+			if (ArrayProperty)
+			{
+				FScriptArrayHelper ArrayHelper(ArrayProperty, GetWorkMemory(Context)->GetData<uint8>(PropertyIndex));
+				return ArrayHelper.Num();
+			}
 		}
 		return 1;
 	}
@@ -618,8 +625,6 @@ public:
 			{
 				return *GetWorkMemory(Context)->GetData<T>(InParameter.GetRegisterIndex());
 			}
-			
-			return *GetWorkMemory(Context)->GetData<T>(InParameter.GetRegisterIndex());
 		}
 		return DefaultValue;
 	}
@@ -961,12 +966,14 @@ private:
 		const TFunctionRef<void(
 			FRigVMExtendedExecuteContext& /*Context*/,
 			FRigVMRegistryHandle& /*RegistryHandle*/,
+			int32 /*InstructionIndex*/,
 			int32 /*InHandleIndex*/,
 			const FRigVMBranchInfoKey& /*InBranchInfoKey*/,
 			const FRigVMOperand& /*InArg*/
 		)>& InOpFunc);
 	RIGVM_API void PrepareMemoryForExecution(FRigVMExtendedExecuteContext& Context);
-	RIGVM_API void CacheMemoryHandlesIfRequired(FRigVMExtendedExecuteContext& Context);
+	RIGVM_API void CacheMemoryHandles(FRigVMExtendedExecuteContext& Context);
+	RIGVM_API void PatchMemoryHandles(FRigVMExtendedExecuteContext& Context);
 	RIGVM_API void RebuildByteCodeOnLoad();
 
 	UPROPERTY(transient)
@@ -1018,6 +1025,12 @@ private:
 
 	int32 MemoryHandleCount = 0;
 
+	TArray<FRigVMMemoryHandle> PreCachedMemoryHandles;
+	TArray<ERigVMMemoryType> PreCachedMemoryHandlesMemoryType;
+	TArray<int32> InstructionsIndexesWithExternalMemory;
+	int32 NumHandlesWithExternalMemory = 0;
+	std::atomic<int32> CurrenntInstructionWithExternalMemory = 0;
+
 	TArray<FRigVMExternalVariableDef> ExternalVariables;
 	TArray<FRigVMLazyBranch> LazyBranches;
 	
@@ -1039,7 +1052,9 @@ private:
 	// debug watch register memory needs to be cleared for each execution
 	RIGVM_API void ClearDebugMemory(FRigVMExtendedExecuteContext& Context);
 
-	RIGVM_API void CacheSingleMemoryHandle(FRigVMExtendedExecuteContext& Context, const FRigVMRegistryHandle& InRegistry, int32 InHandleIndex, const FRigVMBranchInfoKey& InBranchInfoKey, const FRigVMOperand& InArg, bool bForExecute = false);
+	RIGVM_API void PreCacheSingleMemoryHandle(const FRigVMRegistryHandle& InRegistry, int32 InInstructionIndex, int32 InHandleIndex, const FRigVMBranchInfoKey& InBranchInfoKey, const FRigVMOperand& InArg, bool bForExecute = false);
+	RIGVM_API void PatchExternalMemoryHandles(FRigVMExtendedExecuteContext& Context, const FRigVMRegistryHandle& InRegistry, int32 InHandleIndex, const FRigVMOperand& InArg);
+
 
 	void CopyOperandForDebuggingIfNeeded(FRigVMExtendedExecuteContext& Context, const FRigVMOperand& InArg, const FRigVMMemoryHandle& InHandle)
 	{
diff --git a/Engine/Plugins/Runtime/RigVM/Source/RigVM/Public/RigVMCore/RigVMByteCode.h b/Engine/Plugins/Runtime/RigVM/Source/RigVM/Public/RigVMCore/RigVMByteCode.h
--- a/Engine/Plugins/Runtime/RigVM/Source/RigVM/Public/RigVMCore/RigVMByteCode.h
+++ b/Engine/Plugins/Runtime/RigVM/Source/RigVM/Public/RigVMCore/RigVMByteCode.h
@@ -1091,6 +1091,8 @@ public:
 	TArray<FRigVMInstruction>::RangedForConstIteratorType begin() const { return Instructions.begin(); }
 	TArray<FRigVMInstruction>::RangedForConstIteratorType end() const { return Instructions.end(); }
 
+	const FRigVMInstruction* GetData() const { return Instructions.GetData(); }
+
 private:
 
 	// hide utility constructor
diff --git a/Engine/Plugins/Runtime/RigVM/Source/RigVM/Public/RigVMCore/RigVMMemoryCommon.h b/Engine/Plugins/Runtime/RigVM/Source/RigVM/Public/RigVMCore/RigVMMemoryCommon.h
--- a/Engine/Plugins/Runtime/RigVM/Source/RigVM/Public/RigVMCore/RigVMMemoryCommon.h
+++ b/Engine/Plugins/Runtime/RigVM/Source/RigVM/Public/RigVMCore/RigVMMemoryCommon.h
@@ -43,7 +43,8 @@ enum class ERigVMMemoryType: uint8
 	Literal = 1, // Const / fixed state
 	External = 2, // Unowned external memory
 	Debug = 3, // Owned memory used for debug watches
-	Invalid
+	Invalid,
+	Num = Invalid
 };
 
 namespace RigVM
