Control Rig: Add ability to insert pin in sequence node

diff --git a/Engine/Plugins/Runtime/RigVM/Source/RigVMEditor/Private/RigVMEditorModule.cpp b/Engine/Plugins/Runtime/RigVM/Source/RigVMEditor/Private/RigVMEditorModule.cpp
--- a/Engine/Plugins/Runtime/RigVM/Source/RigVMEditor/Private/RigVMEditorModule.cpp
+++ b/Engine/Plugins/Runtime/RigVM/Source/RigVMEditor/Private/RigVMEditorModule.cpp
@@ -45,6 +45,8 @@
 #include "UObject/UObjectIterator.h"
 #include "Toolkits/ToolkitManager.h"
 #include "ToolMenus.h"
+#include "RigVMFunctions/Execution/RigVMFunction_Sequence.h"
+#include "RigVMModel/RigVMControllerActions.h"
 #include "Widgets/SRigVMSwapFunctionsWidget.h"
 #include "Widgets/SRigVMBulkEditDialog.h"
 #include "Trace/RigVMTraceTrack.h"
@@ -1669,6 +1671,7 @@ void FRigVMEditorModule::GetPinContextMenuActions(TScriptInterface<IRigVMClientH
 	GetPinWorkflowContextMenuActions(RigVMClientHost, EdGraphPin, ModelPin, Menu);
 	GetPinDebugContextMenuActions(RigVMClientHost, EdGraphPin, ModelPin, Menu);
 	GetPinArrayContextMenuActions(RigVMClientHost, EdGraphPin, ModelPin, Menu);
+	GetPinSequenceContextMenuActions(RigVMClientHost, EdGraphPin, ModelPin, Menu);
 	GetPinAggregateContextMenuActions(RigVMClientHost, EdGraphPin, ModelPin, Menu);
 	GetPinTemplateContextMenuActions(RigVMClientHost, EdGraphPin, ModelPin, Menu);
 	GetPinConversionContextMenuActions(RigVMClientHost, EdGraphPin, ModelPin, Menu);
@@ -1804,6 +1807,61 @@ void FRigVMEditorModule::GetPinArrayContextMenuActions(TScriptInterface<IRigVMCl
 	}
 }
 
+void FRigVMEditorModule::GetPinSequenceContextMenuActions(TScriptInterface<IRigVMClientHost> RigVMClientHost, const UEdGraphPin* EdGraphPin, URigVMPin* ModelPin, UToolMenu* Menu) const
+{
+	if (ModelPin->GetDirection() == ERigVMPinDirection::Output)
+	{
+		URigVMController* Controller = RigVMClientHost->GetRigVMClient()->GetController(ModelPin->GetGraph());
+		if (URigVMAggregateNode* AggregateNode = Cast<URigVMAggregateNode>(ModelPin->GetNode()))
+		{
+			if (URigVMUnitNode* UnitNode = Cast<URigVMUnitNode>(AggregateNode->GetFirstInnerNode()))
+			{
+				if (UnitNode->GetScriptStruct() == FRigVMFunction_Sequence::StaticStruct())
+				{
+					FToolMenuSection& Section = Menu->AddSection("RigVMEditorContextMenuSequencePin", LOCTEXT("SequencePin", "Sequence"));
+					Section.AddMenuEntry(
+						"InsertSequencePin",
+						LOCTEXT("InsertSequencePin", "Insert Pin Above"),
+						LOCTEXT("InsertSequencePin_Tooltip", "Inserts a new output pin"),
+						FSlateIcon(),
+						FUIAction(FExecuteAction::CreateLambda([Controller, ModelPin, AggregateNode]() {
+
+							Controller->GetActionStack()->OpenUndoBracket(Controller, TEXT("Insert Sequence Pin"));
+							if (!Controller->AddAggregatePin(AggregateNode, FString(), FString(), true, true).IsEmpty())
+							{
+								TArray<URigVMPin*> Pins = AggregateNode->GetPins();
+								for (int32 i=Pins.Num()-1; i > 0; i--)
+								{
+									if (Pins[i] == ModelPin)
+									{
+										break;
+									}
+	
+									if (Pins.IsValidIndex(i-1))
+									{
+										TArray<URigVMPin*> TargetPins = Pins[i-1]->GetLinkedTargetPins();
+										
+										Controller->BreakAllLinks(Pins[i-1], false);
+										for (URigVMPin* TargetPin : TargetPins)
+										{
+											Controller->AddLink(Pins[i], TargetPin);
+										}
+									}
+								}
+							}
+							else
+							{
+								Controller->GetActionStack()->CancelUndoBracket(Controller);
+							}
+							Controller->GetActionStack()->CloseUndoBracket(Controller);
+						})
+					));
+				}
+			}
+		}
+	}
+}
+
 void FRigVMEditorModule::GetPinAggregateContextMenuActions(TScriptInterface<IRigVMClientHost> RigVMClientHost, const UEdGraphPin* EdGraphPin, URigVMPin* ModelPin, UToolMenu* Menu) const
 {
 	URigVMController* Controller = RigVMClientHost->GetRigVMClient()->GetController(ModelPin->GetGraph());
diff --git a/Engine/Plugins/Runtime/RigVM/Source/RigVMEditor/Public/RigVMEditorModule.h b/Engine/Plugins/Runtime/RigVM/Source/RigVMEditor/Public/RigVMEditorModule.h
--- a/Engine/Plugins/Runtime/RigVM/Source/RigVMEditor/Public/RigVMEditorModule.h
+++ b/Engine/Plugins/Runtime/RigVM/Source/RigVMEditor/Public/RigVMEditorModule.h
@@ -169,6 +169,7 @@ protected:
 	UE_API virtual void GetPinWorkflowContextMenuActions(TScriptInterface<IRigVMClientHost> RigVMClientHost, const UEdGraphPin* EdGraphPin, URigVMPin* ModelPin, UToolMenu* Menu) const;
 	UE_API virtual void GetPinDebugContextMenuActions(TScriptInterface<IRigVMClientHost> RigVMClientHost, const UEdGraphPin* EdGraphPin, URigVMPin* ModelPin, UToolMenu* Menu) const;
 	UE_API virtual void GetPinArrayContextMenuActions(TScriptInterface<IRigVMClientHost> RigVMClientHost, const UEdGraphPin* EdGraphPin, URigVMPin* ModelPin, UToolMenu* Menu) const;
+	UE_API virtual void GetPinSequenceContextMenuActions(TScriptInterface<IRigVMClientHost> RigVMClientHost, const UEdGraphPin* EdGraphPin, URigVMPin* ModelPin, UToolMenu* Menu) const;
 	UE_API virtual void GetPinAggregateContextMenuActions(TScriptInterface<IRigVMClientHost> RigVMClientHost, const UEdGraphPin* EdGraphPin, URigVMPin* ModelPin, UToolMenu* Menu) const;
 	UE_API virtual void GetPinTemplateContextMenuActions(TScriptInterface<IRigVMClientHost> RigVMClientHost, const UEdGraphPin* EdGraphPin, URigVMPin* ModelPin, UToolMenu* Menu) const;
 	UE_API virtual void GetPinConversionContextMenuActions(TScriptInterface<IRigVMClientHost> RigVMClientHost, const UEdGraphPin* EdGraphPin, URigVMPin* ModelPin, UToolMenu* Menu) const;
