Fixed a Control Rig crash adding a graph after renaming another newly added graph

diff --git a/Engine/Plugins/Runtime/RigVM/Source/RigVMDeveloper/Private/RigVMModel/RigVMClient.cpp b/Engine/Plugins/Runtime/RigVM/Source/RigVMDeveloper/Private/RigVMModel/RigVMClient.cpp
--- a/Engine/Plugins/Runtime/RigVM/Source/RigVMDeveloper/Private/RigVMModel/RigVMClient.cpp
+++ b/Engine/Plugins/Runtime/RigVM/Source/RigVMDeveloper/Private/RigVMModel/RigVMClient.cpp
@@ -971,6 +971,7 @@ FName FRigVMClient::RenameModel(const FString& InNodePathOrName, const FName& In
 
 		const FString OldNodePath = Model->GetNodePath();
 		const FName SafeNewName = GetUniqueName(InNewName);
+		const FSoftObjectPath OldKey(Model);
 		FString NewNodePath;
 
 		{
@@ -981,6 +982,7 @@ FName FRigVMClient::RenameModel(const FString& InNodePathOrName, const FName& In
 			if (Controller)
 			{
 				Controllers.Add(Model, *Controller);
+				Controllers.Remove(OldKey);
 			}
 		}
 
