Fix for RigVM debug info use after delete

diff --git a/Engine/Plugins/Animation/ControlRig/Source/ControlRig/Private/ControlRig.cpp b/Engine/Plugins/Animation/ControlRig/Source/ControlRig/Private/ControlRig.cpp
--- a/Engine/Plugins/Animation/ControlRig/Source/ControlRig/Private/ControlRig.cpp
+++ b/Engine/Plugins/Animation/ControlRig/Source/ControlRig/Private/ControlRig.cpp
@@ -950,10 +950,10 @@ bool UControlRig::Execute(const FName& InEventName)
 	if (UControlRig* CDO = GetClass()->GetDefaultObject<UControlRig>())
 	{
 		// Copy the breakpoints. This will not override the state of the breakpoints
-		DebugInfo.CopyFrom(CDO->DebugInfo);
+		DebugInfo->CopyFrom(CDO->DebugInfo.Get());
 	}
 
-	ExtendedExecuteContext.SetDebugInfo(&DebugInfo);
+	ExtendedExecuteContext.SetDebugInfo(DebugInfo);
 	
 	if (IsProfilingEnabled())
     {
diff --git a/Engine/Plugins/Runtime/RigVM/Source/RigVM/Private/RigVMCore/RigVM.cpp b/Engine/Plugins/Runtime/RigVM/Source/RigVM/Private/RigVMCore/RigVM.cpp
--- a/Engine/Plugins/Runtime/RigVM/Source/RigVM/Private/RigVMCore/RigVM.cpp
+++ b/Engine/Plugins/Runtime/RigVM/Source/RigVM/Private/RigVMCore/RigVM.cpp
@@ -1817,31 +1817,36 @@ ERigVMExecuteResult URigVM::ExecuteVM(FRigVMExtendedExecuteContext& Context, con
 	}
 
 #if WITH_EDITOR
-	FRigVMDebugInfo* RigVMDebugInfo = Context.GetRigVMDebugInfo();
+	FRigVMDebugInfo* RigVMDebugInfo = nullptr;
 	FRigVMDebugInfo PreviousDebugInfo;
-	if(bIsRootEntry && RigVMDebugInfo)
+
+	if (TSharedPtr<FRigVMDebugInfo> DebugInfoShared = Context.GetRigVMDebugInfoWeak().Pin())
 	{
-		PreviousDebugInfo.CopyFrom(*RigVMDebugInfo);
-		
-		if (!RigVMDebugInfo->IsEmpty())
+		RigVMDebugInfo = DebugInfoShared.Get();
+		if (bIsRootEntry && RigVMDebugInfo)
 		{
-			// if the instruction we are about to early exit from is outside of the set of instructions
-			// we are about to run - disable the debug info.
-			const int32 EarlyExistInstructionIndex = RigVMDebugInfo->GetInstructionForEarlyExit();
-			if (EarlyExistInstructionIndex != FMath::Clamp(EarlyExistInstructionIndex, ContextPublicData.InstructionIndex, LastInstruction))
+			PreviousDebugInfo.CopyFrom(*RigVMDebugInfo);
+
+			if (!RigVMDebugInfo->IsEmpty())
 			{
-				RigVMDebugInfo->Reset();
+				// if the instruction we are about to early exit from is outside of the set of instructions
+				// we are about to run - disable the debug info.
+				const int32 EarlyExistInstructionIndex = RigVMDebugInfo->GetInstructionForEarlyExit();
+				if (EarlyExistInstructionIndex != FMath::Clamp(EarlyExistInstructionIndex, ContextPublicData.InstructionIndex, LastInstruction))
+				{
+					RigVMDebugInfo->Reset();
+				}
+				else
+				{
+					LastInstruction = EarlyExistInstructionIndex;
+				}
 			}
-			else
+
+			if (!RigVMDebugInfo->IsEmpty())
 			{
-				LastInstruction = EarlyExistInstructionIndex;
+				RigVMDebugInfo->ExecutionHalted().Broadcast(INDEX_NONE, nullptr, InEntryName);
 			}
 		}
-		
-		if (!RigVMDebugInfo->IsEmpty())
-		{
-			RigVMDebugInfo->ExecutionHalted().Broadcast(INDEX_NONE, nullptr, InEntryName);
-		}
 	}
 #endif
 
diff --git a/Engine/Plugins/Runtime/RigVM/Source/RigVM/Private/RigVMCore/RigVMExecuteContext.cpp b/Engine/Plugins/Runtime/RigVM/Source/RigVM/Private/RigVMCore/RigVMExecuteContext.cpp
--- a/Engine/Plugins/Runtime/RigVM/Source/RigVM/Private/RigVMCore/RigVMExecuteContext.cpp
+++ b/Engine/Plugins/Runtime/RigVM/Source/RigVM/Private/RigVMCore/RigVMExecuteContext.cpp
@@ -117,10 +117,10 @@ void FRigVMExtendedExecuteContext::Reset()
 	bCurrentlyRunningRootEntry = false;
 
 #if WITH_EDITOR
-	if (DebugInfo != nullptr)
+	if (TSharedPtr<FRigVMDebugInfo> DebugInfoShared = DebugInfoWeak.Pin())
 	{
-		DebugInfo->Reset();
-		DebugInfo = nullptr;
+		DebugInfoShared->Reset();
+		DebugInfoWeak.Reset();
 	}
 #endif // WITH_EDITOR
 }
diff --git a/Engine/Plugins/Runtime/RigVM/Source/RigVM/Private/RigVMHost.cpp b/Engine/Plugins/Runtime/RigVM/Source/RigVM/Private/RigVMHost.cpp
--- a/Engine/Plugins/Runtime/RigVM/Source/RigVM/Private/RigVMHost.cpp
+++ b/Engine/Plugins/Runtime/RigVM/Source/RigVM/Private/RigVMHost.cpp
@@ -45,6 +45,9 @@ URigVMHost::URigVMHost(const FObjectInitializer& ObjectInitializer)
 	, bRequiresInitExecution(false)
 	, InitBracket(0)
 	, ExecuteBracket(0)
+#if WITH_EDITOR
+	, DebugInfo(MakeShared<FRigVMDebugInfo>())
+#endif
 #if UE_RIGVM_PROFILE_EXECUTE_UNITS_NUM
 	, ProfilingRunsLeft(0)
 	, AccumulatedCycles(0)
@@ -196,7 +199,7 @@ void URigVMHost::BeginDestroy()
 	InitializedEvent.Clear();
 	ExecutedEvent.Clear();
 #if WITH_EDITOR
-	DebugInfo.Reset();
+	DebugInfo->Reset();
 #endif
 
 	TRACE_OBJECT_LIFETIME_END(this);
@@ -503,10 +506,10 @@ bool URigVMHost::Execute(const FName& InEventName)
 
 	if (CDO)
 	{
-		DebugInfo.CopyFrom(CDO->DebugInfo);
+		DebugInfo->CopyFrom(CDO->DebugInfo.Get());
 	}
 
-	ExtendedExecuteContext.SetDebugInfo(&DebugInfo);
+	ExtendedExecuteContext.SetDebugInfo(DebugInfo);
 	
 	if (IsProfilingEnabled())
 	{
diff --git a/Engine/Plugins/Runtime/RigVM/Source/RigVM/Public/RigVMCore/RigVMExecuteContext.h b/Engine/Plugins/Runtime/RigVM/Source/RigVM/Public/RigVMCore/RigVMExecuteContext.h
--- a/Engine/Plugins/Runtime/RigVM/Source/RigVM/Public/RigVMCore/RigVMExecuteContext.h
+++ b/Engine/Plugins/Runtime/RigVM/Source/RigVM/Public/RigVMCore/RigVMExecuteContext.h
@@ -894,8 +894,11 @@ struct FRigVMExtendedExecuteContext
 
 #if WITH_EDITOR
 
+	UE_DEPRECATED(5.8, "DebugInfo pointer has been deprecated, please use DebugInfoWeak.")
 	FRigVMDebugInfo* DebugInfo = nullptr;
 
+	TWeakPtr<FRigVMDebugInfo> DebugInfoWeak;
+
 	FRigVMInstructionVisitInfo* InstructionVisitInfo = nullptr;
 
 	void SetInstructionVisitInfo(FRigVMInstructionVisitInfo* InInstructionVisitInfo)
@@ -908,14 +911,20 @@ struct FRigVMExtendedExecuteContext
 		return InstructionVisitInfo;
 	}
 
-	void SetDebugInfo(FRigVMDebugInfo* InDebugInfo)
+	void SetDebugInfo(const TSharedPtr<FRigVMDebugInfo>& InSharedDebugInfo)
 	{
-		DebugInfo = InDebugInfo;
+		DebugInfoWeak = InSharedDebugInfo.ToWeakPtr();
 	}
 
-	FRigVMDebugInfo* GetRigVMDebugInfo() const
+	UE_DEPRECATED(5.8, "SetDebugInfo with DebugInfo pointer has been deprecated, please use SharedDebugInfo version.")
+	void SetDebugInfo(FRigVMDebugInfo* InDebugInfo) {}
+
+	UE_DEPRECATED(5.8, "SetDebugInfo with DebugInfo pointer has been deprecated, please use GetRigVMDebugInfoWeak version.")
+	FRigVMDebugInfo* GetRigVMDebugInfo() const { return nullptr; }
+	
+	const TWeakPtr<FRigVMDebugInfo>& GetRigVMDebugInfoWeak() const
 	{
-		return DebugInfo;
+		return DebugInfoWeak;
 	}
 
 	FRigVMProfilingInfo* ProfilingInfo = nullptr;
diff --git a/Engine/Plugins/Runtime/RigVM/Source/RigVM/Public/RigVMHost.h b/Engine/Plugins/Runtime/RigVM/Source/RigVM/Public/RigVMHost.h
--- a/Engine/Plugins/Runtime/RigVM/Source/RigVM/Public/RigVMHost.h
+++ b/Engine/Plugins/Runtime/RigVM/Source/RigVM/Public/RigVMHost.h
@@ -558,7 +558,7 @@ private:
 protected:
 	
 	FRigVMInstructionVisitInfo InstructionVisitInfo;
-	FRigVMDebugInfo DebugInfo;
+	TSharedRef<FRigVMDebugInfo> DebugInfo;
 	FRigVMProfilingInfo ProfilingInfo;
 	TMap<FString, bool> LoggedMessages;
 	UE_API void LogOnce(EMessageSeverity::Type InSeverity, int32 InInstructionIndex, const FString& InMessage);
@@ -572,8 +572,8 @@ public:
 		return IsInDebugMode() || VMRuntimeSettings.bEnableProfiling;
 	}
 	
-	FRigVMDebugInfo& GetDebugInfo() { return DebugInfo; }
-	const FRigVMDebugInfo& GetDebugInfo() const { return DebugInfo; }
+	FRigVMDebugInfo& GetDebugInfo() { return DebugInfo.Get(); }
+	const FRigVMDebugInfo& GetDebugInfo() const { return DebugInfo.Get(); }
 
 	FRigVMProfilingInfo& GetProfilingInfo() { return ProfilingInfo; }
 	const FRigVMProfilingInfo& GetProfilingInfo() const { return ProfilingInfo; }
