Fixed RigVM bad deprecation message and better use of weak pointer

#rnx

diff --git a/Engine/Plugins/Runtime/RigVM/Source/RigVM/Private/RigVMCore/RigVM.cpp b/Engine/Plugins/Runtime/RigVM/Source/RigVM/Private/RigVMCore/RigVM.cpp
--- a/Engine/Plugins/Runtime/RigVM/Source/RigVM/Private/RigVMCore/RigVM.cpp
+++ b/Engine/Plugins/Runtime/RigVM/Source/RigVM/Private/RigVMCore/RigVM.cpp
@@ -1817,13 +1817,13 @@ ERigVMExecuteResult URigVM::ExecuteVM(FRigVMExtendedExecuteContext& Context, con
 	}
 
 #if WITH_EDITOR
-	FRigVMDebugInfo* RigVMDebugInfo = nullptr;
+	TSharedPtr<FRigVMDebugInfo> RigVMDebugInfo = Context.GetRigVMDebugInfoWeak().Pin();
+
 	FRigVMDebugInfo PreviousDebugInfo;
 
-	if (TSharedPtr<FRigVMDebugInfo> DebugInfoShared = Context.GetRigVMDebugInfoWeak().Pin())
+	if (RigVMDebugInfo)
 	{
-		RigVMDebugInfo = DebugInfoShared.Get();
-		if (bIsRootEntry && RigVMDebugInfo)
+		if (bIsRootEntry)
 		{
 			PreviousDebugInfo.CopyFrom(*RigVMDebugInfo);
 
diff --git a/Engine/Plugins/Runtime/RigVM/Source/RigVM/Public/RigVMCore/RigVMExecuteContext.h b/Engine/Plugins/Runtime/RigVM/Source/RigVM/Public/RigVMCore/RigVMExecuteContext.h
--- a/Engine/Plugins/Runtime/RigVM/Source/RigVM/Public/RigVMCore/RigVMExecuteContext.h
+++ b/Engine/Plugins/Runtime/RigVM/Source/RigVM/Public/RigVMCore/RigVMExecuteContext.h
@@ -919,7 +919,7 @@ struct FRigVMExtendedExecuteContext
 	UE_DEPRECATED(5.8, "SetDebugInfo with DebugInfo pointer has been deprecated, please use SharedDebugInfo version.")
 	void SetDebugInfo(FRigVMDebugInfo* InDebugInfo) {}
 
-	UE_DEPRECATED(5.8, "SetDebugInfo with DebugInfo pointer has been deprecated, please use GetRigVMDebugInfoWeak version.")
+	UE_DEPRECATED(5.8, "GetRigVMDebugInfo has been deprecated, please use GetRigVMDebugInfoWeak instead.")
 	FRigVMDebugInfo* GetRigVMDebugInfo() const { return nullptr; }
 	
 	const TWeakPtr<FRigVMDebugInfo>& GetRigVMDebugInfoWeak() const
