Control Rig: Fix Items Array workflows not working since this rig unit has been converted to a dispatch node
#rnx

diff --git a/Engine/Plugins/Runtime/RigVM/Source/RigVMDeveloper/Private/RigVMUserWorkflowRegistry.cpp b/Engine/Plugins/Runtime/RigVM/Source/RigVMDeveloper/Private/RigVMUserWorkflowRegistry.cpp
--- a/Engine/Plugins/Runtime/RigVM/Source/RigVMDeveloper/Private/RigVMUserWorkflowRegistry.cpp
+++ b/Engine/Plugins/Runtime/RigVM/Source/RigVMDeveloper/Private/RigVMUserWorkflowRegistry.cpp
@@ -1,6 +1,7 @@
 ﻿// Copyright Epic Games, Inc. All Rights Reserved.
 
 #include "RigVMUserWorkflowRegistry.h"
+#include "RigVMModel/Nodes/RigVMTemplateNode.h"
 
 #include UE_INLINE_GENERATED_CPP_BY_NAME(RigVMUserWorkflowRegistry)
 
@@ -12,16 +13,30 @@ URigVMUserWorkflowRegistry* URigVMUserWorkflowRegistry::Get()
 int32 URigVMUserWorkflowRegistry::RegisterProvider(const UScriptStruct* InStruct, FRigVMUserWorkflowProvider InProvider)
 {
 	const int32 Handle = ++MaxHandle;
-	Providers.Emplace(Handle, InStruct, InProvider);
+	ProvidersForUnits.Emplace(Handle, InStruct, InProvider);
+	return Handle;
+}
+
+int32 URigVMUserWorkflowRegistry::RegisterProviderForFunction(const FString InFunctionNotation, FRigVMUserWorkflowProvider InProvider)
+{
+	const int32 Handle = ++MaxHandle;
+	ProvidersForDispatchFunctions.Emplace(Handle, InFunctionNotation, InProvider);
 	return Handle;
 }
 
 void URigVMUserWorkflowRegistry::UnregisterProvider(int32 InHandle)
 {
-	Providers.RemoveAll([InHandle](const TTuple<int32, const UScriptStruct*, FRigVMUserWorkflowProvider>& Provider) -> bool
+	int32 NumRemoved = ProvidersForUnits.RemoveAll([InHandle](const TTuple<int32, const UScriptStruct*, FRigVMUserWorkflowProvider>& Provider) -> bool
+		{
+			return Provider.Get<0>() == InHandle;
+		});
+	if (NumRemoved == 0)
 	{
-		return Provider.Get<0>() == InHandle;
-	});
+		ProvidersForDispatchFunctions.RemoveAll([InHandle](const TTuple<int32, const FString, FRigVMUserWorkflowProvider>& Provider) -> bool
+		{
+			return Provider.Get<0>() == InHandle;
+		});
+	}
 }
 
 TArray<FRigVMUserWorkflow> URigVMUserWorkflowRegistry::GetWorkflows(ERigVMUserWorkflowType InType, const UScriptStruct* InStruct, const UObject* InSubject) const
@@ -29,17 +44,38 @@ TArray<FRigVMUserWorkflow> URigVMUserWorkflowRegistry::GetWorkflows(ERigVMUserWo
 	TArray<FRigVMUserWorkflow> Workflows;
 
 	// remove stale delegates
-	Providers.RemoveAll([](const TTuple<int32, const UScriptStruct*, FRigVMUserWorkflowProvider>& Provider) -> bool
+	ProvidersForUnits.RemoveAll([](const TTuple<int32, const UScriptStruct*, FRigVMUserWorkflowProvider>& Provider) -> bool
+	{
+		return !Provider.Get<2>().IsBound();
+	});
+	ProvidersForDispatchFunctions.RemoveAll([](const TTuple<int32, const FString, FRigVMUserWorkflowProvider>& Provider) -> bool
 	{
 		return !Provider.Get<2>().IsBound();
 	});
 	
 
-	for(const TTuple<int32, const UScriptStruct*, FRigVMUserWorkflowProvider>& Provider : Providers)
+	if (InStruct)
+	{
+		for(const TTuple<int32, const UScriptStruct*, FRigVMUserWorkflowProvider>& Provider : ProvidersForUnits)
+		{
+			if(Provider.Get<1>() == InStruct || Provider.Get<1>() == nullptr)
+			{
+				Workflows.Append(Provider.Get<2>().Execute(InSubject));
+			}
+		}
+	}
+	else if (const URigVMTemplateNode* TemplateNode = Cast<URigVMTemplateNode>(InSubject))
 	{
-		if(Provider.Get<1>() == InStruct || Provider.Get<1>() == nullptr)
+		if (const FRigVMFunction* Function = TemplateNode->GetResolvedFunction())
 		{
-			Workflows.Append(Provider.Get<2>().Execute(InSubject));
+			const FString FunctionName = Function->GetName();
+			for(const TTuple<int32, const FString, FRigVMUserWorkflowProvider>& Provider : ProvidersForDispatchFunctions)
+			{
+				if(Provider.Get<1>() == FunctionName)
+				{
+					Workflows.Append(Provider.Get<2>().Execute(InSubject));
+				}
+			}
 		}
 	}
 
diff --git a/Engine/Plugins/Runtime/RigVM/Source/RigVMDeveloper/Public/RigVMUserWorkflowRegistry.h b/Engine/Plugins/Runtime/RigVM/Source/RigVMDeveloper/Public/RigVMUserWorkflowRegistry.h
--- a/Engine/Plugins/Runtime/RigVM/Source/RigVMDeveloper/Public/RigVMUserWorkflowRegistry.h
+++ b/Engine/Plugins/Runtime/RigVM/Source/RigVMDeveloper/Public/RigVMUserWorkflowRegistry.h
@@ -36,6 +36,9 @@ public:
 	UFUNCTION(BlueprintCallable, Category = FRigVMUserWorkflowRegistry)
 	UE_API int32 RegisterProvider(const UScriptStruct* InStruct, FRigVMUserWorkflowProvider InProvider);
 
+	UFUNCTION(BlueprintCallable, Category = FRigVMUserWorkflowRegistry)
+	UE_API int32 RegisterProviderForFunction(const FString InFunctionNotation, FRigVMUserWorkflowProvider InProvider);
+
 	UFUNCTION(BlueprintCallable, Category = FRigVMUserWorkflowRegistry)
 	UE_API void UnregisterProvider(int32 InHandle);
 
@@ -45,7 +48,8 @@ public:
 private:
 
 	int32 MaxHandle = 0;
-	mutable TArray<TTuple<int32, const UScriptStruct*, FRigVMUserWorkflowProvider>> Providers;
+	mutable TArray<TTuple<int32, const UScriptStruct*, FRigVMUserWorkflowProvider>> ProvidersForUnits;
+	mutable TArray<TTuple<int32, const FString, FRigVMUserWorkflowProvider>> ProvidersForDispatchFunctions;
 };
 
 #undef UE_API
