RigVM Json Utils

#rnx

diff --git a/Engine/Plugins/Runtime/RigVM/Source/RigVMDeveloper/Private/RigVMBlueprintLegacy.cpp b/Engine/Plugins/Runtime/RigVM/Source/RigVMDeveloper/Private/RigVMBlueprintLegacy.cpp
--- a/Engine/Plugins/Runtime/RigVM/Source/RigVMDeveloper/Private/RigVMBlueprintLegacy.cpp
+++ b/Engine/Plugins/Runtime/RigVM/Source/RigVMDeveloper/Private/RigVMBlueprintLegacy.cpp
@@ -700,6 +700,19 @@ TArray<FRigVMGraphVariableDescription> URigVMBlueprint::GetAssetVariables() cons
 		FRigVMGraphVariableDescription NewVariable;
 		NewVariable.Name = BPVariable.VarName;
 		NewVariable.DefaultValue = BPVariable.DefaultValue;
+		if (NewVariable.DefaultValue.IsEmpty() && GeneratedClass)
+		{
+			if (const FProperty* Property = GeneratedClass->FindPropertyByName(NewVariable.Name))
+			{
+				if (const URigVMHost* CDO = Cast<URigVMHost>(GeneratedClass->GetDefaultObject(true)))
+				{
+					if (const void* PropertyAddr = Property->ContainerPtrToValuePtr<void>(CDO))
+					{
+						Property->ExportTextItem_Direct(NewVariable.DefaultValue, PropertyAddr, PropertyAddr, nullptr, PPF_SerializedAsImportText);
+					}
+				}
+			}
+		}
 		NewVariable.Category = BPVariable.Category;
 		FString CPPType;
 		UObject* CPPTypeObject;
diff --git a/Engine/Plugins/Runtime/RigVM/Source/RigVMDeveloper/Private/RigVMJsonUtils.cpp b/Engine/Plugins/Runtime/RigVM/Source/RigVMDeveloper/Private/RigVMJsonUtils.cpp
new file mode 100644
--- /dev/null
+++ b/Engine/Plugins/Runtime/RigVM/Source/RigVMDeveloper/Private/RigVMJsonUtils.cpp
@@ -0,0 +1,690 @@
+// Copyright Epic Games, Inc. All Rights Reserved.
+
+#include "RigVMJsonUtils.h"
+#include "HAL/IConsoleManager.h"
+#include "RigVMDeveloperModule.h"
+#include "RigVMCore/RigVMStruct.h"
+#include "RigVMAsset.h"
+#include "RigVMModel/RigVMGraph.h"
+#include "RigVMModel/Nodes/RigVMUnitNode.h"
+#include "RigVMModel/Nodes/RigVMDispatchNode.h"
+#include "RigVMModel/Nodes/RigVMVariableNode.h"
+#include "RigVMModel/Nodes/RigVMFunctionReferenceNode.h"
+#include "RigVMModel/Nodes/RigVMCollapseNode.h"
+#include "RigVMModel/Nodes/RigVMCommentNode.h"
+#include "RigVMModel/RigVMVariableDescription.h"
+#include "Serialization/JsonSerializer.h"
+
+#define LOCTEXT_NAMESPACE "RigVMJsonUtils"
+
+FAutoConsoleCommand FCmdRigVMJsonPrintAvailableNodes
+(
+	TEXT("RigVM.JSON.PrintAvailableNodes"),
+	TEXT("Prints out a json string containing information about all available nodes."),
+	FConsoleCommandDelegate::CreateLambda([]()
+	{
+		FRigVMRegistryReadLock Registry;
+			
+		TSharedPtr<FJsonObject> JsonObject = RigVMJsonUtils::AvailableNodesToJson(Registry, [](const UScriptStruct* InStruct) -> bool
+		{
+			check(InStruct);
+			const UPackage* Package = InStruct->GetPackage();
+			if (Package->GetName().Contains(TEXT("EngineTest")))
+			{
+				return false;
+			}
+			return true;
+		});
+			
+		if (JsonObject)
+		{
+			FString Text;
+			RigVMJsonUtils::FJsonWriterRef Writer = RigVMJsonUtils::MakeJsonWriter(Text);
+			if (FJsonSerializer::Serialize(JsonObject, Writer))
+			{
+				UE_LOG(LogRigVMDeveloper, Display, TEXT("Json:\n\n%s"), *Text);
+			}
+		}
+	})
+);
+
+FAutoConsoleCommand FCmdRigVMJsonPrintGraph
+(
+	TEXT("RigVM.JSON.PrintGraph"),
+	TEXT("Prints out a json string representing a rigvm graph."),
+	FConsoleCommandWithArgsDelegate::CreateLambda([](const TArray<FString>& Args)
+	{
+		if (Args.Num() != 1)
+		{
+			UE_LOG(LogRigVMDeveloper, Error, TEXT("Missing graph pathname argument for RigVM.JSON.PrintGraph"));
+			return;
+		}
+
+		FString GraphPathName = Args[0];
+		UObject* GraphObject = StaticLoadObject(URigVMGraph::StaticClass(), nullptr, GraphPathName);
+		if (!GraphObject)
+		{
+			UE_LOG(LogRigVMDeveloper, Error, TEXT("Cannot find graph '%s' for RigVM.JSON.PrintGraph"), *GraphPathName);
+			return;
+		}
+			
+		FRigVMRegistryReadLock Registry;
+		if (TSharedPtr<FJsonObject> JsonObject = RigVMJsonUtils::ToJson(Registry, CastChecked<URigVMGraph>(GraphObject)))
+		{
+			FString Text;
+			RigVMJsonUtils::FJsonWriterRef Writer = RigVMJsonUtils::MakeJsonWriter(Text);
+			if (FJsonSerializer::Serialize(JsonObject, Writer))
+			{
+				UE_LOG(LogRigVMDeveloper, Display, TEXT("Json:\n\n%s"), *Text);
+			}
+		}
+	})
+);
+
+FAutoConsoleCommand FCmdRigVMJsonPrintAsset
+(
+	TEXT("RigVM.JSON.PrintAsset"),
+	TEXT("Prints out a json string representing a rigvm asset."),
+	FConsoleCommandWithArgsDelegate::CreateLambda([](const TArray<FString>& Args)
+	{
+		if (Args.Num() != 1)
+		{
+			UE_LOG(LogRigVMDeveloper, Error, TEXT("Missing asset pathname argument for RigVM.JSON.PrintAsset"));
+			return;
+		}
+
+		FString AssetPathName = Args[0];
+		UObject* AssetObject = StaticLoadObject(UObject::StaticClass(), nullptr, AssetPathName);
+		if (!AssetObject)
+		{
+			UE_LOG(LogRigVMDeveloper, Error, TEXT("Cannot find asset '%s' for RigVM.JSON.PrintAsset"), *AssetPathName);
+			return;
+		}
+
+		const IRigVMAssetInterface* Asset = Cast<IRigVMAssetInterface>(AssetObject);
+		if (!Asset)
+		{
+			UE_LOG(LogRigVMDeveloper, Error, TEXT("Object '%s' exists, but it's not a RigVMAsset (when running RigVM.JSON.PrintAsset)"), *AssetPathName);
+			return;
+		}
+
+		FRigVMRegistryReadLock Registry;
+		if (TSharedPtr<FJsonObject> JsonObject = RigVMJsonUtils::ToJson(Registry, Asset))
+		{
+			FString Text;
+			RigVMJsonUtils::FJsonWriterRef Writer = RigVMJsonUtils::MakeJsonWriter(Text);
+			if (FJsonSerializer::Serialize(JsonObject, Writer))
+			{
+				UE_LOG(LogRigVMDeveloper, Display, TEXT("Json:\n\n%s"), *Text);
+			}
+		}
+	})
+);
+
+RigVMJsonUtils::FJsonWriterRef RigVMJsonUtils::MakeJsonWriter(FString& OutJsonText)
+{
+	return FJsonFactory::Create(&OutJsonText);
+}
+
+TSharedPtr<FJsonObject> RigVMJsonUtils::AvailableNodesToJson(FRigVMRegistryHandle& InRegistry, TFunction<bool(const UScriptStruct*)> InFilterCallback)
+{
+	if (!InRegistry.IsValid())
+	{
+		return nullptr;
+	}
+	
+	TArray<TSharedPtr<FJsonValue>> FunctionJsonValues, FactoryJsonValues;
+	const TChunkedArray<FRigVMFunction>& Functions = InRegistry->GetFunctions_NoLock();
+	const TArray<FRigVMDispatchFactory*>& Factories = InRegistry->GetFactories_NoLock();
+
+	for (const FRigVMFunction& Function : Functions)
+	{
+		if (Function.Factory)
+		{
+			continue;
+		}
+		check(Function.Struct);
+		if (InFilterCallback && !InFilterCallback(Function.Struct))
+		{
+			continue;
+		}
+		if (const TSharedPtr<FJsonObject> JsonObject = ToJson(InRegistry, &Function))
+		{
+			FunctionJsonValues.Add(MakeShared<FJsonValueObject>(JsonObject));
+		}
+	}
+
+	for (const FRigVMDispatchFactory* Factory : Factories)
+	{
+		check(Factory);
+		if (InFilterCallback && !InFilterCallback(Factory->GetScriptStruct()))
+		{
+			continue;
+		}
+		if (const TSharedPtr<FJsonObject> JsonObject = ToJson(InRegistry, Factory))
+		{
+			FactoryJsonValues.Add(MakeShared<FJsonValueObject>(JsonObject));
+		}
+	}
+
+	TSharedPtr<FJsonObject> JsonObject = MakeShared<FJsonObject>();
+	JsonObject->SetArrayField(TEXT("Functions"), FunctionJsonValues);
+	JsonObject->SetArrayField(TEXT("Dispatches"), FactoryJsonValues);
+	return JsonObject;
+}
+
+TSharedPtr<FJsonObject> RigVMJsonUtils::ToJson(FRigVMRegistryHandle& InRegistry, const FRigVMFunction* InFunction)
+{
+	check(InFunction);
+	const UScriptStruct* Struct = InFunction->Struct;
+	check(Struct);
+
+	TSharedPtr<FJsonObject> JsonObject = MakeShared<FJsonObject>();
+	GenericStructAttributesToJson(InRegistry, Struct, JsonObject.Get());
+
+	FString TemplateName;
+	if (const FRigVMTemplate* Template = InFunction->GetTemplate_NoLock(InRegistry))
+	{
+		TemplateName = Template->GetName().ToString();
+	}
+	JsonObject->SetStringField(TEXT("TemplateName"), TemplateName);
+
+	TArray<TSharedPtr<FJsonValue>> IO, Inputs, Outputs;
+
+	FStructOnScope StructOnScope(Struct);
+	FRigVMStruct* StructInstance = reinterpret_cast<FRigVMStruct*>(StructOnScope.GetStructMemory());
+
+	TMap<FName,const FRigVMFunctionArgument*> ArgumentLookup;
+	for (int32 Index = 0; Index < InFunction->Arguments.Num(); ++Index)
+	{
+		ArgumentLookup.Add(InFunction->Arguments[Index].Name, &InFunction->Arguments[Index]);
+	}
+	
+	for (TFieldIterator<FProperty> PropertyIt(Struct); PropertyIt; ++PropertyIt)
+	{
+		const FProperty* Property = *PropertyIt;
+		bool bIsExecutePin = false;
+		if (const FStructProperty* StructProperty = CastField<FStructProperty>(Property))
+		{
+			bIsExecutePin = StructProperty->Struct->IsChildOf(FRigVMExecutePin::StaticStruct());
+		}
+		static const FString ExecutePinType = FRigVMExecuteContext::StaticStruct()->GetStructCPPName();
+		
+		const FRigVMFunctionArgument** ArgumentPtr = ArgumentLookup.Find(Property->GetFName());
+		const FRigVMFunctionArgument* Argument = ArgumentPtr ? *ArgumentPtr : nullptr;
+
+		if (!Argument && !bIsExecutePin)
+		{
+			continue;
+		}
+
+		bool bIsInput = Argument ? Argument->Direction == ERigVMFunctionArgumentDirection::Input : true;
+		bool bIsOutput = Argument ? Argument->Direction == ERigVMFunctionArgumentDirection::Output : true;
+		
+#if WITH_EDITOR
+		bIsInput =
+			Property->HasMetaData(FRigVMStruct::InputMetaName) ||
+			Property->HasMetaData(FRigVMStruct::VisibleMetaName);
+		bIsOutput = Property->HasMetaData(FRigVMStruct::OutputMetaName);
+#endif
+
+		if (!bIsInput && !bIsOutput)
+		{
+			continue;
+		}
+
+		FString DefaultValue;
+		if (!bIsExecutePin)
+		{
+			const uint8* Memory = Property->ContainerPtrToValuePtr<uint8>(StructInstance);
+			Property->ExportTextItem_Direct(DefaultValue, Memory, Memory, nullptr, PPF_None, nullptr);
+			if (DefaultValue.IsEmpty())
+			{
+				if (Property->IsA<FArrayProperty>() || Property->IsA<FStructProperty>())
+				{
+					DefaultValue = TEXT("()");
+				}
+			}
+		}
+
+		TSharedPtr<FJsonObject> ArgumentJsonObject = MakeShared<FJsonObject>();
+		ArgumentJsonObject->SetStringField(TEXT("Name"), Property->GetName());
+		ArgumentJsonObject->SetStringField(TEXT("Type"), Argument ? FString(Argument->Type) : ExecutePinType);
+		ArgumentJsonObject->SetStringField(TEXT("PropertyToolTip"), Property->GetToolTipText().ToString());
+		ArgumentJsonObject->SetStringField(TEXT("DefaultValue"), DefaultValue);
+
+		if (bIsInput && bIsOutput)
+		{
+			IO.Add(MakeShared<FJsonValueObject>(ArgumentJsonObject));
+		}
+		else if (bIsInput)
+		{
+			Inputs.Add(MakeShared<FJsonValueObject>(ArgumentJsonObject));
+		}
+		else // if bIsOutput
+		{
+			Outputs.Add(MakeShared<FJsonValueObject>(ArgumentJsonObject));
+		}
+	}
+	
+	JsonObject->SetArrayField(TEXT("IO"), IO);
+	JsonObject->SetArrayField(TEXT("Inputs"), Inputs);
+	JsonObject->SetArrayField(TEXT("Outputs"), Outputs);
+
+	return JsonObject;
+}
+
+TSharedPtr<FJsonObject> RigVMJsonUtils::ToJson(FRigVMRegistryHandle& InRegistry, const FRigVMDispatchFactory* InFactory)
+{
+	check(InFactory);
+	const UScriptStruct* Struct = InFactory->GetScriptStruct();
+	check(Struct);
+
+	TSharedPtr<FJsonObject> JsonObject = MakeShared<FJsonObject>();
+	GenericStructAttributesToJson(InRegistry, Struct, JsonObject.Get());
+
+	const FRigVMTemplate* Template = InFactory->GetTemplate_NoLock(InRegistry);
+	JsonObject->SetStringField(TEXT("TemplateName"), Template->GetName().ToString());
+
+	FRigVMDispatchContext DispatchContext;
+	TArray<TSharedPtr<FJsonValue>> IO, Inputs, Outputs;
+
+	TArray<FRigVMExecuteArgument> ExecuteArguments = InFactory->GetExecuteArguments_NoLock(DispatchContext, InRegistry);
+	for (const FRigVMExecuteArgument& ExecuteArgument : ExecuteArguments)
+	{
+		static const FString ExecutePinType = FRigVMExecuteContext::StaticStruct()->GetStructCPPName();
+
+		TSharedPtr<FJsonObject> ArgumentJsonObject = MakeShared<FJsonObject>();
+		ArgumentJsonObject->SetStringField(TEXT("Name"), ExecuteArgument.Name.ToString());
+		ArgumentJsonObject->SetStringField(TEXT("Type"), ExecutePinType);
+		ArgumentJsonObject->SetStringField(TEXT("PropertyToolTip"), FString::Printf(TEXT("The %s execute argument."), *ExecuteArgument.Name.ToString()));
+		ArgumentJsonObject->SetStringField(TEXT("DefaultValue"), FString());
+
+		if (ExecuteArgument.Direction == ERigVMPinDirection::IO)
+		{
+			IO.Add(MakeShared<FJsonValueObject>(ArgumentJsonObject));
+		}
+		else if (ExecuteArgument.Direction == ERigVMPinDirection::Input)
+		{
+			Inputs.Add(MakeShared<FJsonValueObject>(ArgumentJsonObject));
+		}
+		else // if bIsOutput
+		{
+			Outputs.Add(MakeShared<FJsonValueObject>(ArgumentJsonObject));
+		}
+	}
+
+	const TArray<FRigVMTemplateArgumentInfo>& ArgumentInfos = InFactory->GetArgumentInfos(InRegistry);
+	for (const FRigVMTemplateArgumentInfo& ArgumentInfo : ArgumentInfos)
+	{
+		if (ArgumentInfo.Direction != ERigVMPinDirection::Input &&
+			ArgumentInfo.Direction != ERigVMPinDirection::Visible &&
+			ArgumentInfo.Direction != ERigVMPinDirection::IO &&
+			ArgumentInfo.Direction != ERigVMPinDirection::Output)
+		{
+			continue;
+		}
+		
+		const FRigVMTemplateArgument& Argument = ArgumentInfo.GetArgument_NoLock(InRegistry);
+		const TArray<TRigVMTypeIndex>& TypeIndices = Argument.GetSupportedTypeIndices_NoLock({}, InRegistry);
+		check(!TypeIndices.IsEmpty());
+
+		FString Type;
+		if (TypeIndices.Num() == 1)
+		{
+			Type = InRegistry->GetType_NoLock(TypeIndices[0]).CPPType.ToString();
+		}
+		else if (TypeIndices.Num() < 16)
+		{
+			TArray<FString> Types;
+			Types.Reserve(TypeIndices.Num());
+			for (const TRigVMTypeIndex& TypeIndex : TypeIndices)
+			{
+				Types.Add(InRegistry->GetType_NoLock(TypeIndex).CPPType.ToString());
+			}
+			Type = RigVMStringUtils::JoinStrings(Types, TEXT(","));
+		}
+		else
+		{
+			static const FString WildcardType = FRigVMUnknownType::StaticStruct()->GetStructCPPName();
+			static const FString WildcardArrayType = RigVMTypeUtils::ArrayTypeFromBaseType(WildcardType);
+
+			int32 NumArrayTypes = 0;
+			for (const TRigVMTypeIndex& TypeIndex : TypeIndices)
+			{
+				if (InRegistry->IsArrayType_NoLock(TypeIndex))
+				{
+					NumArrayTypes++;
+				}
+			}
+			if (NumArrayTypes == TypeIndices.Num())
+			{
+				Type = WildcardArrayType;
+			}
+			else
+			{
+				Type = WildcardType;
+			}
+		}
+		
+		TSharedPtr<FJsonObject> ArgumentJsonObject = MakeShared<FJsonObject>();
+		ArgumentJsonObject->SetStringField(TEXT("Name"), ArgumentInfo.Name.ToString());
+		ArgumentJsonObject->SetStringField(TEXT("Type"), Type);
+		ArgumentJsonObject->SetStringField(TEXT("PropertyToolTip"), InFactory->GetArgumentTooltip(ArgumentInfo.Name, TypeIndices[0]).ToString());
+		ArgumentJsonObject->SetStringField(TEXT("DefaultValue"), FString());
+
+		if (Argument.GetDirection() == ERigVMPinDirection::IO)
+		{
+			IO.Add(MakeShared<FJsonValueObject>(ArgumentJsonObject));
+		}
+		else if (Argument.GetDirection() == ERigVMPinDirection::Input)
+		{
+			Inputs.Add(MakeShared<FJsonValueObject>(ArgumentJsonObject));
+		}
+		else // if bIsOutput
+		{
+			Outputs.Add(MakeShared<FJsonValueObject>(ArgumentJsonObject));
+		}
+	}
+
+	JsonObject->SetArrayField(TEXT("IO"), IO);
+	JsonObject->SetArrayField(TEXT("Inputs"), Inputs);
+	JsonObject->SetArrayField(TEXT("Outputs"), Outputs);
+
+	return JsonObject;
+}
+
+void RigVMJsonUtils::GenericStructAttributesToJson(FRigVMRegistryHandle& InRegistry, const UScriptStruct* InStruct, FJsonObject* OutJsonObject)
+{
+	check(InStruct);
+	check(OutJsonObject);
+	
+	FString CategoryMetadata, KeywordsMetadata, DisplayNameMetadata;
+#if WITH_EDITOR
+	InStruct->GetStringMetaDataHierarchical(FRigVMStruct::CategoryMetaName, &CategoryMetadata);
+	InStruct->GetStringMetaDataHierarchical(FRigVMStruct::KeywordsMetaName, &KeywordsMetadata);
+	InStruct->GetStringMetaDataHierarchical(FRigVMStruct::DisplayNameMetaName, &DisplayNameMetadata);
+#endif
+	if(DisplayNameMetadata.IsEmpty())
+	{
+		DisplayNameMetadata = InStruct->GetDisplayNameText().ToString();
+	}
+
+	OutJsonObject->SetStringField(TEXT("DisplayName"), DisplayNameMetadata);
+	OutJsonObject->SetStringField(TEXT("Type"), InStruct->GetStructCPPName());
+	OutJsonObject->SetStringField(TEXT("Category"), CategoryMetadata);
+	OutJsonObject->SetStringField(TEXT("Keywords"), KeywordsMetadata);
+	OutJsonObject->SetStringField(TEXT("ToolTip"), InStruct->GetToolTipText().ToString());
+	OutJsonObject->SetStringField(TEXT("PackageName"), InStruct->GetOutermost()->GetPathName());
+}
+
+TSharedPtr<FJsonObject> RigVMJsonUtils::ToJson(FRigVMRegistryHandle& InRegistry, const IRigVMAssetInterface* InAsset)
+{
+	check(InAsset);
+	
+	TArray<TSharedPtr<FJsonValue>> Variables, GraphValues;
+
+	TArray<FRigVMGraphVariableDescription> VariableDescriptions = InAsset->GetAssetVariables();
+	for (const FRigVMGraphVariableDescription& VariableDescription : VariableDescriptions)
+	{
+		if (TSharedPtr<FJsonObject> VariableJsonObject = ToJson(InRegistry, VariableDescription))
+		{
+			Variables.Add(MakeShared<FJsonValueObject>(VariableJsonObject));
+		}
+	}
+
+	if (const FRigVMClient* Client = InAsset->GetRigVMClient())
+	{
+		const TArray<URigVMGraph*> Graphs =Client->GetAllModels(true, false);
+		for (const URigVMGraph* Graph : Graphs)
+		{
+			if (TSharedPtr<FJsonObject> GraphJsonObject = ToJson(InRegistry, Graph))
+			{
+				GraphValues.Add(MakeShared<FJsonValueObject>(GraphJsonObject));
+			}
+		}
+	}
+
+	TSharedPtr<FJsonObject> JsonObject = MakeShared<FJsonObject>();
+
+	if (const UObject* AssetObject = InAsset->GetObject())
+	{
+		JsonObject->SetStringField(TEXT("Name"), AssetObject->GetName());
+		JsonObject->SetStringField(TEXT("PackageName"), AssetObject->GetOutermost()->GetPathName());
+	}
+	JsonObject->SetArrayField(TEXT("Variables"), Variables);
+	JsonObject->SetArrayField(TEXT("Graphs"), GraphValues);
+	return JsonObject;
+}
+
+TSharedPtr<FJsonObject> RigVMJsonUtils::ToJson(FRigVMRegistryHandle& InRegistry, const URigVMGraph* InGraph, bool bOnlySelection)
+{
+	check(InGraph);
+
+	TSet<FName> ExportedNodes;
+	TArray<TSharedPtr<FJsonValue>> Variables, Nodes, Links;
+
+	if(!InGraph->IsTopLevelGraph())
+	{
+		TArray<FRigVMGraphVariableDescription> VariableDescriptions = InGraph->GetLocalVariables(false);
+		for (const FRigVMGraphVariableDescription& VariableDescription : VariableDescriptions)
+		{
+			if (TSharedPtr<FJsonObject> VariableJsonObject = ToJson(InRegistry, VariableDescription))
+			{
+				Variables.Add(MakeShared<FJsonValueObject>(VariableJsonObject));
+			}
+		}
+	}
+
+	const TArray<FName> SelectedNodes = InGraph->GetSelectNodes(); 
+	for (const URigVMNode* Node : InGraph->GetNodes())
+	{
+		if (bOnlySelection)
+		{
+			if (!SelectedNodes.Contains(Node->GetName()))
+			{
+				continue;
+			}
+		}
+		
+		if (TSharedPtr<FJsonObject> NodeJsonObject = ToJson(InRegistry, Node))
+		{
+			Nodes.Add(MakeShared<FJsonValueObject>(NodeJsonObject));
+			ExportedNodes.Add(Node->GetFName());
+		}
+	}
+	
+	for (const URigVMLink* Link : InGraph->GetLinks())
+	{
+		const URigVMNode* SourceNode = Link->GetSourceNode();
+		const URigVMNode* TargetNode = Link->GetTargetNode();
+		if (!SourceNode || !TargetNode)
+		{
+			continue;
+		}
+		if (!ExportedNodes.Contains(SourceNode->GetFName()) || !ExportedNodes.Contains(TargetNode->GetFName()))
+		{
+			continue;
+		}
+		if (TSharedPtr<FJsonObject> LinkJsonObject = ToJson(InRegistry, Link))
+		{
+			Links.Add(MakeShared<FJsonValueObject>(LinkJsonObject));
+		}
+	}
+
+	TSharedPtr<FJsonObject> JsonObject = MakeShared<FJsonObject>();
+
+	if (InGraph->IsA<URigVMFunctionLibrary>())
+	{
+		JsonObject->SetStringField(TEXT("Name"), TEXT("FunctionLibrary"));
+		JsonObject->SetArrayField(TEXT("GraphFunctions"), Nodes);
+	}
+	else
+	{
+		JsonObject->SetStringField(TEXT("Name"), InGraph->GetName());
+		JsonObject->SetArrayField(TEXT("LocalVariables"), Variables);
+		JsonObject->SetArrayField(TEXT("Nodes"), Nodes);
+		JsonObject->SetArrayField(TEXT("Links"), Links);
+	}
+	return JsonObject;
+}
+
+TSharedPtr<FJsonObject> RigVMJsonUtils::ToJson(FRigVMRegistryHandle& InRegistry, const URigVMNode* InNode)
+{
+	check(InNode);
+
+	TArray<TSharedPtr<FJsonValue>> Pins;
+	for (const URigVMPin* Pin : InNode->GetPins())
+	{
+		if (Pin->GetDirection() != ERigVMPinDirection::Input &&
+			Pin->GetDirection() != ERigVMPinDirection::Visible &&
+			Pin->GetDirection() != ERigVMPinDirection::IO &&
+			Pin->GetDirection() != ERigVMPinDirection::Output)
+		{
+			continue;
+		}
+
+		if (TSharedPtr<FJsonObject> PinJsonObject = ToJson(InRegistry, Pin))
+		{
+			Pins.Add(MakeShared<FJsonValueObject>(PinJsonObject));
+		}
+	}
+	
+	TSharedPtr<FJsonObject> JsonObject = MakeShared<FJsonObject>();
+	JsonObject->SetStringField(TEXT("Name"), InNode->GetName());
+	JsonObject->SetStringField(TEXT("Title"), InNode->GetNodeTitle());
+	JsonObject->SetStringField(TEXT("Type"), InNode->GetClass()->GetName());
+
+	if (!InNode->IsInjected())
+	{
+		const FVector2D Position = InNode->GetPosition();
+		JsonObject->SetStringField(TEXT("Position"), StructValueToJson(TBaseStructure<FVector2D>::Get(), &Position));
+	}
+
+	if (const URigVMUnitNode* UnitNode = Cast<URigVMUnitNode>(InNode))
+	{
+		JsonObject->SetStringField(TEXT("Function"), UnitNode->GetScriptStruct()->GetStructCPPName());
+	}
+	else if (const URigVMDispatchNode* DispatchNode = Cast<URigVMDispatchNode>(InNode))
+	{
+		JsonObject->SetStringField(TEXT("Factory"), DispatchNode->GetFactory()->GetScriptStruct()->GetStructCPPName());
+	}
+	else if (const URigVMVariableNode* VariableNode = Cast<URigVMVariableNode>(InNode))
+	{
+		JsonObject->SetStringField(TEXT("Variable"), VariableNode->GetVariableName().ToString());
+	}
+	else if (const URigVMFunctionReferenceNode* FunctionReferenceNode = Cast<URigVMFunctionReferenceNode>(InNode))
+	{
+		FString LibraryNodePath = FunctionReferenceNode->GetFunctionIdentifier().GetLibraryNodePath();
+		if (const IRigVMAssetInterface* Asset = FunctionReferenceNode->GetImplementingOuter<IRigVMAssetInterface>())
+		{
+			if (const FRigVMClient* Client = Asset->GetRigVMClient())
+			{
+				if (const URigVMFunctionLibrary* FunctionLibrary = Client->GetFunctionLibrary())
+				{
+					const FString FunctionLibraryPrefix = FunctionLibrary->GetPathName() + TEXT(".");
+					if (LibraryNodePath.StartsWith(FunctionLibraryPrefix))
+					{
+						LibraryNodePath = LibraryNodePath.Mid(FunctionLibraryPrefix.Len());
+					}
+				}
+			}
+		}
+		
+		JsonObject->SetStringField(TEXT("GraphFunction"), LibraryNodePath);
+	}
+	else if (const URigVMCollapseNode* CollapseNode = Cast<URigVMCollapseNode>(InNode))
+	{
+		if (TSharedPtr<FJsonObject> ContainedGraphJsonObject = ToJson(InRegistry, CollapseNode->GetContainedGraph()))
+		{
+			if (CollapseNode->GetGraph()->IsA<URigVMFunctionLibrary>())
+			{
+				JsonObject->SetStringField(TEXT("Category"), CollapseNode->GetNodeCategory());
+				JsonObject->SetStringField(TEXT("Keywords"), CollapseNode->GetNodeKeywords());
+				JsonObject->SetStringField(TEXT("Description"), CollapseNode->GetNodeDescription());
+			}
+			JsonObject->SetObjectField(TEXT("Graph"), ContainedGraphJsonObject);
+		}
+	}
+	else if (const URigVMCommentNode* CommentNode = Cast<URigVMCommentNode>(InNode))
+	{
+		const FVector2D Size = InNode->GetSize();
+		const FLinearColor Color = InNode->GetNodeColor();
+		JsonObject->SetStringField(TEXT("Size"), StructValueToJson(TBaseStructure<FVector2D>::Get(), &Size));
+		JsonObject->SetStringField(TEXT("Color"), StructValueToJson(TBaseStructure<FLinearColor>::Get(), &Color));
+		JsonObject->SetStringField(TEXT("CommentText"), CommentNode->GetCommentText());
+	}
+
+	JsonObject->SetArrayField(TEXT("Pins"), Pins);
+	return JsonObject;
+}
+
+TSharedPtr<FJsonObject> RigVMJsonUtils::ToJson(FRigVMRegistryHandle& InRegistry, const URigVMPin* InPin)
+{
+	check(InPin);
+	
+	TSharedPtr<FJsonObject> JsonObject = MakeShared<FJsonObject>();
+	JsonObject->SetStringField(TEXT("Name"), InPin->GetName());
+	JsonObject->SetStringField(TEXT("Type"), InPin->GetCPPType());
+	JsonObject->SetStringField(TEXT("DefaultValue"), InPin->GetDefaultValue());
+
+	const TArray<URigVMInjectionInfo*> InjectedNodes = InPin->GetInjectedNodes();
+	check(InjectedNodes.Num() < 2);
+	if (!InjectedNodes.IsEmpty())
+	{
+		JsonObject->SetStringField(TEXT("InjectedNode"), InjectedNodes[0]->Node->GetName());
+	}
+	
+	return JsonObject;
+}
+
+TSharedPtr<FJsonObject> RigVMJsonUtils::ToJson(FRigVMRegistryHandle& InRegistry, const URigVMLink* InLink)
+{
+	check(InLink);
+
+	const URigVMPin* SourcePin = InLink->GetSourcePin();
+	const URigVMPin* TargetPin = InLink->GetTargetPin();
+	if (!SourcePin || !TargetPin)
+	{
+		return nullptr;
+	}
+
+	TSharedPtr<FJsonObject> JsonObject = MakeShared<FJsonObject>();
+	JsonObject->SetStringField(TEXT("Source"), SourcePin->GetPinPath());
+	JsonObject->SetStringField(TEXT("Target"), TargetPin->GetPinPath());
+	return JsonObject;
+}
+
+TSharedPtr<FJsonObject> RigVMJsonUtils::ToJson(FRigVMRegistryHandle& InRegistry, const FRigVMGraphVariableDescription& InVariable)
+{
+	if (InVariable.Name.IsNone())
+	{
+		return nullptr;
+	}
+
+	static const TCHAR* True = TEXT("True");
+	static const TCHAR* False = TEXT("False");
+
+	TSharedPtr<FJsonObject> JsonObject = MakeShared<FJsonObject>();
+	JsonObject->SetStringField(TEXT("Name"), InVariable.Name.ToString());
+	JsonObject->SetStringField(TEXT("TypeName"), InVariable.CPPType);
+	JsonObject->SetStringField(TEXT("TypePath"), InVariable.CPPTypeObjectPath.IsNone() ? FString() : InVariable.CPPTypeObjectPath.ToString());
+	JsonObject->SetStringField(TEXT("DefaultValue"), InVariable.DefaultValue);
+	JsonObject->SetStringField(TEXT("Category"), InVariable.Category.ToString());
+	JsonObject->SetStringField(TEXT("Tooltip"), InVariable.Tooltip.ToString());
+	JsonObject->SetStringField(TEXT("ExposedOnSpawn"), InVariable.bExposedOnSpawn ? True : False);
+	JsonObject->SetStringField(TEXT("ExposeToCinematics"), InVariable.bExposeToCinematics ? True : False);
+	JsonObject->SetStringField(TEXT("Public"), InVariable.bPublic ? True : False);
+	JsonObject->SetStringField(TEXT("Private"), InVariable.bPrivate ? True : False);
+	return JsonObject;
+}
+
+FString RigVMJsonUtils::StructValueToJson(const UScriptStruct* InScriptStruct, const void* InMemory)
+{
+	check(InScriptStruct);
+	FString Value;
+	InScriptStruct->ExportText(Value, InMemory, InMemory, nullptr, PPF_ExternalEditor, nullptr);
+	return Value;
+}
+
+#undef LOCTEXT_NAMESPACE
diff --git a/Engine/Plugins/Runtime/RigVM/Source/RigVMDeveloper/Public/RigVMJsonUtils.h b/Engine/Plugins/Runtime/RigVM/Source/RigVMDeveloper/Public/RigVMJsonUtils.h
new file mode 100644
--- /dev/null
+++ b/Engine/Plugins/Runtime/RigVM/Source/RigVMDeveloper/Public/RigVMJsonUtils.h
@@ -0,0 +1,50 @@
+﻿// Copyright Epic Games, Inc. All Rights Reserved.
+
+#pragma once
+
+#include "CoreMinimal.h"
+#include "Dom/JsonObject.h"
+#include "Policies/PrettyJsonPrintPolicy.h"
+#include "Serialization/JsonWriter.h"
+#include "RigVMCore/RigVMRegistry.h"
+#include "RigVMCore/RigVMFunction.h"
+#include "RigVMCore/RigVMDispatchFactory.h"
+
+class IRigVMAssetInterface;
+class URigVMGraph;
+class URigVMNode;
+class URigVMPin;
+class URigVMLink;
+struct FRigVMGraphVariableDescription;
+
+namespace RigVMJsonUtils
+{
+	template <class CharType>
+	struct TRigVMJsonPrintPolicy
+	: public TPrettyJsonPrintPolicy<CharType>
+	{
+	};
+
+	typedef TJsonWriter< TCHAR, TRigVMJsonPrintPolicy< TCHAR > > FJsonWriter; 
+	typedef TSharedRef< FJsonWriter > FJsonWriterRef;
+	typedef TJsonWriterFactory< TCHAR, TRigVMJsonPrintPolicy< TCHAR > > FJsonFactory; 
+
+	RIGVMDEVELOPER_API FJsonWriterRef MakeJsonWriter(FString& OutJsonText);
+
+	// introspection for the registry
+	RIGVMDEVELOPER_API TSharedPtr<FJsonObject> AvailableNodesToJson(FRigVMRegistryHandle& InRegistry, TFunction<bool(const UScriptStruct*)> InFilterCallback);
+	RIGVMDEVELOPER_API TSharedPtr<FJsonObject> ToJson(FRigVMRegistryHandle& InRegistry, const FRigVMFunction* InFunction);
+	RIGVMDEVELOPER_API TSharedPtr<FJsonObject> ToJson(FRigVMRegistryHandle& InRegistry, const FRigVMDispatchFactory* InFactory);
+	RIGVMDEVELOPER_API void GenericStructAttributesToJson(FRigVMRegistryHandle& InRegistry, const UScriptStruct* InStruct, FJsonObject* OutJsonObject);
+
+	// introspection for existing graphs
+	RIGVMDEVELOPER_API TSharedPtr<FJsonObject> ToJson(FRigVMRegistryHandle& InRegistry, const IRigVMAssetInterface* InAsset);
+	RIGVMDEVELOPER_API TSharedPtr<FJsonObject> ToJson(FRigVMRegistryHandle& InRegistry, const URigVMGraph* InGraph, bool bOnlySelection = false);
+	RIGVMDEVELOPER_API TSharedPtr<FJsonObject> ToJson(FRigVMRegistryHandle& InRegistry, const URigVMNode* InNode);
+	RIGVMDEVELOPER_API TSharedPtr<FJsonObject> ToJson(FRigVMRegistryHandle& InRegistry, const URigVMPin* InPin);
+	RIGVMDEVELOPER_API TSharedPtr<FJsonObject> ToJson(FRigVMRegistryHandle& InRegistry, const URigVMLink* InLink);
+	RIGVMDEVELOPER_API TSharedPtr<FJsonObject> ToJson(FRigVMRegistryHandle& InRegistry, const FRigVMGraphVariableDescription& InVariable);
+
+	// string conversion
+	RIGVMDEVELOPER_API FString StructValueToJson(const UScriptStruct* InScriptStruct, const void* InMemory);
+}
diff --git a/Engine/Plugins/Runtime/RigVM/Source/RigVMDeveloper/RigVMDeveloper.Build.cs b/Engine/Plugins/Runtime/RigVM/Source/RigVMDeveloper/RigVMDeveloper.Build.cs
--- a/Engine/Plugins/Runtime/RigVM/Source/RigVMDeveloper/RigVMDeveloper.Build.cs
+++ b/Engine/Plugins/Runtime/RigVM/Source/RigVMDeveloper/RigVMDeveloper.Build.cs
@@ -17,6 +17,7 @@ public class RigVMDeveloper : ModuleRules
                 "VisualGraphUtils",
                 "KismetCompiler",
                 "Projects",
+				"Json",
             }
         );
 
