RigVM: Scribble Integration

#rnx

diff --git a/Engine/Plugins/Runtime/RigVM/RigVM.uplugin b/Engine/Plugins/Runtime/RigVM/RigVM.uplugin
--- a/Engine/Plugins/Runtime/RigVM/RigVM.uplugin
+++ b/Engine/Plugins/Runtime/RigVM/RigVM.uplugin
@@ -37,6 +37,10 @@
 		{
 			"Name": "GameplayInsights",
 			"Enabled": true
+		},
+		{
+			"Name": "Scribble",
+			"Enabled": true
 		}
 	]
 }
\ No newline at end of file
diff --git a/Engine/Plugins/Runtime/RigVM/Source/RigVMDeveloper/Private/RigVMModel/RigVMGraph.cpp b/Engine/Plugins/Runtime/RigVM/Source/RigVMDeveloper/Private/RigVMModel/RigVMGraph.cpp
--- a/Engine/Plugins/Runtime/RigVM/Source/RigVMDeveloper/Private/RigVMModel/RigVMGraph.cpp
+++ b/Engine/Plugins/Runtime/RigVM/Source/RigVMDeveloper/Private/RigVMModel/RigVMGraph.cpp
@@ -15,6 +15,8 @@
 #include "RigVMSettings.h"
 #endif
 
+#include "ScribbleGraph.h"
+
 #include UE_INLINE_GENERATED_CPP_BY_NAME(RigVMGraph)
 
 URigVMGraph::URigVMGraph()
@@ -51,6 +53,8 @@ void URigVMGraph::PostLoad()
 			TargetPin->Links.AddUnique(Link);
 		}
 	}
+
+	SetupScribbleGraph();
 }
 
 const TArray<URigVMNode*>& URigVMGraph::GetNodes() const
@@ -614,6 +618,70 @@ bool URigVMGraph::IsNameAvailable(const FString& InName) const
 	return true;
 }
 
+TSharedPtr<FScribbleGraphData> URigVMGraph::SetupScribbleGraph(bool bCreateIfMissing)
+{
+	if (!ScribbleGraph.IsValid())
+	{
+		if (bCreateIfMissing)
+		{
+			Modify();
+			(void)ScribbleGraph.GetOrCreateDataPtr();
+		}
+		else
+		{
+			return nullptr;
+		}
+	}
+
+	TSharedPtr<FScribbleGraphData> Graph = ScribbleGraph.GetDataPtr();
+	if (!Graph.IsValid())
+	{
+		return nullptr;
+	}
+
+	Graph->OnModify().Clear();
+
+	TWeakObjectPtr<URigVMGraph> WeakThis = this;
+	Graph->OnModify().AddLambda([WeakThis]()
+	{
+		if (!WeakThis.IsValid())
+		{
+			return;
+		}
+		WeakThis->Modify();
+	});
+
+	Graph->OnGetCurrentAnchor().BindLambda([WeakThis]() -> FName
+	{
+		if (!WeakThis.IsValid())
+		{
+			return NAME_None;
+		}
+		const TArray<FName> SelectedNodes = WeakThis->GetSelectNodes();
+		if (SelectedNodes.IsEmpty())
+		{
+			return NAME_None;
+		}
+		return SelectedNodes[0];
+	});
+
+	Graph->OnResolveAnchor().BindLambda([WeakThis](const FName& InNodeName) -> TOptional<FVector2f>
+	{
+		if (!WeakThis.IsValid())
+		{
+			return {};
+		}
+		const URigVMNode* Node = WeakThis->FindNodeByName(InNodeName);
+		if (!Node)
+		{
+			return {};
+		}
+		return FVector2f(Node->GetPosition());
+	});
+
+	return Graph;
+}
+
 void URigVMGraph::PrepareCycleChecking(URigVMPin* InPin, bool bAsInput)
 {
 	TArray<URigVMLink*> LinksToSkip;
diff --git a/Engine/Plugins/Runtime/RigVM/Source/RigVMDeveloper/Public/RigVMModel/RigVMGraph.h b/Engine/Plugins/Runtime/RigVM/Source/RigVMDeveloper/Public/RigVMModel/RigVMGraph.h
--- a/Engine/Plugins/Runtime/RigVM/Source/RigVMDeveloper/Public/RigVMModel/RigVMGraph.h
+++ b/Engine/Plugins/Runtime/RigVM/Source/RigVMDeveloper/Public/RigVMModel/RigVMGraph.h
@@ -11,6 +11,7 @@
 #include "Nodes/RigVMFunctionReturnNode.h"
 #include "RigVMCompiler/RigVMAST.h"
 #include "UObject/Interface.h"
+#include "ScribbleGraph.h"
 #include "RigVMGraph.generated.h"
 
 #define UE_API RIGVMDEVELOPER_API
@@ -190,6 +191,8 @@ public:
 	UE_API TArray<FRigVMGraphSection> GetMatchingSections(const FRigVMGraphSection& InSection) const;
 	UE_API void UpdateSections(bool bForce = false);
 
+	UE_API TSharedPtr<FScribbleGraphData> SetupScribbleGraph(bool bCreateIfMissing = false);
+	
 private:
 
 	FRigVMGraphModifiedEvent ModifiedEvent;
@@ -244,6 +247,9 @@ private:
 
 	UE_API bool IsNameAvailable(const FString& InName) const;
 
+	UPROPERTY()
+	FScribbleGraph ScribbleGraph;
+
 	friend class URigVMController;
 	friend class URigVMBlueprint;
 	friend class IRigVMAssetInterface;
diff --git a/Engine/Plugins/Runtime/RigVM/Source/RigVMDeveloper/RigVMDeveloper.Build.cs b/Engine/Plugins/Runtime/RigVM/Source/RigVMDeveloper/RigVMDeveloper.Build.cs
--- a/Engine/Plugins/Runtime/RigVM/Source/RigVMDeveloper/RigVMDeveloper.Build.cs
+++ b/Engine/Plugins/Runtime/RigVM/Source/RigVMDeveloper/RigVMDeveloper.Build.cs
@@ -17,7 +17,8 @@ public class RigVMDeveloper : ModuleRules
                 "VisualGraphUtils",
                 "KismetCompiler",
                 "Projects",
-				"Json",
+                "Json",
+                "Scribble",
             }
         );
 
diff --git a/Engine/Plugins/Runtime/RigVM/Source/RigVMEditor/Private/Editor/RigVMLegacyEditor.cpp b/Engine/Plugins/Runtime/RigVM/Source/RigVMEditor/Private/Editor/RigVMLegacyEditor.cpp
--- a/Engine/Plugins/Runtime/RigVM/Source/RigVMEditor/Private/Editor/RigVMLegacyEditor.cpp
+++ b/Engine/Plugins/Runtime/RigVM/Source/RigVMEditor/Private/Editor/RigVMLegacyEditor.cpp
@@ -7,6 +7,7 @@
 #include "RigVMEditorModule.h"
 #include "SBlueprintEditorSelectedDebugObjectWidget.h"
 #include "RigVMEditorCommands.h"
+#include "SScribbleGraph.h"
 
 struct FRigVMEditorZoomLevelsContainer;
 
@@ -165,11 +166,42 @@ void FRigVMLegacyEditor::CreateDefaultCommands()
 TSharedRef<SGraphEditor> FRigVMLegacyEditor::CreateGraphEditorWidget(TSharedRef<FTabInfo> InTabInfo, UEdGraph* InGraph)
 {
 	TSharedRef<SGraphEditor> GraphEditor = FBlueprintEditor::CreateGraphEditorWidget(InTabInfo, InGraph);
-	GraphEditor->GetGraphPanel()->SetZoomLevelsContainer<FRigVMEditorZoomLevelsContainer>();
+		GraphEditor->GetGraphPanel()->SetZoomLevelsContainer<FRigVMEditorZoomLevelsContainer>();
 
 	TWeakObjectPtr<URigVMEdGraph> WeakGraph(Cast<URigVMEdGraph>(InGraph));
 	if (WeakGraph.IsValid())
 	{
+		if (URigVMGraph* Model = WeakGraph->GetModel())
+		{
+			if (CVarScribbleEnabled.GetValueOnAnyThread())
+			{
+				TSharedRef<SWidget> GraphPanelWidget = GraphEditor->GetGraphPanel()->AsShared();
+				if (TSharedPtr<SWidget> ParentWidget = GraphPanelWidget->GetParentWidget())
+				{
+					if (ParentWidget->GetWidgetClass().GetWidgetType() == SOverlay::StaticWidgetClass().GetWidgetType())
+					{
+						TSharedPtr<SOverlay> Overlay = StaticCastSharedPtr<SOverlay>(ParentWidget);
+						if (Overlay)
+						{
+							Overlay->AddSlot()
+							.Padding(0.f)
+							.VAlign(VAlign_Fill)
+							.HAlign(HAlign_Fill)
+							[
+								SNew(SScribbleGraph)
+								.GraphData(Model->SetupScribbleGraph(true))
+								.AllowNavigation(true)
+								.ShouldDrawBackground(false)
+								.ShouldDrawSurroundingShadow(false)
+								.ToolbarPadding(FMargin(8,44,0,0))
+								.PanelToSync(StaticCastWeakPtr<SGraphPanel,SWidget>(GraphPanelWidget.ToWeakPtr()))
+							];
+						}
+					}
+				}
+			}
+		}
+		
 		TWeakPtr<FRigVMEditorBase> WeakThisEditor = SharedThis(this).ToWeakPtr();
 		TWeakPtr<SGraphEditor> WeakGraphEditor = GraphEditor.ToWeakPtr();
 		
diff --git a/Engine/Plugins/Runtime/RigVM/Source/RigVMEditor/RigVMEditor.Build.cs b/Engine/Plugins/Runtime/RigVM/Source/RigVMEditor/RigVMEditor.Build.cs
--- a/Engine/Plugins/Runtime/RigVM/Source/RigVMEditor/RigVMEditor.Build.cs
+++ b/Engine/Plugins/Runtime/RigVM/Source/RigVMEditor/RigVMEditor.Build.cs
@@ -44,6 +44,8 @@ public class RigVMEditor : ModuleRules
                 "TraceServices",
                 "RewindDebuggerInterface",
                 "GameplayInsights",
+				"Scribble",
+				"ScribbleEditor",
             }
 		);
         
