[RigVM] Optimize unmounting, reduce core contention, expose API to batch unmounting calls together

#tests local build, replay runs
#rnx

diff --git a/Engine/Plugins/Runtime/RigVM/Source/RigVM/Private/RigVMCore/RigVMRegistry.cpp b/Engine/Plugins/Runtime/RigVM/Source/RigVM/Private/RigVMCore/RigVMRegistry.cpp
--- a/Engine/Plugins/Runtime/RigVM/Source/RigVM/Private/RigVMCore/RigVMRegistry.cpp
+++ b/Engine/Plugins/Runtime/RigVM/Source/RigVM/Private/RigVMCore/RigVMRegistry.cpp
@@ -2,6 +2,10 @@
 
 #include "RigVMCore/RigVMRegistry.h"
 #include "AssetRegistry/AssetData.h"
+#include "Containers/ContainerAllocationPolicies.h"
+#include "Containers/StaticBitArray.h"
+#include "Misc/CoreMiscDefines.h"
+#include "Misc/ScopeLock.h"
 #include "RigVMCore/RigVMStruct.h"
 #include <RigVMCore/RigVMTrait.h>
 #include "RigVMTypeUtils.h"
@@ -12,11 +16,13 @@
 #include "AssetRegistry/AssetRegistryModule.h"
 #include "Engine/UserDefinedEnum.h"
 #include "StructUtils/UserDefinedStruct.h"
+#include "UObject/NameTypes.h"
 #include "UObject/UObjectIterator.h"
 #include "Misc/CoreDelegates.h"
 #include "Misc/DelayedAutoRegister.h"
 #include "RigVMFunctions/RigVMDispatch_Core.h"
 #include "Interfaces/IPluginManager.h"
+#include <atomic>
 
 // When the object system has been completely loaded, load in all the engine types that we haven't registered already in InitializeIfNeeded 
 static FDelayedAutoRegisterHelper GRigVMRegistrySingletonHelper(EDelayedRegisterRunPhase::EndOfEngineInit, &FRigVMRegistry_NoLock::OnEngineInit);
@@ -918,12 +924,13 @@ bool FRigVMRegistry_NoLock::IsWithinPlugin(const TArray<FString>& InModulePrefix
 	{
 		return false;
 	}
-	
-	const FSoftObjectPath SoftObjectPath(InObject);
-	const FString ObjectPathString = SoftObjectPath.ToString();
+
+	FNameBuilder PackagePathStr;
+	FSoftObjectPath(InObject).AppendString(PackagePathStr);
+	const FStringView PackagePath = PackagePathStr.ToView();
 	for (const FString& Prefix : InModulePrefixes)
 	{
-		if (ObjectPathString.StartsWith(Prefix))
+		if (PackagePath.StartsWith(Prefix))
 		{
 			return true;
 		}
@@ -951,8 +958,29 @@ bool FRigVMRegistry_NoLock::OnModulesUnloaded_NoLock(TConstArrayView<FName> Modu
 	return OnModulesUnloadedInternal(ModulePrefixes);
 }
 
+struct alignas(PLATFORM_CACHE_LINE_SIZE) FNoFalseSharingBoolSet
+{
+	static constexpr int32 NumBitsPerBucket = PLATFORM_CACHE_LINE_SIZE * 8;
+	static int32 CalcNumBuckets(int32 TotalNumberOfElements)
+	{
+		return FMath::DivideAndRoundUp(TotalNumberOfElements, NumBitsPerBucket);
+	}
+	static int32 CalcNumBitsInTheBucket(int32 BucketIdx, int32 TotalNumberOfElements)
+	{
+		return FMath::Min(NumBitsPerBucket, TotalNumberOfElements - BucketIdx * NumBitsPerBucket);
+	}
+
+	TStaticBitArray<NumBitsPerBucket> Set;
+};
+static_assert(sizeof(FNoFalseSharingBoolSet) == PLATFORM_CACHE_LINE_SIZE, "Make sure that the FNoFalseSharingBoolSet fits inside the cache line");
+
 bool FRigVMRegistry_NoLock::OnModulesUnloadedInternal(const TArray<FString>& ModulePrefixes)
 {
+	if (ModulePrefixes.IsEmpty())
+	{
+		return false;
+	}
+
 	TRACE_CPUPROFILER_EVENT_SCOPE(FRigVMRegistry_NoLock::OnModulesUnloadedInternal);
 
 	bool bRegistryChanged = false;
@@ -1017,33 +1045,36 @@ bool FRigVMRegistry_NoLock::OnModulesUnloadedInternal(const TArray<FString>& Mod
 	{
 		TRACE_CPUPROFILER_EVENT_SCOPE(FRigVMRegistry_NoLock::OnModulesUnloadedInternal::Functions);
 
-		TArray<int32> FunctionIndicesToRemove;
-		ParallelFor(Functions.Num(), [this, &ModulePrefixes, &ParallelForDataLock, &FunctionIndicesToRemove](int32 Index)
+		TArray<FNoFalseSharingBoolSet> FunctionIndicesToRemove;
+		const int32 TotalNumberOfElements = Functions.Num();
+		const int32 NumBuckets = FNoFalseSharingBoolSet::CalcNumBuckets(TotalNumberOfElements);
+		FunctionIndicesToRemove.SetNum(NumBuckets);
+		ParallelFor(NumBuckets, [this, &ModulePrefixes, &FunctionIndicesToRemove, TotalNumberOfElements](int32 BucketIdx)
 		{
-			if (!Functions[Index].IsValid())
+			const int32 NumBitsInThisBucket = FNoFalseSharingBoolSet::CalcNumBitsInTheBucket(BucketIdx, TotalNumberOfElements);
+			for (int32 Index = 0; Index < NumBitsInThisBucket; ++Index)
 			{
-				return;
+				const int32 FuncIdx = Index + BucketIdx * FNoFalseSharingBoolSet::NumBitsPerBucket;
+				FunctionIndicesToRemove[BucketIdx].Set[Index] =
+					Functions[FuncIdx].IsValid()
+					&& Functions[FuncIdx].Struct != nullptr
+					&& IsWithinPlugin(ModulePrefixes, Functions[FuncIdx].Struct);
 			}
-
-			if (Functions[Index].Struct == nullptr)
-			{
-				return;
-			}
-
-			if (!IsWithinPlugin(ModulePrefixes, Functions[Index].Struct))
-			{
-				return;
-			}
-
-			FScopeLock Lock(&ParallelForDataLock);
-			FunctionIndicesToRemove.Add(Index);
 		});
 
-		for (int32 Index : FunctionIndicesToRemove)
+		for (int32 BucketIdx = 0; BucketIdx < NumBuckets; ++BucketIdx)
 		{
-			if (RemoveFunction_NoLock(Index))
+			const int32 NumBitsInThisBucket = FNoFalseSharingBoolSet::CalcNumBitsInTheBucket(BucketIdx, TotalNumberOfElements);
+			for (int32 Index = 0; Index < NumBitsInThisBucket; ++Index)
 			{
-				bRegistryChanged = true;
+				if (FunctionIndicesToRemove[BucketIdx].Set[Index])
+				{
+					const int32 FuncIdx = Index + BucketIdx * FNoFalseSharingBoolSet::NumBitsPerBucket;
+					if (RemoveFunction_NoLock(FuncIdx))
+					{
+						bRegistryChanged = true;
+					}
+				}
 			}
 		}
 	}
@@ -1052,36 +1083,39 @@ bool FRigVMRegistry_NoLock::OnModulesUnloadedInternal(const TArray<FString>& Mod
 	{
 		TRACE_CPUPROFILER_EVENT_SCOPE(FRigVMRegistry_NoLock::OnModulesUnloadedInternal::Types);
 
-		TArray<int32> TypeIndicesToRemove;
-		ParallelFor(Types.Num(), [this, &ModulePrefixes, &ParallelForDataLock, &TypeIndicesToRemove](int32 Index)
+		TArray<FNoFalseSharingBoolSet> TypeIndicesToRemove;
+		const int32 TotalNumberOfElements = Types.Num();
+		const int32 NumBuckets = FNoFalseSharingBoolSet::CalcNumBuckets(TotalNumberOfElements);
+		TypeIndicesToRemove.SetNum(NumBuckets);
+		ParallelFor(NumBuckets, [this, &ModulePrefixes, &TypeIndicesToRemove, TotalNumberOfElements](int32 BucketIdx)
 		{
-			if (Types[Index].IsUnknownType())
-			{
-				return;
-			}
-
-			// avoid array types, the base type removal will clean up related array types.
-			// For example removing FMyStruct will also clean up TArray<FMyStruct>.
-			// Note: There's always a base type registered for an array type being present.
-			if (Types[Index].bIsArray)
+			const int32 NumBitsInThisBucket = FNoFalseSharingBoolSet::CalcNumBitsInTheBucket(BucketIdx, TotalNumberOfElements);
+			for (int32 Index = 0; Index < NumBitsInThisBucket; ++Index)
 			{
-				return;
+				const int32 TypeIdx = Index + BucketIdx * FNoFalseSharingBoolSet::NumBitsPerBucket;
+				TypeIndicesToRemove[BucketIdx].Set[Index] = 
+					!Types[TypeIdx].IsUnknownType()
+					// avoid array types, the base type removal will clean up related array types.
+					// For example removing FMyStruct will also clean up TArray<FMyStruct>.
+					// Note: There's always a base type registered for an array type being present.
+					&& !Types[TypeIdx].bIsArray
+					&& IsWithinPlugin(ModulePrefixes, Types[TypeIdx].Type.CPPTypeObject);
 			}
-
-			if (!IsWithinPlugin(ModulePrefixes, Types[Index].Type.CPPTypeObject))
-			{
-				return;
-			}
-
-			FScopeLock Lock(&ParallelForDataLock);
-			TypeIndicesToRemove.Add(Index);
 		});
 
-		for (int32 Index : TypeIndicesToRemove)
+		for (int32 BucketIdx = 0; BucketIdx < NumBuckets; ++BucketIdx)
 		{
-			if (RemoveType_NoLock(Index))
+			const int32 NumBitsInThisBucket = FNoFalseSharingBoolSet::CalcNumBitsInTheBucket(BucketIdx, TotalNumberOfElements);
+			for (int32 Index = 0; Index < NumBitsInThisBucket; ++Index)
 			{
-				bRegistryChanged = true;
+				if (TypeIndicesToRemove[BucketIdx].Set[Index])
+				{
+					const int32 TypeIdx = Index + BucketIdx * FNoFalseSharingBoolSet::NumBitsPerBucket;
+					if (RemoveType_NoLock(TypeIdx))
+					{
+						bRegistryChanged = true;
+					}
+				}
 			}
 		}
 	}
@@ -1534,8 +1568,7 @@ bool FRigVMRegistry_NoLock::RemoveType_NoLock(TRigVMTypeIndex InTypeIndex)
 
 	check(!IsArrayType_NoLock(InTypeIndex));
 
-	TArray<TRigVMTypeIndex> Indices;
-	Indices.Init(INDEX_NONE, 3);
+	TStaticArray<TRigVMTypeIndex, 3> Indices {InPlace, INDEX_NONE};
 	Indices[0] = InTypeIndex;
 	Indices[1] = GetArrayTypeFromBaseTypeIndex_NoLock(Indices[0]);
 
@@ -1689,8 +1722,38 @@ void FRigVMRegistry_RWLock::OnPluginLoaded(IPlugin& InPlugin)
 	}
 }
 
+static std::atomic_int32_t GUnmountSuspendedCounter = 0;
+static FCriticalSection GListToUnmountSync;
+static TArray<FString> GListOfModules;
+
+void FRigVMRegistry_RWLock::SuspendGlobalUnmount()
+{
+	GUnmountSuspendedCounter.fetch_add(1, std::memory_order_acq_rel);
+}
+
+void FRigVMRegistry_RWLock::ResumeGlobalUnmount()
+{
+	if (GUnmountSuspendedCounter.fetch_sub(1, std::memory_order_acq_rel) == 1)
+	{
+		TArray<FString> LocalListOfModules;
+		{
+			FScopeLock Guard(&GListToUnmountSync);
+			LocalListOfModules = MoveTemp(GListOfModules);
+		}
+		Get().OnModulesUnloadedInternal(LocalListOfModules);
+	}
+}
+
 void FRigVMRegistry_RWLock::OnPluginUnloaded(IPlugin& InPlugin)
 {
+	if (IsGlobalRegistry() && GUnmountSuspendedCounter.load(std::memory_order_acquire) > 0)
+	{
+		TArray<FString> ModulesPrefixesList = GetPluginModulePrefixes(InPlugin);
+		FScopeLock Guard(&GListToUnmountSync);
+		GListOfModules.Append(MoveTemp(ModulesPrefixesList));
+		return;
+	}
+
 	bool bRegistryChanged;
 	{
 		FConditionalWriteScopeLock _(*this);
diff --git a/Engine/Plugins/Runtime/RigVM/Source/RigVM/Public/RigVMCore/RigVMRegistry.h b/Engine/Plugins/Runtime/RigVM/Source/RigVM/Public/RigVMCore/RigVMRegistry.h
--- a/Engine/Plugins/Runtime/RigVM/Source/RigVM/Public/RigVMCore/RigVMRegistry.h
+++ b/Engine/Plugins/Runtime/RigVM/Source/RigVM/Public/RigVMCore/RigVMRegistry.h
@@ -684,6 +684,13 @@ public:
 		}
 	};
 
+	class FSuspendPluginUnmounting
+	{
+	public:
+		UE_NODISCARD_CTOR FSuspendPluginUnmounting() { SuspendGlobalUnmount(); }
+		~FSuspendPluginUnmounting() { ResumeGlobalUnmount(); }
+	};
+
 	// Returns the singleton registry
 	static RIGVM_API FRigVMRegistry_RWLock& Get();
 
@@ -1538,6 +1545,9 @@ private:
 	static void EnsureLocked(ELockType InLockType);
 
 	UE_API void Initialize_NoLock();
+	
+	UE_API static void SuspendGlobalUnmount();
+	UE_API static void ResumeGlobalUnmount();
 
 	mutable FTransactionallySafeRWLock Lock;
 	mutable std::atomic<ELockType> LockType = LockType_Invalid;
