Don't call RegisterRigVMStructs immedately in modular builds when the RigVM module loads, as its UObjects may not have been initialized yet.

diff --git a/Engine/Plugins/Runtime/RigVM/Source/RigVM/Private/RigVMCore/RigVMRegistry.cpp b/Engine/Plugins/Runtime/RigVM/Source/RigVM/Private/RigVMCore/RigVMRegistry.cpp
--- a/Engine/Plugins/Runtime/RigVM/Source/RigVM/Private/RigVMCore/RigVMRegistry.cpp
+++ b/Engine/Plugins/Runtime/RigVM/Source/RigVM/Private/RigVMCore/RigVMRegistry.cpp
@@ -55,7 +55,13 @@ static void RegisterRigVMStructs()
 
 static FDelayedAutoRegisterHelper GRigVMRegistry_PostUObjectInit(EDelayedRegisterRunPhase::ObjectSystemReady, []()
 {
+#if IS_MONOLITHIC
+	// Doing this in modular builds means we may try and register the structes declared in this module before their UObjects have been initialized
+	// because we don't know the order the dynamic initializers for this helper, the FRegisterRigVMStructs and the UObject registration will happen in.
+	// We may also be able to remove this call in monolithic builds if we can guarantee there will be another call to ProcessNewlyLoadedUObjects after 
+	// EDelayedRegisterRunPhase::ObjectSystemReady
 	RegisterRigVMStructs();
+#endif
 	// Register newly added structs as new modules create their UObjects
 	FCoreUObjectDelegates::CompiledInUObjectsRegisteredDelegate.AddLambda([](FName ModuleName, ECompiledInUObjectsRegisteredStatus Status)
 	{
