Fixed RigVM NotEqual operand

diff --git a/Engine/Plugins/Runtime/RigVM/Source/RigVM/Private/RigVMCore/RigVM.cpp b/Engine/Plugins/Runtime/RigVM/Source/RigVM/Private/RigVMCore/RigVM.cpp
--- a/Engine/Plugins/Runtime/RigVM/Source/RigVM/Private/RigVMCore/RigVM.cpp
+++ b/Engine/Plugins/Runtime/RigVM/Source/RigVM/Private/RigVMCore/RigVM.cpp
@@ -2083,11 +2083,11 @@ ERigVMExecuteResult URigVM::ExecuteInstructions(FRigVMExtendedExecuteContext& Co
 			case ERigVMOpCode::Equals:
 			case ERigVMOpCode::NotEquals:
 			{
-				const FRigVMComparisonOp& Op = ByteCode.GetOpAt<FRigVMComparisonOp>(Instruction);
-
 				FRigVMMemoryHandle& HandleA = Context.CachedMemoryHandles[FirstHandleForInstruction[ContextPublicData.InstructionIndex]];
 				FRigVMMemoryHandle& HandleB = Context.CachedMemoryHandles[FirstHandleForInstruction[ContextPublicData.InstructionIndex] + 1];
-				const bool Result = HandleA.GetProperty()->Identical(HandleA.GetInputData(), HandleB.GetInputData());
+
+				const bool Identical = HandleA.GetProperty()->Identical(HandleA.GetInputData(), HandleB.GetInputData());
+				const bool Result = Instruction.OpCode == ERigVMOpCode::Equals ? Identical : !Identical;
 
 				*((bool*)Context.CachedMemoryHandles[FirstHandleForInstruction[ContextPublicData.InstructionIndex]+2].GetOutputData()) = Result;
 				ContextPublicData.InstructionIndex++;
