Control rig: Fix uint16 max instructions restriction
#rnx

diff --git a/Engine/Plugins/Runtime/RigVM/Source/RigVM/Private/RigVMCore/RigVM.cpp b/Engine/Plugins/Runtime/RigVM/Source/RigVM/Private/RigVMCore/RigVM.cpp
--- a/Engine/Plugins/Runtime/RigVM/Source/RigVM/Private/RigVMCore/RigVM.cpp
+++ b/Engine/Plugins/Runtime/RigVM/Source/RigVM/Private/RigVMCore/RigVM.cpp
@@ -1743,7 +1743,7 @@ ERigVMExecuteResult URigVM::ExecuteVM(FRigVMExtendedExecuteContext& Context, con
 			return Context.CurrentExecuteResult = ERigVMExecuteResult::Failed;
 		}
 		
-		SetInstructionIndex(Context, (uint16)ByteCode.GetEntry(EntryIndex).InstructionIndex);
+		SetInstructionIndex(Context, ByteCode.GetEntry(EntryIndex).InstructionIndex);
 		EntryIndexToPush = EntryIndex;
 
 		if(bIsRootEntry)
@@ -1881,6 +1881,7 @@ ERigVMExecuteResult URigVM::ExecuteInstructions(FRigVMExtendedExecuteContext& Co
 {
 	// make we are already executing this VM
 	check(!Context.CurrentVMMemory.IsEmpty());
+	check(InFirstInstruction <= InLastInstruction);
 
 	FRigVMExecuteContext& ContextPublicData = Context.GetPublicData<>();
 	TGuardValue<int32> InstructionIndexGuard(ContextPublicData.InstructionIndex, InFirstInstruction);
@@ -1912,7 +1913,7 @@ ERigVMExecuteResult URigVM::ExecuteInstructions(FRigVMExtendedExecuteContext& Co
 		}
 #endif
 		
-		if(ContextPublicData.InstructionIndex > (uint16)InLastInstruction)
+		if(ContextPublicData.InstructionIndex > InLastInstruction)
 		{
 			return Context.CurrentExecuteResult = ERigVMExecuteResult::Succeeded;
 		}
diff --git a/Engine/Plugins/Runtime/RigVM/Source/RigVM/Public/RigVMCore/RigVM.h b/Engine/Plugins/Runtime/RigVM/Source/RigVM/Public/RigVMCore/RigVM.h
--- a/Engine/Plugins/Runtime/RigVM/Source/RigVM/Public/RigVMCore/RigVM.h
+++ b/Engine/Plugins/Runtime/RigVM/Source/RigVM/Public/RigVMCore/RigVM.h
@@ -981,7 +981,7 @@ private:
 
 protected:
 
-	virtual void SetInstructionIndex(FRigVMExtendedExecuteContext& Context, uint16 InInstructionIndex) { Context.GetPublicData<>().InstructionIndex = InInstructionIndex; }
+	virtual void SetInstructionIndex(FRigVMExtendedExecuteContext& Context, int32 InInstructionIndex) { Context.GetPublicData<>().InstructionIndex = InInstructionIndex; }
 
 #if WITH_EDITORONLY_DATA
 	UPROPERTY(transient)
diff --git a/Engine/Plugins/Runtime/RigVM/Source/RigVM/Public/RigVMCore/RigVMNativized.h b/Engine/Plugins/Runtime/RigVM/Source/RigVM/Public/RigVMCore/RigVMNativized.h
--- a/Engine/Plugins/Runtime/RigVM/Source/RigVM/Public/RigVMCore/RigVMNativized.h
+++ b/Engine/Plugins/Runtime/RigVM/Source/RigVM/Public/RigVMCore/RigVMNativized.h
@@ -62,7 +62,7 @@ public:
 	virtual bool ContainsEntry(const FName& InEntryName) const override { return GetEntryNames().Contains(InEntryName); }
 	virtual int32 FindEntry(const FName& InEntryName) const override { return GetEntryNames().Find(InEntryName); }
 	virtual const TArray<FName>& GetEntryNames() const override { static const TArray<FName> EmptyEntries; return EmptyEntries; }
-	virtual void SetInstructionIndex(FRigVMExtendedExecuteContext& Context, uint16 InInstructionIndex) override
+	virtual void SetInstructionIndex(FRigVMExtendedExecuteContext& Context, int32 InInstructionIndex) override
 	{
 		Super::SetInstructionIndex(Context, InInstructionIndex);
 #if WITH_EDITOR
