Control Rig: In Control Rig editor, when using the shape combo box in a Graph Node or Details, the scrollbox no longer wiggles when scrolling

#rnx

diff --git a/Engine/Plugins/Runtime/RigVM/Source/RigVMEditor/Private/Widgets/SRigVMGraphPinNameListValueWidget.cpp b/Engine/Plugins/Runtime/RigVM/Source/RigVMEditor/Private/Widgets/SRigVMGraphPinNameListValueWidget.cpp
--- a/Engine/Plugins/Runtime/RigVM/Source/RigVMEditor/Private/Widgets/SRigVMGraphPinNameListValueWidget.cpp
+++ b/Engine/Plugins/Runtime/RigVM/Source/RigVMEditor/Private/Widgets/SRigVMGraphPinNameListValueWidget.cpp
@@ -4,6 +4,7 @@
 #include "DetailLayoutBuilder.h"
 #include "Framework/Application/SlateUser.h"
 #include "Framework/Application/SlateApplication.h"
+#include "Framework/MultiBox/MultiBoxBuilder.h"
 
 #define LOCTEXT_NAMESPACE "GraphPinNameListValueWidget"
 
@@ -17,33 +18,41 @@ void SRigVMGraphPinNameListValueWidget::Construct(const FArguments& InArgs)
 	OptionsSource = InArgs._OptionsSource;
 	CustomScrollbar = InArgs._CustomScrollbar;
 
-	TSharedRef<SWidget> ComboBoxMenuContent =
+	const bool bShouldCloseAfterSelection = true;
+	const TSharedPtr<FUICommandList> NullCommandList = nullptr;
+	FMenuBuilder MenuBuilder(bShouldCloseAfterSelection, NullCommandList);
+	
+	const FName DefaultSection = NAME_None;
+	MenuBuilder.BeginSection(DefaultSection);
+	{
+		const FText& NoLabel = FText::GetEmpty();
+
+		MenuBuilder.AddWidget(
+			SAssignNew(this->SearchField, SEditableTextBox)
+			.HintText(InArgs._SearchHintText)
+			.OnTextChanged(this, &SRigVMGraphPinNameListValueWidget::OnSearchTextChanged)
+			.OnTextCommitted(this, &SRigVMGraphPinNameListValueWidget::OnSearchTextCommitted)
+			.OnKeyDownHandler(this, &SRigVMGraphPinNameListValueWidget::OnSearchTextKeyDown),
+			NoLabel);
+
+		MenuBuilder.AddWidget(
+			SAssignNew(this->ComboListView, SComboListType)
+			.ListItemsSource(OptionsSource)
+			.OnGenerateRow(this, &SRigVMGraphPinNameListValueWidget::GenerateMenuItemRow)
+			.OnSelectionChanged(this, &SRigVMGraphPinNameListValueWidget::OnSelectionChanged_Internal, false)
+			.SelectionMode(ESelectionMode::Single)
+			.ExternalScrollbar(InArgs._CustomScrollbar)
+			.OnKeyDownHandler(this, &SRigVMGraphPinNameListValueWidget::OnComboListKeyDown),
+			NoLabel);
+	}
+	MenuBuilder.EndSection();
+
+	const TSharedRef<SWidget> ComboBoxMenuContent =
 		SNew(SBox)
 		.MinDesiredWidth(150.0f)
 		.MaxDesiredHeight(InArgs._MaxListHeight)
 		[
-			SNew(SVerticalBox)
-
-			+ SVerticalBox::Slot()
-			.AutoHeight()
-			[
-				SAssignNew(this->SearchField, SEditableTextBox)
-				.HintText(InArgs._SearchHintText)
-				.OnTextChanged(this, &SRigVMGraphPinNameListValueWidget::OnSearchTextChanged)
-				.OnTextCommitted(this, &SRigVMGraphPinNameListValueWidget::OnSearchTextCommitted)
-				.OnKeyDownHandler(this, &SRigVMGraphPinNameListValueWidget::OnSearchTextKeyDown)
-			]
-
-			+ SVerticalBox::Slot()
-			[
-				SAssignNew(this->ComboListView, SComboListType)
-				.ListItemsSource(OptionsSource)
-				.OnGenerateRow(this, &SRigVMGraphPinNameListValueWidget::GenerateMenuItemRow)
-				.OnSelectionChanged(this, &SRigVMGraphPinNameListValueWidget::OnSelectionChanged_Internal, false)
-				.SelectionMode(ESelectionMode::Single)
-				.ExternalScrollbar(InArgs._CustomScrollbar)
-				.OnKeyDownHandler(this, &SRigVMGraphPinNameListValueWidget::OnComboListKeyDown)
-			]
+			MenuBuilder.MakeWidget()
 		];
 
 	// Set up content
diff --git a/Engine/Plugins/Runtime/RigVM/Source/RigVMEditor/Public/Widgets/SRigVMGraphPinNameListValueWidget.h b/Engine/Plugins/Runtime/RigVM/Source/RigVMEditor/Public/Widgets/SRigVMGraphPinNameListValueWidget.h
--- a/Engine/Plugins/Runtime/RigVM/Source/RigVMEditor/Public/Widgets/SRigVMGraphPinNameListValueWidget.h
+++ b/Engine/Plugins/Runtime/RigVM/Source/RigVMEditor/Public/Widgets/SRigVMGraphPinNameListValueWidget.h
@@ -147,7 +147,7 @@ private:
 	/** Delegate to invoke when we need to visualize an option as a widget. */
 	FOnGenerateWidget OnGenerateWidget;
 
-	const TArray< TSharedPtr<FRigVMStringWithTag> >* OptionsSource;
+	const TArray< TSharedPtr<FRigVMStringWithTag> >* OptionsSource = nullptr;
 	bool AllowUserProvidedText;
 
 	/** Used to focus the name box immediately following construction */
