RigVM Json Utils: Add support for ClipboardCopy for introspection console commands

#rnx

diff --git a/Engine/Plugins/Runtime/RigVM/Source/RigVMDeveloper/Private/RigVMJsonUtils.cpp b/Engine/Plugins/Runtime/RigVM/Source/RigVMDeveloper/Private/RigVMJsonUtils.cpp
--- a/Engine/Plugins/Runtime/RigVM/Source/RigVMDeveloper/Private/RigVMJsonUtils.cpp
+++ b/Engine/Plugins/Runtime/RigVM/Source/RigVMDeveloper/Private/RigVMJsonUtils.cpp
@@ -17,6 +17,10 @@
 
 #define LOCTEXT_NAMESPACE "RigVMJsonUtils"
 
+#if WITH_EDITOR
+#include "HAL/PlatformApplicationMisc.h"
+#endif
+
 FAutoConsoleCommand FCmdRigVMJsonPrintAvailableNodes
 (
 	TEXT("RigVM.JSON.PrintAvailableNodes"),
@@ -43,6 +47,10 @@ FAutoConsoleCommand FCmdRigVMJsonPrintAvailableNodes
 			if (FJsonSerializer::Serialize(JsonObject, Writer))
 			{
 				UE_LOG(LogRigVMDeveloper, Display, TEXT("Json:\n\n%s"), *Text);
+#if WITH_EDITOR
+				UE_LOG(LogRigVMDeveloper, Display, TEXT("The content has also been copied to the clipboard."), *Text);
+				FPlatformApplicationMisc::ClipboardCopy(*Text);
+#endif
 			}
 		}
 	})
@@ -76,6 +84,10 @@ FAutoConsoleCommand FCmdRigVMJsonPrintGraph
 			if (FJsonSerializer::Serialize(JsonObject, Writer))
 			{
 				UE_LOG(LogRigVMDeveloper, Display, TEXT("Json:\n\n%s"), *Text);
+#if WITH_EDITOR
+				UE_LOG(LogRigVMDeveloper, Display, TEXT("The content has also been copied to the clipboard."), *Text);
+				FPlatformApplicationMisc::ClipboardCopy(*Text);
+#endif
 			}
 		}
 	})
diff --git a/Engine/Plugins/Runtime/RigVM/Source/RigVMDeveloper/RigVMDeveloper.Build.cs b/Engine/Plugins/Runtime/RigVM/Source/RigVMDeveloper/RigVMDeveloper.Build.cs
--- a/Engine/Plugins/Runtime/RigVM/Source/RigVMDeveloper/RigVMDeveloper.Build.cs
+++ b/Engine/Plugins/Runtime/RigVM/Source/RigVMDeveloper/RigVMDeveloper.Build.cs
@@ -34,6 +34,7 @@ public class RigVMDeveloper : ModuleRules
 					"MessageLog",
 					"BlueprintGraph",
 					"GraphEditor",
+					"ApplicationCore",
                 }
 			);
 
