Removed RigVM cached handles array from instances, as a memory saving optimization
This also eliminates the need to patch handles for each instance initialization (it generates patched handles on the fly for each instruction)

diff --git a/Engine/Plugins/Runtime/RigVM/Source/RigVM/Private/RigVMCore/RigVM.cpp b/Engine/Plugins/Runtime/RigVM/Source/RigVM/Private/RigVMCore/RigVM.cpp
--- a/Engine/Plugins/Runtime/RigVM/Source/RigVM/Private/RigVMCore/RigVM.cpp
+++ b/Engine/Plugins/Runtime/RigVM/Source/RigVM/Private/RigVMCore/RigVM.cpp
@@ -817,7 +817,6 @@ void URigVM::Empty(FRigVMExtendedExecuteContext& Context)
 
 	OperandToDebugRegisters.Empty();
 
-	Context.CachedMemoryHandles.Empty();
 	Context.ExternalVariableRuntimeData.Reset();
 }
 
@@ -1251,9 +1250,6 @@ void URigVM::InvalidateCachedMemory_Internal()
 	LazyBranches.Reset();
 	PreCachedMemoryHandles.Reset();
 	PreCachedMemoryHandlesMemoryType.Reset();
-	InstructionsIndexesWithExternalMemory.Reset();
-	NumHandlesWithExternalMemory = 0;
-	CurrenntInstructionWithExternalMemory = 0;
 }
 
 void URigVM::InvalidateCachedMemory(FRigVMExtendedExecuteContext& Context)
@@ -1450,27 +1446,20 @@ void URigVM::PrepareMemoryForExecution(FRigVMExtendedExecuteContext& Context)
 
 	// Count how many handles we need and set up the indirection offsets for the handles.
 	int32 HandleCount = 0;
-	int32& NumHandlesWithExternalMemoryRef = NumHandlesWithExternalMemory;
 	FirstHandleForInstruction.Add(0);
 	for (int32 InstructionIndex = 0; InstructionIndex < Instructions.Num(); InstructionIndex++)
 	{
 		InstructionOpEval(Context, RegistryHandle, InstructionIndex, INDEX_NONE,
-			[&HandleCount, &NumHandlesWithExternalMemoryRef](FRigVMExtendedExecuteContext&, const FRigVMRegistryHandle&, int32, int32, const FRigVMBranchInfoKey&, const FRigVMOperand& InArg)
+			[&HandleCount](FRigVMExtendedExecuteContext&, const FRigVMRegistryHandle&, int32, int32, const FRigVMBranchInfoKey&, const FRigVMOperand& InArg)
 			{
 				HandleCount++;
-				if (InArg.GetMemoryType() == ERigVMMemoryType::External)
-				{
-					NumHandlesWithExternalMemoryRef++;
-				}
 			});
 		FirstHandleForInstruction.Add(HandleCount);
 	}
 
 	MemoryHandleCount = HandleCount;
-	Context.CachedMemoryHandles.Reset(MemoryHandleCount);
 	PreCachedMemoryHandles.Reset(MemoryHandleCount);
 	PreCachedMemoryHandlesMemoryType.Reset(MemoryHandleCount);
-	InstructionsIndexesWithExternalMemory.SetNumUninitialized(NumHandlesWithExternalMemory);
 }
 
 void URigVM::CacheMemoryHandles(FRigVMExtendedExecuteContext& Context)
@@ -1483,7 +1472,7 @@ void URigVM::CacheMemoryHandles(FRigVMExtendedExecuteContext& Context)
 
 	check((NumInstructions + 1) == FirstHandleForInstruction.Num());
 
-	if (Context.CachedMemoryHandles.Num() == MemoryHandleCount)
+	if (PreCachedMemoryHandles.Num() == MemoryHandleCount)
 	{
 		return;
 	}
@@ -1510,66 +1499,6 @@ void URigVM::CacheMemoryHandles(FRigVMExtendedExecuteContext& Context)
 	);
 }
 
-void URigVM::PatchMemoryHandles(FRigVMExtendedExecuteContext& Context)
-{
-	const int32 NumInstructions = Instructions.Num();
-	if (NumInstructions == 0)
-	{
-		return;
-	}
-
-	check((NumInstructions + 1) == FirstHandleForInstruction.Num());
-
-	if (Context.CachedMemoryHandles.Num() == MemoryHandleCount)
-	{
-		return;
-	}
-
-	// Copy the pre-cached handles from the VM into the instance context
-	Context.CachedMemoryHandles = PreCachedMemoryHandles;
-
-	UE_RIGVM_GET_READ_REGISTRY_CONTEXT();
-
-	FRigVMByteCode& ByteCode = GetByteCode();
-
-	// Now patch the handles memory, which contain offsets at this point
-	static TStaticArray<ERigVMMemoryType, (int32)ERigVMMemoryType::Num> MemoryTypes = { ERigVMMemoryType::Work, ERigVMMemoryType::Literal, ERigVMMemoryType::External, ERigVMMemoryType::Debug };
-	static_assert((int32)ERigVMMemoryType::Num == 4);
-
-	TStaticArray<FRigVMMemoryStorageStruct*, (int32)ERigVMMemoryType::Num> MemoryPtrs;
-	for (int32 MemoryTypeIndex = 0; MemoryTypeIndex < MemoryTypes.Num(); MemoryTypeIndex++)
-	{
-		MemoryPtrs[MemoryTypeIndex] = GetMemoryByType(Context, MemoryTypes[MemoryTypeIndex]);
-	}
-
-	check(PreCachedMemoryHandlesMemoryType.Num() == Context.CachedMemoryHandles.Num());
-
-	// Fix all pre cached handles memory ptr, by adding the corresponding memory start (different by memory type)
-	FRigVMMemoryHandle* RESTRICT MemoryHandle = Context.CachedMemoryHandles.GetData();
-	const FRigVMMemoryHandle* RESTRICT MemoryHandlesEnd = Context.CachedMemoryHandles.GetData() + Context.CachedMemoryHandles.Num();
-	ERigVMMemoryType* RESTRICT HandleMemoryType = PreCachedMemoryHandlesMemoryType.GetData();
-	for (; MemoryHandle < MemoryHandlesEnd; MemoryHandle++, HandleMemoryType++)
-	{
-		if (const FRigVMMemoryStorageStruct* Memory = MemoryPtrs[(int32)*HandleMemoryType])
-		{
-			const PTRINT BaseAddres = (PTRINT)Memory->GetContainerPtr();
-			const PTRINT Offset = (PTRINT)MemoryHandle->Ptr;
-			MemoryHandle->Ptr = (uint8*)(BaseAddres + Offset);
-		}
-	}
-
-	// Fix handles using external memory, these can not be fixed just with the base memory address, have to be fully processed
-	for (const int32 InstructionIndex : InstructionsIndexesWithExternalMemory)
-	{
-		InstructionOpEval(Context, RegistryHandle, InstructionIndex, FirstHandleForInstruction[InstructionIndex],
-			[&](FRigVMExtendedExecuteContext& InContext, const FRigVMRegistryHandle& InRegistry, int32 InInstructionIndex, int32 InHandleIndex, const FRigVMBranchInfoKey& InBranchInfoKey, const FRigVMOperand& InOp)
-			{
-				// This will only process the handles of the instruction that use external memory, the rest will be ignored (these have been already patched)
-				PatchExternalMemoryHandles(Context, InRegistry, InHandleIndex, InOp);
-			});
-	}
-}
-
 void URigVM::RebuildByteCodeOnLoad()
 {
 	Instructions = GetByteCode().GetInstructions();
@@ -1635,8 +1564,6 @@ bool URigVM::InitializeInstance(FRigVMExtendedExecuteContext& Context, bool bCop
 {
 	TGuardValue<int32> GuardThreadId(Context.ExecutingThreadId, FPlatformTLS::GetCurrentThreadId());
 
-	Context.CachedMemoryHandles.Reset(MemoryHandleCount);
-
 	const int32 LazyBranchSize = GetByteCode().BranchInfos.Num();
 	Context.LazyBranchExecuteState.Reset(LazyBranchSize);
 	Context.LazyBranchExecuteState.SetNumZeroed(LazyBranchSize);
@@ -1650,9 +1577,6 @@ bool URigVM::InitializeInstance(FRigVMExtendedExecuteContext& Context, bool bCop
 #endif // WITH_EDITOR
 	}
 
-	UE::TScopeLock CreateLocalizedRegistryScopeLock(CreateLocalizedRegistryMutex);
-	PatchMemoryHandles(Context);
-
 	return true;
 }
 
@@ -1965,16 +1889,30 @@ ERigVMExecuteResult URigVM::ExecuteInstructions(FRigVMExtendedExecuteContext& Co
 #endif
 #endif
 
+		// have the VM instance memory by type ready for handles patching
+		static_assert((int32)ERigVMMemoryType::Num == 4);
+		const TStaticArray<FRigVMMemoryStorageStruct*, (int32)ERigVMMemoryType::Num> VMMemoryPtrs = 
+		{
+			GetMemoryByType(Context, ERigVMMemoryType::Work),
+			GetMemoryByType(Context, ERigVMMemoryType::Literal),
+			GetMemoryByType(Context, ERigVMMemoryType::External),
+			GetMemoryByType(Context, ERigVMMemoryType::Debug),
+		};
+		check(PreCachedMemoryHandles.Num() == MemoryHandleCount);
+		check(PreCachedMemoryHandlesMemoryType.Num() == MemoryHandleCount);
+
 		switch (Instruction.OpCode)
 		{
 			case ERigVMOpCode::Execute:
 			{
 				const FRigVMExecuteOp& Op = ByteCode.GetOpAt<FRigVMExecuteOp>(Instruction);
 				const int32 OperandCount = FirstHandleForInstruction[ContextPublicData.InstructionIndex + 1] - FirstHandleForInstruction[ContextPublicData.InstructionIndex];
-				FRigVMMemoryHandleArray Handles;
+
+				TArray<FRigVMMemoryHandle> Handles;
+
 				if(OperandCount > 0)
 				{
-					Handles = FRigVMMemoryHandleArray(&Context.CachedMemoryHandles[FirstHandleForInstruction[ContextPublicData.InstructionIndex]], OperandCount);
+					GeneratePatchedHandles(Handles, Context.ExternalVariableRuntimeData, VMMemoryPtrs, PreCachedMemoryHandles, PreCachedMemoryHandlesMemoryType, FirstHandleForInstruction[ContextPublicData.InstructionIndex], OperandCount);
 				}
 
 				FRigVMPredicateBranchArray Predicates;
@@ -2004,20 +1942,22 @@ ERigVMExecuteResult URigVM::ExecuteInstructions(FRigVMExtendedExecuteContext& Co
 			}
 			case ERigVMOpCode::Zero:
 			{
-				const FRigVMMemoryHandle& Handle = Context.CachedMemoryHandles[FirstHandleForInstruction[ContextPublicData.InstructionIndex]];
+				const int32 FirstHandleIndex = FirstHandleForInstruction[ContextPublicData.InstructionIndex];
+				const FRigVMMemoryHandle Handle = GeneratePatchedHandle(Context.ExternalVariableRuntimeData, VMMemoryPtrs, PreCachedMemoryHandlesMemoryType[FirstHandleIndex], PreCachedMemoryHandles[FirstHandleIndex]);
+
 				if(Handle.GetProperty()->IsA<FIntProperty>())
 				{
-					*((int32*)Context.CachedMemoryHandles[FirstHandleForInstruction[ContextPublicData.InstructionIndex]].GetInputData()) = 0;
+					*((int32*)Handle.GetInputData()) = 0;
 				}
 				else if(Handle.GetProperty()->IsA<FNameProperty>())
 				{
-					*((FName*)Context.CachedMemoryHandles[FirstHandleForInstruction[ContextPublicData.InstructionIndex]].GetInputData()) = NAME_None;
+					*((FName*)Handle.GetInputData()) = NAME_None;
 				}
 #if WITH_EDITOR
 				if(GetDebugMemory(Context) && GetDebugMemory(Context)->Num() > 0)
 				{
 					const FRigVMUnaryOp& Op = ByteCode.GetOpAt<FRigVMUnaryOp>(Instruction);
-					CopyOperandForDebuggingIfNeeded(Context, Op.Arg, Context.CachedMemoryHandles[FirstHandleForInstruction[ContextPublicData.InstructionIndex]]);
+					CopyOperandForDebuggingIfNeeded(Context, Op.Arg, Handle);
 				}
 #endif
 
@@ -2026,13 +1966,19 @@ ERigVMExecuteResult URigVM::ExecuteInstructions(FRigVMExtendedExecuteContext& Co
 			}
 			case ERigVMOpCode::BoolFalse:
 			{
-				*((bool*)Context.CachedMemoryHandles[FirstHandleForInstruction[ContextPublicData.InstructionIndex]].GetInputData()) = false;
+				const int32 FirstHandleIndex = FirstHandleForInstruction[ContextPublicData.InstructionIndex];
+				const FRigVMMemoryHandle Handle = GeneratePatchedHandle(Context.ExternalVariableRuntimeData, VMMemoryPtrs, PreCachedMemoryHandlesMemoryType[FirstHandleIndex], PreCachedMemoryHandles[FirstHandleIndex]);
+
+				*((bool*)Handle.GetInputData()) = false;
 				ContextPublicData.InstructionIndex++;
 				break;
 			}
 			case ERigVMOpCode::BoolTrue:
 			{
-				*((bool*)Context.CachedMemoryHandles[FirstHandleForInstruction[ContextPublicData.InstructionIndex]].GetInputData()) = true;
+				const int32 FirstHandleIndex = FirstHandleForInstruction[ContextPublicData.InstructionIndex];
+				const FRigVMMemoryHandle Handle = GeneratePatchedHandle(Context.ExternalVariableRuntimeData, VMMemoryPtrs, PreCachedMemoryHandlesMemoryType[FirstHandleIndex], PreCachedMemoryHandles[FirstHandleIndex]);
+
+				*((bool*)Handle.GetInputData()) = true;
 				ContextPublicData.InstructionIndex++;
 				break;
 			}
@@ -2040,8 +1986,10 @@ ERigVMExecuteResult URigVM::ExecuteInstructions(FRigVMExtendedExecuteContext& Co
 			{
 				const FRigVMCopyOp& Op = ByteCode.GetOpAt<FRigVMCopyOp>(Instruction);
 
-				FRigVMMemoryHandle& SourceHandle = Context.CachedMemoryHandles[FirstHandleForInstruction[ContextPublicData.InstructionIndex]];
-				FRigVMMemoryHandle& TargetHandle = Context.CachedMemoryHandles[FirstHandleForInstruction[ContextPublicData.InstructionIndex] + 1];
+				const int32 FirstHandleIndex = FirstHandleForInstruction[ContextPublicData.InstructionIndex];
+				FRigVMMemoryHandle SourceHandle = GeneratePatchedHandle(Context.ExternalVariableRuntimeData, VMMemoryPtrs, PreCachedMemoryHandlesMemoryType[FirstHandleIndex], PreCachedMemoryHandles[FirstHandleIndex]);
+				FRigVMMemoryHandle TargetHandle = GeneratePatchedHandle(Context.ExternalVariableRuntimeData, VMMemoryPtrs, PreCachedMemoryHandlesMemoryType[FirstHandleIndex + 1], PreCachedMemoryHandles[FirstHandleIndex + 1]);
+
 				FRigVMMemoryStorageStruct::CopyProperty(TargetHandle, SourceHandle);
 					
 #if WITH_EDITOR
@@ -2056,12 +2004,15 @@ ERigVMExecuteResult URigVM::ExecuteInstructions(FRigVMExtendedExecuteContext& Co
 			}
 			case ERigVMOpCode::Increment:
 			{
-				(*((int32*)Context.CachedMemoryHandles[FirstHandleForInstruction[ContextPublicData.InstructionIndex]].GetInputData()))++;
+				const int32 FirstHandleIndex = FirstHandleForInstruction[ContextPublicData.InstructionIndex];
+				const FRigVMMemoryHandle Handle = GeneratePatchedHandle(Context.ExternalVariableRuntimeData, VMMemoryPtrs, PreCachedMemoryHandlesMemoryType[FirstHandleIndex], PreCachedMemoryHandles[FirstHandleIndex]);
+
+				(*((int32*)Handle.GetInputData()))++;
 #if WITH_EDITOR
 				if(GetDebugMemory(Context) && GetDebugMemory(Context)->Num() > 0)
 				{
 					const FRigVMUnaryOp& Op = ByteCode.GetOpAt<FRigVMUnaryOp>(Instruction);
-					CopyOperandForDebuggingIfNeeded(Context, Op.Arg, Context.CachedMemoryHandles[FirstHandleForInstruction[ContextPublicData.InstructionIndex]]);
+					CopyOperandForDebuggingIfNeeded(Context, Op.Arg, Handle);
 				}
 #endif
 				ContextPublicData.InstructionIndex++;
@@ -2069,12 +2020,15 @@ ERigVMExecuteResult URigVM::ExecuteInstructions(FRigVMExtendedExecuteContext& Co
 			}
 			case ERigVMOpCode::Decrement:
 			{
-				(*((int32*)Context.CachedMemoryHandles[FirstHandleForInstruction[ContextPublicData.InstructionIndex]].GetInputData()))--;
+				const int32 FirstHandleIndex = FirstHandleForInstruction[ContextPublicData.InstructionIndex];
+				const FRigVMMemoryHandle Handle = GeneratePatchedHandle(Context.ExternalVariableRuntimeData, VMMemoryPtrs, PreCachedMemoryHandlesMemoryType[FirstHandleIndex], PreCachedMemoryHandles[FirstHandleIndex]);
+
+				(*((int32*)Handle.GetInputData()))--;
 #if WITH_EDITOR
 				if(GetDebugMemory(Context) && GetDebugMemory(Context)->Num() > 0)
 				{
 					const FRigVMUnaryOp& Op = ByteCode.GetOpAt<FRigVMUnaryOp>(Instruction);
-					CopyOperandForDebuggingIfNeeded(Context, Op.Arg, Context.CachedMemoryHandles[FirstHandleForInstruction[ContextPublicData.InstructionIndex]]);
+					CopyOperandForDebuggingIfNeeded(Context, Op.Arg, Handle);
 				}
 #endif
 				ContextPublicData.InstructionIndex++;
@@ -2083,13 +2037,15 @@ ERigVMExecuteResult URigVM::ExecuteInstructions(FRigVMExtendedExecuteContext& Co
 			case ERigVMOpCode::Equals:
 			case ERigVMOpCode::NotEquals:
 			{
-				FRigVMMemoryHandle& HandleA = Context.CachedMemoryHandles[FirstHandleForInstruction[ContextPublicData.InstructionIndex]];
-				FRigVMMemoryHandle& HandleB = Context.CachedMemoryHandles[FirstHandleForInstruction[ContextPublicData.InstructionIndex] + 1];
+				const int32 FirstHandleIndex = FirstHandleForInstruction[ContextPublicData.InstructionIndex];
+				const FRigVMMemoryHandle HandleA = GeneratePatchedHandle(Context.ExternalVariableRuntimeData, VMMemoryPtrs, PreCachedMemoryHandlesMemoryType[FirstHandleIndex], PreCachedMemoryHandles[FirstHandleIndex]);
+				const FRigVMMemoryHandle HandleB = GeneratePatchedHandle(Context.ExternalVariableRuntimeData, VMMemoryPtrs, PreCachedMemoryHandlesMemoryType[FirstHandleIndex + 1], PreCachedMemoryHandles[FirstHandleIndex + 1]);
+				FRigVMMemoryHandle HandleR = GeneratePatchedHandle(Context.ExternalVariableRuntimeData, VMMemoryPtrs, PreCachedMemoryHandlesMemoryType[FirstHandleIndex + 2], PreCachedMemoryHandles[FirstHandleIndex + 2]);
 
-				const bool Identical = HandleA.GetProperty()->Identical(HandleA.GetInputData(), HandleB.GetInputData());
+				const bool Identical = HandleA.GetProperty()->Identical(HandleA.GetInputData(), HandleB.GetInputData());;
 				const bool Result = Instruction.OpCode == ERigVMOpCode::Equals ? Identical : !Identical;
 
-				*((bool*)Context.CachedMemoryHandles[FirstHandleForInstruction[ContextPublicData.InstructionIndex]+2].GetOutputData()) = Result;
+				*((bool*)HandleR.GetOutputData()) = Result;
 				ContextPublicData.InstructionIndex++;
 				break;
 			}
@@ -2114,7 +2070,11 @@ ERigVMExecuteResult URigVM::ExecuteInstructions(FRigVMExtendedExecuteContext& Co
 			case ERigVMOpCode::JumpAbsoluteIf:
 			{
 				const FRigVMJumpIfOp& Op = ByteCode.GetOpAt<FRigVMJumpIfOp>(Instruction);
-				const bool Condition = *(bool*)Context.CachedMemoryHandles[FirstHandleForInstruction[ContextPublicData.InstructionIndex]].GetInputData();
+
+				const int32 FirstHandleIndex = FirstHandleForInstruction[ContextPublicData.InstructionIndex];
+				const FRigVMMemoryHandle Handle = GeneratePatchedHandle(Context.ExternalVariableRuntimeData, VMMemoryPtrs, PreCachedMemoryHandlesMemoryType[FirstHandleIndex], PreCachedMemoryHandles[FirstHandleIndex]);
+
+				const bool Condition = *(bool*)Handle.GetInputData();
 				if (Condition == Op.Condition)
 				{
 					ContextPublicData.InstructionIndex = Op.InstructionIndex;
@@ -2128,7 +2088,11 @@ ERigVMExecuteResult URigVM::ExecuteInstructions(FRigVMExtendedExecuteContext& Co
 			case ERigVMOpCode::JumpForwardIf:
 			{
 				const FRigVMJumpIfOp& Op = ByteCode.GetOpAt<FRigVMJumpIfOp>(Instruction);
-				const bool Condition = *(bool*)Context.CachedMemoryHandles[FirstHandleForInstruction[ContextPublicData.InstructionIndex]].GetInputData();
+
+				const int32 FirstHandleIndex = FirstHandleForInstruction[ContextPublicData.InstructionIndex];
+				const FRigVMMemoryHandle Handle = GeneratePatchedHandle(Context.ExternalVariableRuntimeData, VMMemoryPtrs, PreCachedMemoryHandlesMemoryType[FirstHandleIndex], PreCachedMemoryHandles[FirstHandleIndex]);
+
+				const bool Condition = *(bool*)Handle.GetInputData();
 				if (Condition == Op.Condition)
 				{
 					ContextPublicData.InstructionIndex += Op.InstructionIndex;
@@ -2142,7 +2106,11 @@ ERigVMExecuteResult URigVM::ExecuteInstructions(FRigVMExtendedExecuteContext& Co
 			case ERigVMOpCode::JumpBackwardIf:
 			{
 				const FRigVMJumpIfOp& Op = ByteCode.GetOpAt<FRigVMJumpIfOp>(Instruction);
-				const bool Condition = *(bool*)Context.CachedMemoryHandles[FirstHandleForInstruction[ContextPublicData.InstructionIndex]].GetInputData();
+
+				const int32 FirstHandleIndex = FirstHandleForInstruction[ContextPublicData.InstructionIndex];
+				const FRigVMMemoryHandle Handle = GeneratePatchedHandle(Context.ExternalVariableRuntimeData, VMMemoryPtrs, PreCachedMemoryHandlesMemoryType[FirstHandleIndex], PreCachedMemoryHandles[FirstHandleIndex]);
+
+				const bool Condition = *(bool*)Handle.GetInputData();
 				if (Condition == Op.Condition)
 				{
 					ContextPublicData.InstructionIndex -= Op.InstructionIndex;
@@ -2182,8 +2150,12 @@ ERigVMExecuteResult URigVM::ExecuteInstructions(FRigVMExtendedExecuteContext& Co
 			}
 			case ERigVMOpCode::BeginBlock:
 			{
-				const int32 Count = (*((int32*)Context.CachedMemoryHandles[FirstHandleForInstruction[ContextPublicData.InstructionIndex]].GetInputData()));
-				const int32 Index = (*((int32*)Context.CachedMemoryHandles[FirstHandleForInstruction[ContextPublicData.InstructionIndex] + 1].GetInputData()));
+				const int32 FirstHandleIndex = FirstHandleForInstruction[ContextPublicData.InstructionIndex];
+				const FRigVMMemoryHandle CountHandle = GeneratePatchedHandle(Context.ExternalVariableRuntimeData, VMMemoryPtrs, PreCachedMemoryHandlesMemoryType[FirstHandleIndex], PreCachedMemoryHandles[FirstHandleIndex]);
+				const FRigVMMemoryHandle IndexHandle = GeneratePatchedHandle(Context.ExternalVariableRuntimeData, VMMemoryPtrs, PreCachedMemoryHandlesMemoryType[FirstHandleIndex + 1], PreCachedMemoryHandles[FirstHandleIndex + 1]);
+
+				const int32 Count = (*((int32*)CountHandle.GetInputData()));
+				const int32 Index = (*((int32*)IndexHandle.GetInputData()));
 				Context.BeginSlice(Count, Index);
 				ContextPublicData.InstructionIndex++;
 				break;
@@ -2219,8 +2191,12 @@ ERigVMExecuteResult URigVM::ExecuteInstructions(FRigVMExtendedExecuteContext& Co
 			case ERigVMOpCode::JumpToBranch:
 			{
 				const FRigVMJumpToBranchOp& Op = ByteCode.GetOpAt<FRigVMJumpToBranchOp>(Instruction);
+				
+				const int32 FirstHandleIndex = FirstHandleForInstruction[ContextPublicData.InstructionIndex];
+				const FRigVMMemoryHandle Handle = GeneratePatchedHandle(Context.ExternalVariableRuntimeData, VMMemoryPtrs, PreCachedMemoryHandlesMemoryType[FirstHandleIndex], PreCachedMemoryHandles[FirstHandleIndex]);
+
 				// BranchLabel = Op.Arg
-				FName BranchLabel = *(FName*)Context.CachedMemoryHandles[FirstHandleForInstruction[ContextPublicData.InstructionIndex]].GetInputData();
+				FName BranchLabel = *(FName*)Handle.GetInputData();
 
 				// iterate over the branches stored in the bytecode,
 				// starting at the first branch index stored in the operator.
@@ -2272,7 +2248,11 @@ ERigVMExecuteResult URigVM::ExecuteInstructions(FRigVMExtendedExecuteContext& Co
 			case ERigVMOpCode::RunInstructions:
 			{
 				const FRigVMRunInstructionsOp& Op = ByteCode.GetOpAt<FRigVMRunInstructionsOp>(Instruction);
-				FRigVMInstructionSetExecuteState& ExecutionState = *(FRigVMInstructionSetExecuteState*)Context.CachedMemoryHandles[FirstHandleForInstruction[ContextPublicData.InstructionIndex]].GetPrivateData(INDEX_NONE);
+
+				const int32 FirstHandleIndex = FirstHandleForInstruction[ContextPublicData.InstructionIndex];
+				FRigVMMemoryHandle Handle = GeneratePatchedHandle(Context.ExternalVariableRuntimeData, VMMemoryPtrs, PreCachedMemoryHandlesMemoryType[FirstHandleIndex], PreCachedMemoryHandles[FirstHandleIndex]);
+
+				FRigVMInstructionSetExecuteState& ExecutionState = *(FRigVMInstructionSetExecuteState*)Handle.GetPrivateData(INDEX_NONE);
 
 				if((Op.StartInstruction != INDEX_NONE) &&
 					(Op.EndInstruction != INDEX_NONE) &&
@@ -2297,7 +2277,10 @@ ERigVMExecuteResult URigVM::ExecuteInstructions(FRigVMExtendedExecuteContext& Co
 				ContextPublicData.Traits.Reset();
 				ContextPublicData.AdditionalTraitMemoryHandles.Reset();
 
-				const TArray<int32>& TraitList = *(TArray<int32>*)Context.CachedMemoryHandles[FirstHandleForInstruction[ContextPublicData.InstructionIndex]].GetInputData();
+				const int32 FirstHandleIndex = FirstHandleForInstruction[ContextPublicData.InstructionIndex];
+				FRigVMMemoryHandle Handle = GeneratePatchedHandle(Context.ExternalVariableRuntimeData, VMMemoryPtrs, PreCachedMemoryHandlesMemoryType[FirstHandleIndex], PreCachedMemoryHandles[FirstHandleIndex]);
+
+				const TArray<int32>& TraitList = *(TArray<int32>*)Handle.GetInputData();
 				ContextPublicData.Traits.Reserve(TraitList.Num());
 				ContextPublicData.AdditionalTraitMemoryHandles.Reserve(TraitList.Num());
 				int32 AdditionalStartIndex = INDEX_NONE;
@@ -3195,12 +3178,9 @@ void URigVM::PreCacheSingleMemoryHandle(const FRigVMRegistryHandle& InRegistry,
 		FRigVMExternalVariableDef& ExternalVariable = ExternalVariables[ExternalVariableIndex];
 		check(ExternalVariable.IsValid());
 
-		PreCachedMemoryHandles[InHandleIndex] = { nullptr, ExternalVariable.Property, PropertyPath };
+		// External variables store the ExternalVariableIndex as the memory offset, as it is needed to generate the final memory address
+		PreCachedMemoryHandles[InHandleIndex] = { (uint8*)((PTRINT)ExternalVariableIndex), ExternalVariable.Property, PropertyPath };
 		PreCachedMemoryHandlesMemoryType[InHandleIndex] = ERigVMMemoryType::External;
-		{
-			const int32 Index = CurrenntInstructionWithExternalMemory.fetch_add(1); // deal with parallel processing
-			InstructionsIndexesWithExternalMemory[Index] = InInstructionIndex; // These have to be patched on each instance
-		}
 		return;
 	}
 
@@ -3251,40 +3231,53 @@ void URigVM::PreCacheSingleMemoryHandle(const FRigVMRegistryHandle& InRegistry,
 #endif
 }
 
-void URigVM::PatchExternalMemoryHandles(FRigVMExtendedExecuteContext& Context, const FRigVMRegistryHandle& InRegistry, int32 InHandleIndex, const FRigVMOperand& InArg)
+void URigVM::GeneratePatchedHandles(
+	TArray<FRigVMMemoryHandle>& OutHandles,
+	const TArray<FRigVMExternalVariableRuntimeData>& ExternalVariablesRuntimeData,
+	const TStaticArray<FRigVMMemoryStorageStruct*, (int32)ERigVMMemoryType::Num>& VMMemoryPtrs,
+	const TArray<FRigVMMemoryHandle>& PreCachedMemoryHandles,
+	const TArray<ERigVMMemoryType>& PreCachedMemoryHandlesMemoryType,
+	int32 FirstHandleIndex,
+	int32 OperandCount)
 {
-	FRigVMMemoryStorageStruct* Memory = GetMemoryByType(Context, InArg.GetMemoryType()/*, false*/);
-
-	FRigVMMemoryHandle& Handle = Context.CachedMemoryHandles[InHandleIndex];
-
-	if (InArg.GetMemoryType() == ERigVMMemoryType::External)
+	OutHandles.SetNumUninitialized(OperandCount);
+	for (int32 HandleIndex = 0; HandleIndex < OperandCount; HandleIndex++)
 	{
-		const FRigVMPropertyPath* PropertyPath = nullptr;
-		if (InArg.GetRegisterOffset() != INDEX_NONE)
-		{
-			check(ExternalPropertyPaths.IsValidIndex(InArg.GetRegisterOffset()));
-			PropertyPath = &ExternalPropertyPaths[InArg.GetRegisterOffset()];
-		}
+		const int32 PreCachedHandleIndex = FirstHandleIndex + HandleIndex;
+		const FRigVMMemoryHandle& PreCachedHandle = PreCachedMemoryHandles[PreCachedHandleIndex];
+		
+		OutHandles[HandleIndex] = GeneratePatchedHandle(ExternalVariablesRuntimeData, VMMemoryPtrs, PreCachedMemoryHandlesMemoryType[PreCachedHandleIndex], PreCachedHandle);
+	}
+}
 
-		check(ExternalVariables.IsValidIndex(InArg.GetRegisterIndex()));
+FRigVMMemoryHandle URigVM::GeneratePatchedHandle(
+	const TArray<FRigVMExternalVariableRuntimeData>& ExternalVariablesRuntimeData,
+	const TStaticArray<FRigVMMemoryStorageStruct*, (int32)ERigVMMemoryType::Num>& VMMemoryPtrs,
+	ERigVMMemoryType HandleMemoryType,
+	const FRigVMMemoryHandle& PreCachedHandle)
+{
+	FRigVMMemoryHandle Handle = PreCachedHandle;
 
-		const int32 ExternalVariableIndex = InArg.GetRegisterIndex();
-		FRigVMExternalVariableDef& ExternalVariable = ExternalVariables[ExternalVariableIndex];
-		check(ExternalVariable.IsValid());
+	if (HandleMemoryType == ERigVMMemoryType::External)
+	{
+		const int32 ExternalVariableIndex = (int32)((PTRINT)Handle.Ptr); // For external variables we store the RegisterIndex in the memory pointer, so we can retrieve the correct var
 
-		FRigVMExternalVariableRuntimeData& ExternalVariableRuntimeData = Context.ExternalVariableRuntimeData[ExternalVariableIndex];
+		check(ExternalVariablesRuntimeData.IsValidIndex(ExternalVariableIndex));
+		const FRigVMExternalVariableRuntimeData& ExternalVariableRuntimeData = ExternalVariablesRuntimeData[ExternalVariableIndex];
 		check(ExternalVariableRuntimeData.Memory != nullptr);
-		check(Handle.PropertyPath == PropertyPath);
 
 		Handle.Ptr = ExternalVariableRuntimeData.Memory;
+		return Handle;
+	}
 
-#if UE_BUILD_DEBUG
-		// make sure the handle points to valid memory
-		check(Context.CachedMemoryHandles[InHandleIndex].GetOutputData() != nullptr);
-#endif
+	if (const FRigVMMemoryStorageStruct* Memory = VMMemoryPtrs[(int32)HandleMemoryType])
+	{
+		const PTRINT BaseAddress = (PTRINT)Memory->GetContainerPtr();
+		const PTRINT Offset = (PTRINT)Handle.Ptr;
+		Handle.Ptr = (uint8*)(BaseAddress + Offset);
 	}
 
-	// if handle is not using external memory, just ignore it, as it has been already patched
+	return Handle;
 }
 
 void URigVM::CopyOperandForDebuggingImpl(FRigVMExtendedExecuteContext& Context, const FRigVMOperand& InArg, const FRigVMMemoryHandle& InHandle, const FRigVMOperand& InDebugOperand)
diff --git a/Engine/Plugins/Runtime/RigVM/Source/RigVM/Private/RigVMCore/RigVMExecuteContext.cpp b/Engine/Plugins/Runtime/RigVM/Source/RigVM/Private/RigVMCore/RigVMExecuteContext.cpp
--- a/Engine/Plugins/Runtime/RigVM/Source/RigVM/Private/RigVMCore/RigVMExecuteContext.cpp
+++ b/Engine/Plugins/Runtime/RigVM/Source/RigVM/Private/RigVMCore/RigVMExecuteContext.cpp
@@ -98,8 +98,6 @@ void FRigVMExtendedExecuteContext::Reset()
 	
 	OnExecutionReachedExitCallback.Reset();
 
-	CachedMemoryHandles.Reset();
-
 	LazyBranchExecuteState.Reset();
 	ExternalVariableRuntimeData.Reset();
 
@@ -173,8 +171,6 @@ FRigVMExtendedExecuteContext& FRigVMExtendedExecuteContext::operator =(const FRi
 	Slices = Other.Slices;
 	SliceOffsets = Other.SliceOffsets;
 
-	CachedMemoryHandles = Other.CachedMemoryHandles;
-
 	LazyBranchExecuteState = Other.LazyBranchExecuteState;
 	ExternalVariableRuntimeData = Other.ExternalVariableRuntimeData;
 
diff --git a/Engine/Plugins/Runtime/RigVM/Source/RigVM/Public/RigVMCore/RigVM.h b/Engine/Plugins/Runtime/RigVM/Source/RigVM/Public/RigVMCore/RigVM.h
--- a/Engine/Plugins/Runtime/RigVM/Source/RigVM/Public/RigVMCore/RigVM.h
+++ b/Engine/Plugins/Runtime/RigVM/Source/RigVM/Public/RigVMCore/RigVM.h
@@ -973,7 +973,6 @@ private:
 		)>& InOpFunc);
 	RIGVM_API void PrepareMemoryForExecution(FRigVMExtendedExecuteContext& Context);
 	RIGVM_API void CacheMemoryHandles(FRigVMExtendedExecuteContext& Context);
-	RIGVM_API void PatchMemoryHandles(FRigVMExtendedExecuteContext& Context);
 	RIGVM_API void RebuildByteCodeOnLoad();
 
 	UPROPERTY(transient)
@@ -1027,9 +1026,6 @@ private:
 
 	TArray<FRigVMMemoryHandle> PreCachedMemoryHandles;
 	TArray<ERigVMMemoryType> PreCachedMemoryHandlesMemoryType;
-	TArray<int32> InstructionsIndexesWithExternalMemory;
-	int32 NumHandlesWithExternalMemory = 0;
-	std::atomic<int32> CurrenntInstructionWithExternalMemory = 0;
 
 	TArray<FRigVMExternalVariableDef> ExternalVariables;
 	TArray<FRigVMLazyBranch> LazyBranches;
@@ -1053,7 +1049,22 @@ private:
 	RIGVM_API void ClearDebugMemory(FRigVMExtendedExecuteContext& Context);
 
 	RIGVM_API void PreCacheSingleMemoryHandle(const FRigVMRegistryHandle& InRegistry, int32 InInstructionIndex, int32 InHandleIndex, const FRigVMBranchInfoKey& InBranchInfoKey, const FRigVMOperand& InArg, bool bForExecute = false);
-	RIGVM_API void PatchExternalMemoryHandles(FRigVMExtendedExecuteContext& Context, const FRigVMRegistryHandle& InRegistry, int32 InHandleIndex, const FRigVMOperand& InArg);
+
+	RIGVM_API static void GeneratePatchedHandles(
+		TArray<FRigVMMemoryHandle>& OutHandles,
+		const TArray<FRigVMExternalVariableRuntimeData>& ExternalVariablesRuntimeData,
+		const TStaticArray<FRigVMMemoryStorageStruct*, (int32)ERigVMMemoryType::Num>& VMMemoryPtrs,
+		const TArray<FRigVMMemoryHandle>& PreCachedMemoryHandles,
+		const TArray<ERigVMMemoryType>& PreCachedMemoryHandlesMemoryType,
+		int32 FirstHandleIndex,
+		int32 OperandCount);
+
+	RIGVM_API static FRigVMMemoryHandle GeneratePatchedHandle(
+		const TArray<FRigVMExternalVariableRuntimeData>& ExternalVariablesRuntimeData,
+		const TStaticArray<FRigVMMemoryStorageStruct*, (int32)ERigVMMemoryType::Num>& VMMemoryPtrs,
+		ERigVMMemoryType HandleMemoryType,
+		const FRigVMMemoryHandle& PreCachedHandle);
+
 
 
 	void CopyOperandForDebuggingIfNeeded(FRigVMExtendedExecuteContext& Context, const FRigVMOperand& InArg, const FRigVMMemoryHandle& InHandle)
diff --git a/Engine/Plugins/Runtime/RigVM/Source/RigVM/Public/RigVMCore/RigVMExecuteContext.h b/Engine/Plugins/Runtime/RigVM/Source/RigVM/Public/RigVMCore/RigVMExecuteContext.h
--- a/Engine/Plugins/Runtime/RigVM/Source/RigVM/Public/RigVMCore/RigVMExecuteContext.h
+++ b/Engine/Plugins/Runtime/RigVM/Source/RigVM/Public/RigVMCore/RigVMExecuteContext.h
@@ -797,7 +797,6 @@ struct FRigVMExtendedExecuteContext
 
 	void InvalidateCachedMemory()
 	{
-		CachedMemoryHandles.Reset();
 		LazyBranchExecuteState.Reset();
 	}
 	
@@ -850,8 +849,6 @@ struct FRigVMExtendedExecuteContext
 	const FRigVMDispatchFactory* Factory = nullptr;
 	FRigVMNameCache NameCache;
 
-	TArray<FRigVMMemoryHandle> CachedMemoryHandles;
-	
 	int32 ExecutingThreadId = INDEX_NONE;
 
 	TArray<int32> EntriesBeingExecuted;
@@ -863,6 +860,8 @@ struct FRigVMExtendedExecuteContext
 	TArray<FRigVMMemoryStorageStruct*, TInlineAllocator<3>> CurrentVMMemory;
 
 #if WITH_EDITORONLY_DATA
+	UE_DEPRECATED(5.8, "CachedMemoryHandles has been deprecated.")
+	TArray<FRigVMMemoryHandle> CachedMemoryHandles;
 	// changes to the layout of cached memory array should be reflected in GetContainerIndex()
 	UE_DEPRECATED(5.4, "CachedMemory has been deprecated.")
 	TArray<URigVMMemoryStorage*> CachedMemory;
