RigVM: Improve tooltips based on analysis

#rnx

diff --git a/Engine/Plugins/Animation/ControlRig/Source/ControlRig/Private/Units/Core/RigUnit_UserData.cpp b/Engine/Plugins/Animation/ControlRig/Source/ControlRig/Private/Units/Core/RigUnit_UserData.cpp
--- a/Engine/Plugins/Animation/ControlRig/Source/ControlRig/Private/Units/Core/RigUnit_UserData.cpp
+++ b/Engine/Plugins/Animation/ControlRig/Source/ControlRig/Private/Units/Core/RigUnit_UserData.cpp
@@ -7,6 +7,7 @@
 #include "Units/Execution/RigUnit_PrepareForExecution.h"
 
 #include UE_INLINE_GENERATED_CPP_BY_NAME(RigUnit_UserData)
+#define LOCTEXT_NAMESPACE "RigUnit_UserData"
 
 #if WITH_EDITOR
 
@@ -26,6 +27,31 @@ FString FRigDispatch_GetUserData::GetArgumentMetaData(const FName& InArgumentNam
 	return FRigDispatchFactory::GetArgumentMetaData(InArgumentName, InMetaDataKey);
 }
 
+FText FRigDispatch_GetUserData::GetArgumentTooltip(const FName& InArgumentName, TRigVMTypeIndex InTypeIndex) const
+{
+	if (InArgumentName == ArgNameSpaceName)
+	{
+		return LOCTEXT("GetUserDataNameSpaceToolTip", "The namespace to look up the argument within");
+	}
+	if (InArgumentName == ArgPathName)
+	{
+		return LOCTEXT("GetUserDataPathToolTip", "The path to the argument within the namespace");
+	}
+	if (InArgumentName == ArgDefaultName)
+	{
+		return LOCTEXT("GetUserDataDefaultToolTip", "The default value for the argument. This is returned if the argument could not be found.");
+	}
+	if (InArgumentName == ArgResultName)
+	{
+		return LOCTEXT("GetUserDataResultToolTip", "The resulting value for the argument. This may be the default value if the argument could not be found.");
+	}
+	if (InArgumentName == ArgFoundName)
+	{
+		return LOCTEXT("GetUserDataFoundToolTip", "A boolean flag indicating if the argument could be found.");
+	}
+	return FRigDispatchFactory::GetArgumentTooltip(InArgumentName, InTypeIndex);
+}
+
 #endif
 
 const TArray<FRigVMTemplateArgumentInfo>& FRigDispatch_GetUserData::GetArgumentInfos(FRigVMRegistryHandle& InRegistry) const
@@ -237,4 +263,6 @@ FRigUnit_ShapeExists_Execute()
 	{
 		Result = ExecuteContext.OnShapeExistsDelegate.Execute(ShapeName);
 	}
-}
\ No newline at end of file
+}
+
+#undef LOCTEXT_NAMESPACE
diff --git a/Engine/Plugins/Animation/ControlRig/Source/ControlRig/Private/Units/Hierarchy/RigUnit_SetControlTransform.cpp b/Engine/Plugins/Animation/ControlRig/Source/ControlRig/Private/Units/Hierarchy/RigUnit_SetControlTransform.cpp
--- a/Engine/Plugins/Animation/ControlRig/Source/ControlRig/Private/Units/Hierarchy/RigUnit_SetControlTransform.cpp
+++ b/Engine/Plugins/Animation/ControlRig/Source/ControlRig/Private/Units/Hierarchy/RigUnit_SetControlTransform.cpp
@@ -79,14 +79,30 @@ FRigUnit_SetMultiControlFloat_Execute()
 
 FRigUnit_SetControlInteger_Execute()
 {
-    DECLARE_SCOPE_HIERARCHICAL_COUNTER_RIGUNIT()
+	float FloatWeight = (float)Weight;
+	FRigUnit_SetControlIntegerFloatWeight::StaticExecute(ExecuteContext, Control, FloatWeight, IntegerValue, CachedControlIndex);	
+}
+
+FRigVMStructUpgradeInfo FRigUnit_SetControlInteger::GetUpgradeInfo() const
+{
+	FRigUnit_SetControlIntegerFloatWeight NewNode;
+	NewNode.Control = Control;
+	NewNode.Weight = (float)Weight;
+	NewNode.IntegerValue = IntegerValue;
+
+	return FRigVMStructUpgradeInfo(*this, NewNode);
+}
+
+FRigUnit_SetControlIntegerFloatWeight_Execute()
+{
+	DECLARE_SCOPE_HIERARCHICAL_COUNTER_RIGUNIT()
 	URigHierarchy* Hierarchy = ExecuteContext.Hierarchy;
 	if (Hierarchy)
 	{
 		const FRigElementKey Key(Control, ERigElementType::Control);
 		if (CachedControlIndex.UpdateCache(Key, Hierarchy))
 		{
-			if(FMath::IsNearlyEqual((float)Weight, 1.f))
+			if(FMath::IsNearlyEqual(Weight, 1.f))
 			{
 				Hierarchy->SetControlValue(CachedControlIndex, FRigControlValue::Make<int32>(IntegerValue));
 			}
diff --git a/Engine/Plugins/Animation/ControlRig/Source/ControlRig/Public/Units/Core/RigUnit_UserData.h b/Engine/Plugins/Animation/ControlRig/Source/ControlRig/Public/Units/Core/RigUnit_UserData.h
--- a/Engine/Plugins/Animation/ControlRig/Source/ControlRig/Public/Units/Core/RigUnit_UserData.h
+++ b/Engine/Plugins/Animation/ControlRig/Source/ControlRig/Public/Units/Core/RigUnit_UserData.h
@@ -24,6 +24,7 @@ public:
 
 #if WITH_EDITOR
 	virtual FString GetArgumentMetaData(const FName& InArgumentName, const FName& InMetaDataKey) const override;
+	virtual FText GetArgumentTooltip(const FName& InArgumentName, TRigVMTypeIndex InTypeIndex) const override;
 #endif
 	virtual const TArray<FRigVMTemplateArgumentInfo>& GetArgumentInfos(FRigVMRegistryHandle& InRegistry) const override;
 	virtual FRigVMTemplateTypeMap OnNewArgumentType(const FName& InArgumentName, TRigVMTypeIndex InTypeIndex, FRigVMRegistryHandle& InRegistry) const override;
diff --git a/Engine/Plugins/Animation/ControlRig/Source/ControlRig/Public/Units/Execution/RigUnit_DynamicHierarchy.h b/Engine/Plugins/Animation/ControlRig/Source/ControlRig/Public/Units/Execution/RigUnit_DynamicHierarchy.h
--- a/Engine/Plugins/Animation/ControlRig/Source/ControlRig/Public/Units/Execution/RigUnit_DynamicHierarchy.h
+++ b/Engine/Plugins/Animation/ControlRig/Source/ControlRig/Public/Units/Execution/RigUnit_DynamicHierarchy.h
@@ -1577,7 +1577,7 @@ struct FRigUnit_HierarchyAddAnimationChannelVector : public FRigUnit_HierarchyAd
  * Adds a new animation channel to the hierarchy
  * Note: This node only runs as part of the construction event.
  */
-USTRUCT(meta=(DisplayName="Spawn Vector Animation Channel", TemplateName="SpawnScaleAnimationChannel", Keywords="Construction,Create,New,AddAnimationChannel,NewAnimationChannel,CreateAnimationChannel,AddChannel,NewChannel,CreateChannel,SpawnChannel", Varying))
+USTRUCT(meta=(DisplayName="Spawn Scale Animation Channel", TemplateName="SpawnScaleAnimationChannel", Keywords="Construction,Create,New,AddAnimationChannel,NewAnimationChannel,CreateAnimationChannel,AddChannel,NewChannel,CreateChannel,SpawnChannel", Varying))
 struct FRigUnit_HierarchyAddAnimationChannelScaleVector : public FRigUnit_HierarchyAddElement
 {
 	GENERATED_BODY()
diff --git a/Engine/Plugins/Animation/ControlRig/Source/ControlRig/Public/Units/Execution/RigUnit_RigModules.h b/Engine/Plugins/Animation/ControlRig/Source/ControlRig/Public/Units/Execution/RigUnit_RigModules.h
--- a/Engine/Plugins/Animation/ControlRig/Source/ControlRig/Public/Units/Execution/RigUnit_RigModules.h
+++ b/Engine/Plugins/Animation/ControlRig/Source/ControlRig/Public/Units/Execution/RigUnit_RigModules.h
@@ -71,7 +71,7 @@ struct FRigUnit_ResolveConnector : public FRigUnit_RigModulesBase
 };
 
 /**
-* Returns the resolved item of the connector.
+* Returns the resolved array of items of the connector.
 */
 USTRUCT(meta=(DisplayName="Get Array Connection", Keywords="Connector,GetResolved,Target,Resolve", Varying))
 struct FRigUnit_ResolveArrayConnector : public FRigUnit_RigModulesBase
@@ -102,7 +102,7 @@ struct FRigUnit_ResolveArrayConnector : public FRigUnit_RigModulesBase
 	bool SkipSocket;
 
 	/*
-	 * The resulting item the connector is resolved to
+	 * The resulting array of items the connector is resolved to
 	 */
 	UPROPERTY(meta = (Output))
 	TArray<FRigElementKey> Result;
diff --git a/Engine/Plugins/Animation/ControlRig/Source/ControlRig/Public/Units/Hierarchy/RigUnit_GetControlOffset.h b/Engine/Plugins/Animation/ControlRig/Source/ControlRig/Public/Units/Hierarchy/RigUnit_GetControlOffset.h
--- a/Engine/Plugins/Animation/ControlRig/Source/ControlRig/Public/Units/Hierarchy/RigUnit_GetControlOffset.h
+++ b/Engine/Plugins/Animation/ControlRig/Source/ControlRig/Public/Units/Hierarchy/RigUnit_GetControlOffset.h
@@ -8,8 +8,7 @@
 #define UE_API CONTROLRIG_API
 
 /**
- * GetControlOffset is used to perform a change in the hierarchy by setting a single control's transform.
- * This is typically only used during the Construction Event.
+ * GetControlOffset returns the offset transform of a given control.
  */
 USTRUCT(meta=(DisplayName="Get Control Offset", Category="Controls", DocumentationPolicy="Strict", Keywords = "GetControlOffset,Initial,InitialTransform,GetInitialTransform,GetInitialControlTransform", NodeColor="0, 0.364706, 1.0"))
 struct FRigUnit_GetControlOffset : public FRigUnit
diff --git a/Engine/Plugins/Animation/ControlRig/Source/ControlRig/Public/Units/Hierarchy/RigUnit_GetControlTransform.h b/Engine/Plugins/Animation/ControlRig/Source/ControlRig/Public/Units/Hierarchy/RigUnit_GetControlTransform.h
--- a/Engine/Plugins/Animation/ControlRig/Source/ControlRig/Public/Units/Hierarchy/RigUnit_GetControlTransform.h
+++ b/Engine/Plugins/Animation/ControlRig/Source/ControlRig/Public/Units/Hierarchy/RigUnit_GetControlTransform.h
@@ -78,9 +78,8 @@ struct FRigUnit_GetControlFloat : public FRigUnit
 	UPROPERTY()
 	FCachedRigElement CachedControlIndex;
 };
-
 /**
- * GetControlFloat is used to retrieve a single Integer from a hierarchy.
+ * GetControlInteger is used to retrieve a single Integer from a hierarchy.
  */
 USTRUCT(meta=(DisplayName="Get Control Integer", Category="Controls", DocumentationPolicy = "Strict", Keywords="GetControlInteger,GetControlEnum", TemplateName="GetControlValue", NodeColor="0.462745, 1,0, 0.329412", Varying))
 struct FRigUnit_GetControlInteger : public FRigUnit
diff --git a/Engine/Plugins/Animation/ControlRig/Source/ControlRig/Public/Units/Hierarchy/RigUnit_Metadata.h b/Engine/Plugins/Animation/ControlRig/Source/ControlRig/Public/Units/Hierarchy/RigUnit_Metadata.h
--- a/Engine/Plugins/Animation/ControlRig/Source/ControlRig/Public/Units/Hierarchy/RigUnit_Metadata.h
+++ b/Engine/Plugins/Animation/ControlRig/Source/ControlRig/Public/Units/Hierarchy/RigUnit_Metadata.h
@@ -183,7 +183,7 @@ inline bool FRigDispatch_MetadataBase::Equals(const TArray<FLinearColor>& A, con
 }
 
 /*
- * Sets some metadata for the provided item
+ * Gets some metadata for the provided item
  */
 USTRUCT(meta=(DisplayName="Get Metadata"))
 struct FRigDispatch_GetMetadata : public FRigDispatch_MetadataBase
@@ -803,7 +803,7 @@ struct FRigUnit_FindItemsWithMetadataTag : public FRigUnit
 };
 
 /**
- * Returns all items with a specific tag
+ * Returns all items with a specific set of tags
  */
 USTRUCT(meta=(DisplayName="Find Items with multiple Tags", Category="Hierarchy", DocumentationPolicy = "Strict", Keywords="MetadataExists,HasKey,SupportsMetadata", NodeColor="0.462745, 1,0, 0.329412", Varying))
 struct FRigUnit_FindItemsWithMetadataTagArray : public FRigUnit
diff --git a/Engine/Plugins/Animation/ControlRig/Source/ControlRig/Public/Units/Hierarchy/RigUnit_SetControlTransform.h b/Engine/Plugins/Animation/ControlRig/Source/ControlRig/Public/Units/Hierarchy/RigUnit_SetControlTransform.h
--- a/Engine/Plugins/Animation/ControlRig/Source/ControlRig/Public/Units/Hierarchy/RigUnit_SetControlTransform.h
+++ b/Engine/Plugins/Animation/ControlRig/Source/ControlRig/Public/Units/Hierarchy/RigUnit_SetControlTransform.h
@@ -30,7 +30,7 @@ struct FRigUnit_SetControlBool : public FRigUnitMutable
 	FName Control;
 
 	/**
-	 * The bool value to set for the given Control.
+	 * The bool value to get / set for the given Control.
 	 */
 	UPROPERTY(meta = (Input, Output))
 	bool BoolValue;
@@ -49,12 +49,12 @@ struct FRigUnit_SetMultiControlBool_Entry
 		: BoolValue(false)
 	{}
 	/**
-	 * The name of the Control to set the transform for.
+	 * The name of the Control to set the value for.
 	 */
 	UPROPERTY(meta = (Input, CustomWidget = "ControlName"))
 	FName Control;
 	/**
-	 * The bool value to set for the given Control.
+	 * The bool value to get / set for the given Control.
 	 */
 	UPROPERTY(meta = (Input))
 	bool BoolValue;
@@ -105,7 +105,7 @@ struct FRigUnit_SetControlFloat : public FRigUnitMutable
 	UE_API virtual void Execute() override;
 
 	/**
-	 * The name of the Control to set the transform for.
+	 * The name of the Control to set the value for.
 	 */
 	UPROPERTY(meta = (Input, CustomWidget = "ControlName" ))
 	FName Control;
@@ -117,7 +117,7 @@ struct FRigUnit_SetControlFloat : public FRigUnitMutable
 	float Weight;
 
 	/**
-	 * The transform value to set for the given Control.
+	 * The float value to get / set for the given Control.
 	 */
 	UPROPERTY(meta = (Input, Output, UIMin = "0.0", UIMax = "1.0"))
 	float FloatValue;
@@ -136,13 +136,13 @@ struct FRigUnit_SetMultiControlFloat_Entry
 		: FloatValue(0.f)
 	{}
 	/**
-	 * The name of the Control to set the transform for.
+	 * The name of the Control to set the value for.
 	 */
 	UPROPERTY(meta = (Input, CustomWidget = "ControlName" ))
 	FName Control;
 
 	/**
-	 * The transform value to set for the given Control.
+	 * The float value to get / set for the given Control.
 	 */
 	UPROPERTY(meta = (Input, UIMin = "0.0", UIMax = "1.0"))
 	float FloatValue; 
@@ -187,7 +187,7 @@ struct FRigUnit_SetMultiControlFloat : public FRigUnitMutable
 /**
  * SetControlInteger is used to perform a change in the hierarchy by setting a single control's int32 value.
  */
-USTRUCT(meta=(DisplayName="Set Control Integer", Category="Controls", DocumentationPolicy="Strict", Keywords = "SetControlInteger,SetGizmoInteger", TemplateName="SetControlValue", NodeColor = "0.0 0.36470600962638855 1.0"))
+USTRUCT(meta=(DisplayName="Set Control Integer", Category="Controls", DocumentationPolicy="Strict", Keywords = "SetControlInteger,SetGizmoInteger", TemplateName="SetControlValue", NodeColor = "0.0 0.36470600962638855 1.0", Deprecated="5.8.0"))
 struct FRigUnit_SetControlInteger : public FRigUnitMutable
 {
 	GENERATED_BODY()
@@ -201,8 +201,11 @@ struct FRigUnit_SetControlInteger : public FRigUnitMutable
 	RIGVM_METHOD()
 	UE_API virtual void Execute() override;
 
+	RIGVM_METHOD()
+	UE_API virtual FRigVMStructUpgradeInfo GetUpgradeInfo() const override;
+
 	/**
-	 * The name of the Control to set the transform for.
+	 * The name of the Control to set the value for.
 	 */
 	UPROPERTY(meta = (Input, CustomWidget = "ControlName" ))
 	FName Control;
@@ -214,7 +217,47 @@ struct FRigUnit_SetControlInteger : public FRigUnitMutable
 	int32 Weight;
 
 	/**
-	 * The transform value to set for the given Control.
+	 * The integer value to get / set for the given Control.
+	 */
+	UPROPERTY(meta = (Input, Output))
+	int32 IntegerValue;
+
+	// Used to cache the internally used bone index
+	UPROPERTY()
+	FCachedRigElement CachedControlIndex;
+};
+
+/**
+ * SetControlInteger is used to perform a change in the hierarchy by setting a single control's int32 value.
+ */
+USTRUCT(meta=(DisplayName="Set Control Integer", Category="Controls", DocumentationPolicy="Strict", Keywords = "SetControlInteger,SetGizmoInteger", TemplateName="SetControlValue", NodeColor = "0.0 0.36470600962638855 1.0"))
+struct FRigUnit_SetControlIntegerFloatWeight : public FRigUnitMutable
+{
+	GENERATED_BODY()
+
+	FRigUnit_SetControlIntegerFloatWeight()
+		: Weight(1.f)
+		, IntegerValue(0)
+		, CachedControlIndex(FCachedRigElement())
+	{}
+
+	RIGVM_METHOD()
+	UE_API virtual void Execute() override;
+
+	/**
+	 * The name of the Control to set the value for.
+	 */
+	UPROPERTY(meta = (Input, CustomWidget = "ControlName" ))
+	FName Control;
+
+	/**
+	 * The weight of the change - how much the change should be applied
+	 */
+	UPROPERTY(meta = (Input, UIMin = "0.0", UIMax = "1.0"))
+	float Weight;
+
+	/**
+	 * The integer value to get / set for the given Control.
 	 */
 	UPROPERTY(meta = (Input, Output))
 	int32 IntegerValue;
@@ -233,13 +276,13 @@ struct FRigUnit_SetMultiControlInteger_Entry
 		: IntegerValue(0)
 	{}
 	/**
-	 * The name of the Control to set the transform for.
+	 * The name of the Control to set the value for.
 	 */
 	UPROPERTY(meta = (Input, CustomWidget = "ControlName"))
 	FName Control;
 
 	/**
-	 * The transform value to set for the given Control.
+	 * The integer value to get / set for the given Control.
 	 */
 	UPROPERTY(meta = (Input))
 	int32 IntegerValue;
@@ -297,7 +340,7 @@ struct FRigUnit_SetControlVector2D : public FRigUnitMutable
 	UE_API virtual void Execute() override;
 
 	/**
-	 * The name of the Control to set the transform for.
+	 * The name of the Control to set the value for.
 	 */
 	UPROPERTY(meta = (Input, CustomWidget = "ControlName" ))
 	FName Control;
@@ -309,7 +352,7 @@ struct FRigUnit_SetControlVector2D : public FRigUnitMutable
 	float Weight;
 
 	/**
-	 * The transform value to set for the given Control.
+	 * The 2d vector value to get / set for the given Control.
 	 */
 	UPROPERTY(meta = (Input, Output))
 	FVector2D Vector;
@@ -329,13 +372,13 @@ struct FRigUnit_SetMultiControlVector2D_Entry
 	{}
 
 	/**
-	 * The name of the Control to set the transform for.
+	 * The name of the Control to set the value for.
 	 */
 	UPROPERTY(meta = (Input, CustomWidget = "ControlName"))
 	FName Control;
 
 	/**
-	 * The transform value to set for the given Control.
+	 * The 2d vector value to get / set for the given Control.
 	 */
 	UPROPERTY(meta = (Input))
 	FVector2D Vector;
@@ -394,7 +437,7 @@ struct FRigUnit_SetControlVector : public FRigUnitMutable
 	UE_API virtual void Execute() override;
 
 	/**
-	 * The name of the Control to set the transform for.
+	 * The name of the Control to set the value for.
 	 */
 	UPROPERTY(meta = (Input, CustomWidget = "ControlName" ))
 	FName Control;
@@ -406,7 +449,7 @@ struct FRigUnit_SetControlVector : public FRigUnitMutable
 	float Weight;
 
 	/**
-	 * The transform value to set for the given Control.
+	 * The vector value to get / set for the given Control.
 	 */
 	UPROPERTY(meta = (Input, Output))
 	FVector Vector;
@@ -442,7 +485,7 @@ struct FRigUnit_SetControlRotator : public FRigUnitMutable
 	UE_API virtual void Execute() override;
 
 	/**
-	 * The name of the Control to set the transform for.
+	 * The name of the Control to set the value for.
 	 */
 	UPROPERTY(meta = (Input, CustomWidget = "ControlName" ))
 	FName Control;
@@ -454,7 +497,7 @@ struct FRigUnit_SetControlRotator : public FRigUnitMutable
 	float Weight;
 
 	/**
-	 * The transform value to set for the given Control.
+	 * The vector value to get / set for the given Control.
 	 */
 	UPROPERTY(meta = (Input, Output))
 	FRotator Rotator;
@@ -483,13 +526,13 @@ struct FRigUnit_SetMultiControlRotator_Entry
 	}
 
 	/**
-	 * The name of the Control to set the transform for.
+	 * The name of the Control to set the value for.
 	 */
 	UPROPERTY(meta = (Input, CustomWidget = "ControlName"))
 	FName Control;
 
 	/**
-	 * The transform value to set for the given Control.
+	 * The rotator value to get / set for the given Control.
 	 */
 	UPROPERTY(meta = (Input))
 	FRotator Rotator;
@@ -554,7 +597,7 @@ struct FRigUnit_SetControlTransform : public FRigUnitMutable
 	UE_API virtual void Execute() override;
 
 	/**
-	 * The name of the Control to set the transform for.
+	 * The name of the Control to set the value for.
 	 */
 	UPROPERTY(meta = (Input, CustomWidget = "ControlName" ))
 	FName Control;
@@ -566,7 +609,7 @@ struct FRigUnit_SetControlTransform : public FRigUnitMutable
 	float Weight;
 
 	/**
-	 * The transform value to set for the given Control.
+	 * The transform value to get / set for the given Control.
 	 */
 	UPROPERTY(meta = (Input, Output))
 	FTransform Transform;
diff --git a/Engine/Plugins/Animation/ControlRig/Source/ControlRig/Public/Units/Hierarchy/RigUnit_SetRelativeTransform.h b/Engine/Plugins/Animation/ControlRig/Source/ControlRig/Public/Units/Hierarchy/RigUnit_SetRelativeTransform.h
--- a/Engine/Plugins/Animation/ControlRig/Source/ControlRig/Public/Units/Hierarchy/RigUnit_SetRelativeTransform.h
+++ b/Engine/Plugins/Animation/ControlRig/Source/ControlRig/Public/Units/Hierarchy/RigUnit_SetRelativeTransform.h
@@ -104,14 +104,14 @@ struct FRigUnit_SetRelativeTranslationForItem : public FRigUnitMutable
 	UE_API virtual void Execute() override;
 
 	/**
-	 * The child item to set the transform for
+	 * The child item to set the translation for
 	 */
 	UPROPERTY(meta = (Input, ExpandByDefault))
 	FRigElementKey Child;
 
 	/**
 	 * The parent item to use.
-	 * The child transform will be set in the space of the parent.
+	 * The child translation will be set in the space of the parent.
 	 */
 	UPROPERTY(meta = (Input, ExpandByDefault))
 	FRigElementKey Parent;
@@ -123,7 +123,7 @@ struct FRigUnit_SetRelativeTranslationForItem : public FRigUnitMutable
 	UPROPERTY(meta = (Input))
 	bool bParentInitial;
 
-	// The transform of the child item relative to the provided parent
+	// The translation of the child item relative to the provided parent
 	UPROPERTY(meta=(Input))
 	FVector Value;
 
@@ -173,14 +173,14 @@ struct FRigUnit_SetRelativeRotationForItem : public FRigUnitMutable
 	UE_API virtual void Execute() override;
 
 	/**
-	 * The child item to set the transform for
+	 * The child item to set the rotation for
 	 */
 	UPROPERTY(meta = (Input, ExpandByDefault))
 	FRigElementKey Child;
 
 	/**
 	 * The parent item to use.
-	 * The child transform will be set in the space of the parent.
+	 * The child rotation will be set in the space of the parent.
 	 */
 	UPROPERTY(meta = (Input, ExpandByDefault))
 	FRigElementKey Parent;
@@ -192,7 +192,7 @@ struct FRigUnit_SetRelativeRotationForItem : public FRigUnitMutable
 	UPROPERTY(meta = (Input))
 	bool bParentInitial;
 
-	// The transform of the child item relative to the provided parent
+	// The rotation of the child item relative to the provided parent
 	UPROPERTY(meta=(Input))
 	FQuat Value;
 
diff --git a/Engine/Plugins/Animation/ControlRig/Source/ControlRig/Public/Units/Highlevel/Hierarchy/RigUnit_AimBone.h b/Engine/Plugins/Animation/ControlRig/Source/ControlRig/Public/Units/Highlevel/Hierarchy/RigUnit_AimBone.h
--- a/Engine/Plugins/Animation/ControlRig/Source/ControlRig/Public/Units/Highlevel/Hierarchy/RigUnit_AimBone.h
+++ b/Engine/Plugins/Animation/ControlRig/Source/ControlRig/Public/Units/Highlevel/Hierarchy/RigUnit_AimBone.h
@@ -280,9 +280,8 @@ struct FRigUnit_AimBone : public FRigUnit_HighlevelBaseMutable
 	RIGVM_METHOD()
 	UE_API virtual FRigVMStructUpgradeInfo GetUpgradeInfo() const override;
 };
-
 /**
- * Aligns the rotation of a primary and secondary axis of a bone to a global target.
+ * Aligns the rotation of a primary and secondary axis of an item to a global target.
  * Note: This node operates in global space!
  */
 USTRUCT(meta=(DisplayName="Aim", Category="Hierarchy", Keywords="Lookat"))
diff --git a/Engine/Plugins/Runtime/RigVM/Source/RigVM/Private/RigVMCore/RigVMStruct.cpp b/Engine/Plugins/Runtime/RigVM/Source/RigVM/Private/RigVMCore/RigVMStruct.cpp
--- a/Engine/Plugins/Runtime/RigVM/Source/RigVM/Private/RigVMCore/RigVMStruct.cpp
+++ b/Engine/Plugins/Runtime/RigVM/Source/RigVM/Private/RigVMCore/RigVMStruct.cpp
@@ -101,8 +101,13 @@ FRigVMStructUpgradeInfo FRigVMStructUpgradeInfo::MakeFromStructToFactory(UScript
 	check(InFactoryStruct);
 
 	FRigVMRegistry& Registry = FRigVMRegistry::Get();
-	const FRigVMDispatchFactory* Factory = Registry.FindOrAddDispatchFactory(InFactoryStruct);
-	check(Factory);
+	const FName FactoryName = FRigVMDispatchFactory::GetFactoryName(InFactoryStruct);
+	const FRigVMDispatchFactory* Factory = Registry.FindDispatchFactory(FactoryName);
+	if (!Factory)
+	{
+		Factory = Registry.FindOrAddDispatchFactory(InFactoryStruct);
+		check(Factory);
+	}
 	FRigVMTemplateTypeMap Types = GetTypeMapFromStruct(InRigVMStruct);
 	{
 		TArray<FRigVMTemplateArgumentInfo> Infos;
diff --git a/Engine/Plugins/Runtime/RigVM/Source/RigVM/Private/RigVMFunctions/RigVMDispatch_Array.cpp b/Engine/Plugins/Runtime/RigVM/Source/RigVM/Private/RigVMFunctions/RigVMDispatch_Array.cpp
--- a/Engine/Plugins/Runtime/RigVM/Source/RigVM/Private/RigVMFunctions/RigVMDispatch_Array.cpp
+++ b/Engine/Plugins/Runtime/RigVM/Source/RigVM/Private/RigVMFunctions/RigVMDispatch_Array.cpp
@@ -221,7 +221,7 @@ FText FRigVMDispatch_ArrayMake::GetNodeTooltip(const FRigVMTemplateTypeMap& InTy
 
 FText FRigVMDispatch_ArrayMake::GetArgumentTooltip(const FName& InArgumentName, TRigVMTypeIndex InTypeIndex) const
 {
-	if(ValuesName == ArrayName)
+	if(InArgumentName == ValuesName)
 	{
 		return LOCTEXT("ArrayMake_ValuesArgumentToolTip", "The elements of the array to create.");
 	}
@@ -929,7 +929,7 @@ FText FRigVMDispatch_ArrayReverse::GetArgumentTooltip(const FName& InArgumentNam
 {
 	if(InArgumentName == ArrayName)
 	{
-		return LOCTEXT("ArrayReverse_ArrayArgumentToolTip", "The array reverse.");
+		return LOCTEXT("ArrayReverse_ArrayArgumentToolTip", "The array to reverse.");
 	}
 	return FText();
 }
diff --git a/Engine/Plugins/Runtime/RigVM/Source/RigVM/Private/RigVMFunctions/RigVMDispatch_CastEnum.cpp b/Engine/Plugins/Runtime/RigVM/Source/RigVM/Private/RigVMFunctions/RigVMDispatch_CastEnum.cpp
--- a/Engine/Plugins/Runtime/RigVM/Source/RigVM/Private/RigVMFunctions/RigVMDispatch_CastEnum.cpp
+++ b/Engine/Plugins/Runtime/RigVM/Source/RigVM/Private/RigVMFunctions/RigVMDispatch_CastEnum.cpp
@@ -49,6 +49,19 @@ FText FRigVMDispatch_CastEnumToInt::GetNodeTooltip(const FRigVMTemplateTypeMap&
 	return LOCTEXT("CastEnumToolTip", "Casts from enum to int");
 }
 
+FText FRigVMDispatch_CastEnumToInt::GetArgumentTooltip(const FName& InArgumentName, TRigVMTypeIndex InTypeIndex) const
+{
+	if (InArgumentName == ValueName)
+	{
+		return LOCTEXT("ValueToolTip", "The enum value to cast to an integer");
+	}
+	if (InArgumentName == ResultName)
+	{
+		return LOCTEXT("ResultToolTip", "The resulting integer value");
+	}
+	return FRigVMDispatchFactory::GetArgumentTooltip(InArgumentName, InTypeIndex);
+}
+
 #endif
 
 void FRigVMDispatch_CastEnumToInt::Execute(FRigVMExtendedExecuteContext& InContext, FRigVMMemoryHandleArray Handles, FRigVMPredicateBranchArray RigVMBranches)
@@ -143,6 +156,19 @@ FText FRigVMDispatch_CastIntToEnum::GetNodeTooltip(const FRigVMTemplateTypeMap&
 	return LOCTEXT("CastIntToEnumToolTip", "Casts from int to enum");
 }
 
+FText FRigVMDispatch_CastIntToEnum::GetArgumentTooltip(const FName& InArgumentName, TRigVMTypeIndex InTypeIndex) const
+{
+	if (InArgumentName == ValueName)
+	{
+		return LOCTEXT("ValueToolTip", "The integer value to cast to an enum value");
+	}
+	if (InArgumentName == ResultName)
+	{
+		return LOCTEXT("ResultToolTip", "The resulting enum value");
+	}
+	return FRigVMDispatchFactory::GetArgumentTooltip(InArgumentName, InTypeIndex);
+}
+
 #endif
 
 void FRigVMDispatch_CastIntToEnum::Execute(FRigVMExtendedExecuteContext& InContext, FRigVMMemoryHandleArray Handles, FRigVMPredicateBranchArray RigVMBranches)
diff --git a/Engine/Plugins/Runtime/RigVM/Source/RigVM/Private/RigVMFunctions/RigVMDispatch_CastObject.cpp b/Engine/Plugins/Runtime/RigVM/Source/RigVM/Private/RigVMFunctions/RigVMDispatch_CastObject.cpp
--- a/Engine/Plugins/Runtime/RigVM/Source/RigVM/Private/RigVMFunctions/RigVMDispatch_CastObject.cpp
+++ b/Engine/Plugins/Runtime/RigVM/Source/RigVM/Private/RigVMFunctions/RigVMDispatch_CastObject.cpp
@@ -19,6 +19,7 @@ const TArray<FRigVMTemplateArgumentInfo>& FRigVMDispatch_CastObject::GetArgument
 		TArray<FRigVMTemplateArgumentInfo> Infos;
 		Infos.Emplace(ValueName, ERigVMPinDirection::Input, ElementCategories);
 		Infos.Emplace(ResultName, ERigVMPinDirection::Output, INDEX_NONE);
+		Infos.Emplace(SuccessName, ERigVMPinDirection::Output, RigVMTypeUtils::TypeIndex::Bool);
 		
 		CachedArgumentInfos = BuildArgumentListFromPrimaryArgument(Infos, ValueName, InRegistry);
 	}
@@ -36,7 +37,8 @@ bool FRigVMDispatch_CastObject::GetPermutationsFromArgumentType(const FName& InA
 			OutPermutations.Add(
 	{
 				{ ValueName, InTypeIndex },
-				{ ResultName, Type }
+				{ ResultName, Type },
+				{ SuccessName, RigVMTypeUtils::TypeIndex::Bool }
 			});
 		}
 	}
@@ -48,7 +50,8 @@ bool FRigVMDispatch_CastObject::GetPermutationsFromArgumentType(const FName& InA
 			OutPermutations.Add(
 	{
 				{ ValueName, Type },
-				{ ResultName, InTypeIndex }
+				{ ResultName, InTypeIndex },
+				{ SuccessName, RigVMTypeUtils::TypeIndex::Bool }
 			});
 		}
 	}
@@ -102,23 +105,43 @@ FText FRigVMDispatch_CastObject::GetNodeTooltip(const FRigVMTemplateTypeMap& InT
 	return LOCTEXT("CastToolTip", "Casts between object types");
 }
 
+FText FRigVMDispatch_CastObject::GetArgumentTooltip(const FName& InArgumentName, TRigVMTypeIndex InTypeIndex) const
+{
+	if (InArgumentName == ValueName)
+	{
+		return LOCTEXT("ValueToolTip", "The object to cast to a new type");
+	}
+	if (InArgumentName == ResultName)
+	{
+		return LOCTEXT("ResultToolTip", "The resulting cast object. This may be potentially nullptr as well if the cast was not successful.");
+	}
+	if (InArgumentName == SuccessName)
+	{
+		return LOCTEXT("SuccessToolTip", "A boolean result indicating if the cast was successful. If the cast failed the result will be nullptr.");
+	}
+	return FRigVMDispatchFactory::GetArgumentTooltip(InArgumentName, InTypeIndex);
+}
+
 #endif
 
 void FRigVMDispatch_CastObject::Execute(FRigVMExtendedExecuteContext& InContext, FRigVMMemoryHandleArray Handles, FRigVMPredicateBranchArray RigVMBranches)
 {
 	UObject** ValueObjectPtr = (UObject**)Handles[0].GetInputData();
 	UObject** ResultObjectPtr = (UObject**)Handles[1].GetOutputData();
+	bool* Success = (bool*)Handles[2].GetOutputData();
 	const FObjectPropertyBase* ValueObjectProperty = CastFieldChecked<FObjectPropertyBase>(Handles[0].GetResolvedProperty());
 	const FObjectPropertyBase* ResultObjectProperty = CastFieldChecked<FObjectPropertyBase>(Handles[1].GetProperty());
-	if(ValueObjectPtr && ResultObjectPtr && ValueObjectProperty && ResultObjectProperty)
+	if(ValueObjectPtr && ResultObjectPtr && ValueObjectProperty && ResultObjectProperty && Success)
 	{
 		if(*ValueObjectPtr && (*ValueObjectPtr)->GetClass()->IsChildOf(ResultObjectProperty->PropertyClass))
 		{
 			*ResultObjectPtr = *ValueObjectPtr;
+			*Success = true;
 		}
 		else
 		{
 			*ResultObjectPtr = nullptr;
+			*Success = false;
 		}
 	}
 }
diff --git a/Engine/Plugins/Runtime/RigVM/Source/RigVM/Private/RigVMFunctions/RigVMDispatch_Constant.cpp b/Engine/Plugins/Runtime/RigVM/Source/RigVM/Private/RigVMFunctions/RigVMDispatch_Constant.cpp
--- a/Engine/Plugins/Runtime/RigVM/Source/RigVM/Private/RigVMFunctions/RigVMDispatch_Constant.cpp
+++ b/Engine/Plugins/Runtime/RigVM/Source/RigVM/Private/RigVMFunctions/RigVMDispatch_Constant.cpp
@@ -60,6 +60,15 @@ FText FRigVMDispatch_Constant::GetNodeTooltip(const FRigVMTemplateTypeMap& InTyp
 	return LOCTEXT("ConstantToolTip", "Provides a constant value to the graph");
 }
 
+FText FRigVMDispatch_Constant::GetArgumentTooltip(const FName& InArgumentName, TRigVMTypeIndex InTypeIndex) const
+{
+	if (InArgumentName == ValueName)
+	{
+		return LOCTEXT("ConstantValueToolTip", "The constant value defined in this node");
+	}
+	return FRigVMDispatch_CoreBase::GetArgumentTooltip(InArgumentName, InTypeIndex);
+}
+
 #endif
 
 void FRigVMDispatch_Constant::Execute(FRigVMExtendedExecuteContext& InContext, FRigVMMemoryHandleArray Handles, FRigVMPredicateBranchArray RigVMBranches)
diff --git a/Engine/Plugins/Runtime/RigVM/Source/RigVM/Private/RigVMFunctions/RigVMDispatch_Core.cpp b/Engine/Plugins/Runtime/RigVM/Source/RigVM/Private/RigVMFunctions/RigVMDispatch_Core.cpp
--- a/Engine/Plugins/Runtime/RigVM/Source/RigVM/Private/RigVMFunctions/RigVMDispatch_Core.cpp
+++ b/Engine/Plugins/Runtime/RigVM/Source/RigVM/Private/RigVMFunctions/RigVMDispatch_Core.cpp
@@ -6,6 +6,7 @@
 #include "RigVMTypeUtils.h"
 
 #include UE_INLINE_GENERATED_CPP_BY_NAME(RigVMDispatch_Core)
+#define LOCTEXT_NAMESPACE "RigVMDispatch_Core"
 
 FName FRigVMDispatch_CoreEquals::GetArgumentNameForOperandIndex(int32 InOperandIndex, int32 InTotalOperands, FRigVMRegistryHandle& InRegistry) const
 {
@@ -43,6 +44,27 @@ FRigVMTemplateTypeMap FRigVMDispatch_CoreEquals::OnNewArgumentType(const FName&
 	return Types;
 }
 
+#if WITH_EDITOR
+
+FText FRigVMDispatch_CoreEquals::GetArgumentTooltip(const FName& InArgumentName, TRigVMTypeIndex InTypeIndex) const
+{
+	if (InArgumentName == AName)
+	{
+		return LOCTEXT("CoreEqualsAToolTip", "The first value to compare");
+	}
+	if (InArgumentName == BName)
+	{
+		return LOCTEXT("CoreEqualsAToolTip", "The second value to compare");
+	}
+	if (InArgumentName == ResultName)
+	{
+		return LOCTEXT("CoreEqualsAToolTip", "The boolean result of the comparison");
+	}
+	return FRigVMDispatch_CoreBase::GetArgumentTooltip(InArgumentName, InTypeIndex);
+}
+
+#endif
+
 FRigVMFunctionPtr FRigVMDispatch_CoreEquals::GetDispatchFunctionImpl(const FRigVMTemplateTypeMap& InTypes, const FRigVMRegistryHandle& InRegistry) const
 {
 	const TRigVMTypeIndex TypeIndex = InTypes.FindChecked(AName);
@@ -138,3 +160,5 @@ const TArray<FRigVMTemplateArgumentInfo>& FRigVMDispatch_CoreNotEquals::GetArgum
 	}
 	return CachedArgumentInfos;
 }
+
+#undef LOCTEXT_NAMESPACE
diff --git a/Engine/Plugins/Runtime/RigVM/Source/RigVM/Private/RigVMFunctions/RigVMDispatch_If.cpp b/Engine/Plugins/Runtime/RigVM/Source/RigVM/Private/RigVMFunctions/RigVMDispatch_If.cpp
--- a/Engine/Plugins/Runtime/RigVM/Source/RigVM/Private/RigVMFunctions/RigVMDispatch_If.cpp
+++ b/Engine/Plugins/Runtime/RigVM/Source/RigVM/Private/RigVMFunctions/RigVMDispatch_If.cpp
@@ -5,6 +5,7 @@
 #include "RigVMCore/RigVMStruct.h"
 
 #include UE_INLINE_GENERATED_CPP_BY_NAME(RigVMDispatch_If)
+#define LOCTEXT_NAMESPACE "RigVMDispatch_If"
 
 FName FRigVMDispatch_If::GetArgumentNameForOperandIndex(int32 InOperandIndex, int32 InTotalOperands, FRigVMRegistryHandle& InRegistry) const
 {
@@ -60,6 +61,27 @@ FString FRigVMDispatch_If::GetArgumentMetaData(const FName& InArgumentName, cons
 	return FRigVMDispatch_CoreBase::GetArgumentMetaData(InArgumentName, InMetaDataKey);
 }
 
+FText FRigVMDispatch_If::GetArgumentTooltip(const FName& InArgumentName, TRigVMTypeIndex InTypeIndex) const
+{
+	if (InArgumentName == ConditionName)
+	{
+		return LOCTEXT("ConditionToolTip", "The condition on which we'll decide between the True and False values");
+	}
+	if (InArgumentName == TrueName)
+	{
+		return LOCTEXT("TrueToolTip", "The value to use if the condition is True");
+	}
+	if (InArgumentName == FalseName)
+	{
+		return LOCTEXT("FalseToolTip", "The value to use if the condition is False");
+	}
+	if (InArgumentName == ResultName)
+	{
+		return LOCTEXT("ResultToolTip", "The resulting value - either True or False based on the condition");
+	}
+	return FRigVMDispatch_CoreBase::GetArgumentTooltip(InArgumentName, InTypeIndex);
+}
+
 #endif
 
 FRigVMFunctionPtr FRigVMDispatch_If::GetDispatchFunctionImpl(const FRigVMTemplateTypeMap& InTypes, const FRigVMRegistryHandle& InRegistry) const
@@ -91,3 +113,4 @@ void FRigVMDispatch_If::Execute(FRigVMExtendedExecuteContext& InContext, FRigVMM
 		InputHandle.GetResolvedProperty(), Input);
 }
 
+#undef LOCTEXT_NAMESPACE
diff --git a/Engine/Plugins/Runtime/RigVM/Source/RigVM/Private/RigVMFunctions/RigVMDispatch_MakeStruct.cpp b/Engine/Plugins/Runtime/RigVM/Source/RigVM/Private/RigVMFunctions/RigVMDispatch_MakeStruct.cpp
--- a/Engine/Plugins/Runtime/RigVM/Source/RigVM/Private/RigVMFunctions/RigVMDispatch_MakeStruct.cpp
+++ b/Engine/Plugins/Runtime/RigVM/Source/RigVM/Private/RigVMFunctions/RigVMDispatch_MakeStruct.cpp
@@ -89,6 +89,19 @@ FString FRigVMDispatch_MakeStruct::GetKeywords() const
 	return Keywords;
 }
 
+FText FRigVMDispatch_MakeStruct::GetArgumentTooltip(const FName& InArgumentName, TRigVMTypeIndex InTypeIndex) const
+{
+	if (InArgumentName == ElementsName)
+	{
+		return LOCTEXT("ElementsToolTip", "The elements of the structure to make");
+	}
+	if (InArgumentName == StructName)
+	{
+		return LOCTEXT("StructToolTip", "The created structure");
+	}
+	return FRigVMDispatch_CoreBase::GetArgumentTooltip(InArgumentName, InTypeIndex);
+}
+
 #endif
 
 void FRigVMDispatch_MakeStruct::Execute(FRigVMExtendedExecuteContext& InContext, FRigVMMemoryHandleArray Handles, FRigVMPredicateBranchArray RigVMBranches)
@@ -139,6 +152,19 @@ FString FRigVMDispatch_BreakStruct::GetKeywords() const
 	return Keywords;
 }
 
+FText FRigVMDispatch_BreakStruct::GetArgumentTooltip(const FName& InArgumentName, TRigVMTypeIndex InTypeIndex) const
+{
+	if (InArgumentName == ElementsName)
+	{
+		return LOCTEXT("ElementsToolTip", "The elements of the structure");
+	}
+	if (InArgumentName == StructName)
+	{
+		return LOCTEXT("StructToolTip", "The structure to break up into its elements");
+	}
+	return FRigVMDispatch_MakeStruct::GetArgumentTooltip(InArgumentName, InTypeIndex);
+}
+
 #endif
 
 #undef LOCTEXT_NAMESPACE
diff --git a/Engine/Plugins/Runtime/RigVM/Source/RigVM/Private/RigVMFunctions/RigVMDispatch_Print.cpp b/Engine/Plugins/Runtime/RigVM/Source/RigVM/Private/RigVMFunctions/RigVMDispatch_Print.cpp
--- a/Engine/Plugins/Runtime/RigVM/Source/RigVM/Private/RigVMFunctions/RigVMDispatch_Print.cpp
+++ b/Engine/Plugins/Runtime/RigVM/Source/RigVM/Private/RigVMFunctions/RigVMDispatch_Print.cpp
@@ -6,6 +6,7 @@
 #include "Kismet/KismetSystemLibrary.h"
 
 #include UE_INLINE_GENERATED_CPP_BY_NAME(RigVMDispatch_Print)
+#define LOCTEXT_NAMESPACE "RigVMDispatch_Print"
 
 FName FRigVMDispatch_Print::GetArgumentNameForOperandIndex(int32 InOperandIndex, int32 InTotalOperands, FRigVMRegistryHandle& InRegistry) const
 {
@@ -85,6 +86,31 @@ FString FRigVMDispatch_Print::GetArgumentMetaData(const FName& InArgumentName, c
 	return FRigVMDispatchFactory::GetArgumentMetaData(InArgumentName, InMetaDataKey);
 }
 
+FText FRigVMDispatch_Print::GetArgumentTooltip(const FName& InArgumentName, TRigVMTypeIndex InTypeIndex) const
+{
+	if (InArgumentName == PrefixName)
+	{
+		return LOCTEXT("PrefixToolTip", "The prefix to place at the start of the printed out message");
+	}
+	if (InArgumentName == ValueName)
+	{
+		return LOCTEXT("ValueToolTip", "The value to print");
+	}
+	if (InArgumentName == EnabledName)
+	{
+		return LOCTEXT("EnabledToolTip", "If False the printing of the message will be skipped");
+	}
+	if (InArgumentName == ScreenDurationName)
+	{
+		return LOCTEXT("ScreenDurationToolTip", "The duration of the printed message to appear on the screen");
+	}
+	if (InArgumentName == ScreenColorName)
+	{
+		return LOCTEXT("ScreenColorToolTip", "The color of the message on the screen");
+	}
+	return FRigVMDispatchFactory::GetArgumentTooltip(InArgumentName, InTypeIndex);
+}
+
 #endif
 
 void FRigVMDispatch_Print::Execute(FRigVMExtendedExecuteContext& InContext, FRigVMMemoryHandleArray Handles, FRigVMPredicateBranchArray Predicates)
@@ -129,3 +155,4 @@ void FRigVMDispatch_Print::Execute(FRigVMExtendedExecuteContext& InContext, FRig
 #endif
 }
 
+#undef LOCTEXT_NAMESPACE
diff --git a/Engine/Plugins/Runtime/RigVM/Source/RigVM/Private/RigVMFunctions/RigVMDispatch_Select.cpp b/Engine/Plugins/Runtime/RigVM/Source/RigVM/Private/RigVMFunctions/RigVMDispatch_Select.cpp
--- a/Engine/Plugins/Runtime/RigVM/Source/RigVM/Private/RigVMFunctions/RigVMDispatch_Select.cpp
+++ b/Engine/Plugins/Runtime/RigVM/Source/RigVM/Private/RigVMFunctions/RigVMDispatch_Select.cpp
@@ -6,6 +6,7 @@
 #include "RigVMCore/RigVMStruct.h"
 
 #include UE_INLINE_GENERATED_CPP_BY_NAME(RigVMDispatch_Select)
+#define LOCTEXT_NAMESPACE "RigVMDispatch_Select"
 
 FName FRigVMDispatch_SelectInt32::GetArgumentNameForOperandIndex(int32 InOperandIndex, int32 InTotalOperands, FRigVMRegistryHandle& InRegistry) const
 {
@@ -99,6 +100,23 @@ FString FRigVMDispatch_SelectInt32::GetArgumentDefaultValue(const FName& InArgum
 	return FRigVMDispatch_CoreBase::GetArgumentDefaultValue(InArgumentName, InTypeIndex);
 }
 
+FText FRigVMDispatch_SelectInt32::GetArgumentTooltip(const FName& InArgumentName, TRigVMTypeIndex InTypeIndex) const
+{
+	if (InArgumentName == IndexName)
+	{
+		return LOCTEXT("IndexToolTip", "The index of the value to select");
+	}
+	if (InArgumentName == ValuesName)
+	{
+		return LOCTEXT("ValuesToolTip", "The fixed array of values to select from");
+	}
+	if (InArgumentName == ResultName)
+	{
+		return LOCTEXT("ResultToolTip", "The resulting selected value from the values array based on the index");
+	}
+	return FRigVMDispatch_CoreBase::GetArgumentTooltip(InArgumentName, InTypeIndex);
+}
+
 #endif
 
 
@@ -146,3 +164,4 @@ void FRigVMDispatch_SelectInt32::Execute(FRigVMExtendedExecuteContext& InContext
 		InputHandle.GetResolvedProperty(), Input);
 }
 
+#undef LOCTEXT_NAMESPACE
diff --git a/Engine/Plugins/Runtime/RigVM/Source/RigVM/Private/RigVMFunctions/RigVMDispatch_Switch.cpp b/Engine/Plugins/Runtime/RigVM/Source/RigVM/Private/RigVMFunctions/RigVMDispatch_Switch.cpp
--- a/Engine/Plugins/Runtime/RigVM/Source/RigVM/Private/RigVMFunctions/RigVMDispatch_Switch.cpp
+++ b/Engine/Plugins/Runtime/RigVM/Source/RigVM/Private/RigVMFunctions/RigVMDispatch_Switch.cpp
@@ -7,6 +7,7 @@
 #include "AutoRTFM.h"
 
 #include UE_INLINE_GENERATED_CPP_BY_NAME(RigVMDispatch_Switch)
+#define LOCTEXT_NAMESPACE "RigVMDispatch_Switch"
 
 FName FRigVMDispatch_SwitchInt32::GetArgumentNameForOperandIndex(int32 InOperandIndex, int32 InTotalOperands, FRigVMRegistryHandle& InRegistry) const
 {
@@ -81,6 +82,19 @@ FName FRigVMDispatch_SwitchInt32::GetDisplayNameForArgument(const FName& InArgum
 	return FRigVMDispatch_CoreBase::GetDisplayNameForArgument(InArgumentName);
 }
 
+FText FRigVMDispatch_SwitchInt32::GetArgumentTooltip(const FName& InArgumentName, TRigVMTypeIndex InTypeIndex) const
+{
+	if (InArgumentName == IndexName)
+	{
+		return LOCTEXT("IndexToolTip", "The index of the case to switch to");
+	}
+	if (InArgumentName == CasesName)
+	{
+		return LOCTEXT("CasesToolTip", "The fixed set of cases to switch to");
+	}
+	return FRigVMDispatch_CoreBase::GetArgumentTooltip(InArgumentName, InTypeIndex);
+}
+
 #endif
 
 const TArray<FName>& FRigVMDispatch_SwitchInt32::GetControlFlowBlocks_Impl(const FRigVMDispatchContext& InContext) const
@@ -175,3 +189,5 @@ FName FRigVMDispatch_SwitchInt32::GetCaseDisplayName(int32 InIndex)
     NamesFromInt.Add(InIndex, Result);
     return Result;
 }
+
+#undef LOCTEXT_NAMESPACE
diff --git a/Engine/Plugins/Runtime/RigVM/Source/RigVM/Private/RigVMFunctions/RigVMFunction_String.cpp b/Engine/Plugins/Runtime/RigVM/Source/RigVM/Private/RigVMFunctions/RigVMFunction_String.cpp
--- a/Engine/Plugins/Runtime/RigVM/Source/RigVM/Private/RigVMFunctions/RigVMFunction_String.cpp
+++ b/Engine/Plugins/Runtime/RigVM/Source/RigVM/Private/RigVMFunctions/RigVMFunction_String.cpp
@@ -6,6 +6,7 @@
 #include "RigVMStringUtils.h"
 
 #include UE_INLINE_GENERATED_CPP_BY_NAME(RigVMFunction_String)
+#define LOCTEXT_NAMESPACE "RigVMFunction_String"
 
 FRigVMFunction_StringConcat_Execute()
 {
@@ -266,6 +267,23 @@ FRigVMTemplateTypeMap FRigDispatch_ToString::OnNewArgumentType(const FName& InAr
 	return Types;
 }
 
+#if WITH_EDITOR
+
+FText FRigDispatch_ToString::GetArgumentTooltip(const FName& InArgumentName, TRigVMTypeIndex InTypeIndex) const
+{
+	if (InArgumentName == TEXT("Value"))
+	{
+		return LOCTEXT("ToStringValueToolTip", "The value to convert to a string");
+	}
+	if (InArgumentName == TEXT("Result"))
+	{
+		return LOCTEXT("ToStringResultToolTip", "The resulting string representation of the value");
+	}
+	return FRigVMDispatchFactory::GetArgumentTooltip(InArgumentName, InTypeIndex);
+}
+
+#endif
+
 FRigVMFunctionPtr FRigDispatch_ToString::GetDispatchFunctionImpl(const FRigVMTemplateTypeMap& InTypes, const FRigVMRegistryHandle& InRegistry) const
 {
 	// treat name and string special
@@ -334,6 +352,23 @@ FRigVMTemplateTypeMap FRigDispatch_FromString::OnNewArgumentType(const FName& In
 	return Types;
 }
 
+#if WITH_EDITOR
+
+FText FRigDispatch_FromString::GetArgumentTooltip(const FName& InArgumentName, TRigVMTypeIndex InTypeIndex) const
+{
+	if (InArgumentName == TEXT("String"))
+	{
+		return LOCTEXT("FromStringStringToolTip", "The string to convert the value from");
+	}
+	if (InArgumentName == TEXT("Result"))
+	{
+		return LOCTEXT("FromStringResultToolTip", "The resulting value based on the string");
+	}
+	return FRigVMDispatchFactory::GetArgumentTooltip(InArgumentName, InTypeIndex);
+}
+
+#endif
+
 FRigVMFunctionPtr FRigDispatch_FromString::GetDispatchFunctionImpl(const FRigVMTemplateTypeMap& InTypes, const FRigVMRegistryHandle& InRegistry) const
 {
 	// treat name and string special
@@ -421,3 +456,4 @@ void FRigDispatch_FromString::Execute(FRigVMExtendedExecuteContext& InContext, F
 	}
 }
 
+#undef LOCTEXT_NAMESPACE
diff --git a/Engine/Plugins/Runtime/RigVM/Source/RigVM/Public/RigVMFunctions/Debug/RigVMFunction_DebugPrimitives.h b/Engine/Plugins/Runtime/RigVM/Source/RigVM/Public/RigVMFunctions/Debug/RigVMFunction_DebugPrimitives.h
--- a/Engine/Plugins/Runtime/RigVM/Source/RigVM/Public/RigVMFunctions/Debug/RigVMFunction_DebugPrimitives.h
+++ b/Engine/Plugins/Runtime/RigVM/Source/RigVM/Public/RigVMFunctions/Debug/RigVMFunction_DebugPrimitives.h
@@ -196,7 +196,7 @@ struct FRigVMFunction_DebugArcNoSpace : public FRigVMFunction_DebugBaseMutable
 	UPROPERTY(meta = (Input))
 	float Radius;
 
-	// The minimum angle in degress of the arc
+	// The minimum angle in degrees of the arc
 	UPROPERTY(meta = (Input))
 	float MinimumDegrees;
 
diff --git a/Engine/Plugins/Runtime/RigVM/Source/RigVM/Public/RigVMFunctions/Debug/RigVMFunction_VisualLog.h b/Engine/Plugins/Runtime/RigVM/Source/RigVM/Public/RigVMFunctions/Debug/RigVMFunction_VisualLog.h
--- a/Engine/Plugins/Runtime/RigVM/Source/RigVM/Public/RigVMFunctions/Debug/RigVMFunction_VisualLog.h
+++ b/Engine/Plugins/Runtime/RigVM/Source/RigVM/Public/RigVMFunctions/Debug/RigVMFunction_VisualLog.h
@@ -237,7 +237,7 @@ struct FRigVMFunction_VisualLogCircle : public FRigVMFunction_VisualLogWireframe
 	UPROPERTY(meta=(Input))
 	float Radius = 10.0f;
 
-	/** The thickness the circle */
+	/** The thickness of the circle */
 	UPROPERTY(meta=(Input))
 	float Thickness = 0.0f;
 };
@@ -259,7 +259,7 @@ struct FRigVMFunction_VisualLogSegment : public FRigVMFunction_VisualLogObject
 	UPROPERTY(meta=(Input))
 	FVector SegmentEnd = FVector::ZeroVector;
 
-	/** The thickness the circle */
+	/** The thickness of the segment */
 	UPROPERTY(meta=(Input))
 	float Thickness = 0.0f;
 };
diff --git a/Engine/Plugins/Runtime/RigVM/Source/RigVM/Public/RigVMFunctions/Execution/RigVMFunction_UserDefinedEvent.h b/Engine/Plugins/Runtime/RigVM/Source/RigVM/Public/RigVMFunctions/Execution/RigVMFunction_UserDefinedEvent.h
--- a/Engine/Plugins/Runtime/RigVM/Source/RigVM/Public/RigVMFunctions/Execution/RigVMFunction_UserDefinedEvent.h
+++ b/Engine/Plugins/Runtime/RigVM/Source/RigVM/Public/RigVMFunctions/Execution/RigVMFunction_UserDefinedEvent.h
@@ -31,7 +31,7 @@ struct FRigVMFunction_UserDefinedEvent : public FRigVMStruct
 	UPROPERTY(EditAnywhere, Transient, DisplayName = "Execute", Category = "UserDefinedEvent", meta = (Output))
 	FRigVMExecutePin ExecutePin;
 
-	// True if the current interaction is a rotation
+	// The name of the event. Given this name the event can be invoked using the run event node.
 	UPROPERTY(EditAnywhere, Transient, Category = "UserDefinedEvent", meta = (Input,Constant,DetailsOnly))
 	FName EventName;
 };
diff --git a/Engine/Plugins/Runtime/RigVM/Source/RigVM/Public/RigVMFunctions/Math/RigVMFunction_MathDouble.h b/Engine/Plugins/Runtime/RigVM/Source/RigVM/Public/RigVMFunctions/Math/RigVMFunction_MathDouble.h
--- a/Engine/Plugins/Runtime/RigVM/Source/RigVM/Public/RigVMFunctions/Math/RigVMFunction_MathDouble.h
+++ b/Engine/Plugins/Runtime/RigVM/Source/RigVM/Public/RigVMFunctions/Math/RigVMFunction_MathDouble.h
@@ -979,15 +979,15 @@ struct FRigVMFunction_MathDoubleLawOfCosine : public FRigVMFunction_MathDoubleBa
 	UPROPERTY(meta = (Input))
 	double C;
 
-	// The angle of between B and C
+	// The angle between B and C
 	UPROPERTY(meta = (Output))
 	double AlphaAngle;
 
-	// The angle of between A and C
+	// The angle between A and C
 	UPROPERTY(meta = (Output))
 	double BetaAngle;
 
-	// The angle of between A and B
+	// The angle between A and B
 	UPROPERTY(meta = (Output))
 	double GammaAngle;
 
@@ -1025,7 +1025,7 @@ struct FRigVMFunction_MathDoubleArraySum : public FRigVMFunction_MathDoubleBase
 	RIGVM_METHOD()
 	RIGVM_API virtual void Execute() override;
 
-	// The array of number to add up
+	// The array of numbers to add up
 	UPROPERTY(meta = (Input))
 	TArray<double> Array;
 
diff --git a/Engine/Plugins/Runtime/RigVM/Source/RigVM/Public/RigVMFunctions/Math/RigVMFunction_MathFloat.h b/Engine/Plugins/Runtime/RigVM/Source/RigVM/Public/RigVMFunctions/Math/RigVMFunction_MathFloat.h
--- a/Engine/Plugins/Runtime/RigVM/Source/RigVM/Public/RigVMFunctions/Math/RigVMFunction_MathFloat.h
+++ b/Engine/Plugins/Runtime/RigVM/Source/RigVM/Public/RigVMFunctions/Math/RigVMFunction_MathFloat.h
@@ -1060,7 +1060,7 @@ struct FRigVMFunction_MathFloatArraySum : public FRigVMFunction_MathFloatBase
 	RIGVM_METHOD()
 	UE_API virtual void Execute() override;
 
-	// The array of number to add up
+	// The array of numbers to add up
 	UPROPERTY(meta = (Input))
 	TArray<float> Array;
 
diff --git a/Engine/Plugins/Runtime/RigVM/Source/RigVM/Public/RigVMFunctions/Math/RigVMFunction_MathInt.h b/Engine/Plugins/Runtime/RigVM/Source/RigVM/Public/RigVMFunctions/Math/RigVMFunction_MathInt.h
--- a/Engine/Plugins/Runtime/RigVM/Source/RigVM/Public/RigVMFunctions/Math/RigVMFunction_MathInt.h
+++ b/Engine/Plugins/Runtime/RigVM/Source/RigVM/Public/RigVMFunctions/Math/RigVMFunction_MathInt.h
@@ -604,7 +604,7 @@ struct FRigVMFunction_MathIntToString : public FRigVMFunction_MathIntBase
 	int32 Number;
 
 	/*
-	 * For positive numbers you can pad the result to a required with
+	 * For positive numbers you can pad the result to a required width
 	 * so rather than '13' return '00013' for a padded size of 5.
 	 */
 	UPROPERTY(meta = (Input))
@@ -637,7 +637,7 @@ struct FRigVMFunction_MathIntToName : public FRigVMFunction_MathIntBase
 	int32 Number;
 
 	/*
-	 * For positive numbers you can pad the result to a required with
+	 * For positive numbers you can pad the result to a required width
 	 * so rather than '13' return '00013' for a padded size of 5.
 	 */
 	UPROPERTY(meta = (Input))
diff --git a/Engine/Plugins/Runtime/RigVM/Source/RigVM/Public/RigVMFunctions/Math/RigVMFunction_MathQuaternion.h b/Engine/Plugins/Runtime/RigVM/Source/RigVM/Public/RigVMFunctions/Math/RigVMFunction_MathQuaternion.h
--- a/Engine/Plugins/Runtime/RigVM/Source/RigVM/Public/RigVMFunctions/Math/RigVMFunction_MathQuaternion.h
+++ b/Engine/Plugins/Runtime/RigVM/Source/RigVM/Public/RigVMFunctions/Math/RigVMFunction_MathQuaternion.h
@@ -761,7 +761,7 @@ struct FRigVMFunction_MathQuaternionRotationOrder : public FRigVMFunction_MathBa
 };
 
 /**
- * Returns the relative local transform within a parent's transform
+ * Returns the relative local rotation within a parent's rotation
  */
 USTRUCT(meta=(DisplayName="Make Relative", TemplateName="Make Relative", Keywords="Local,Global,Absolute"))
 struct FRigVMFunction_MathQuaternionMakeRelative : public FRigVMFunction_MathQuaternionBase
@@ -790,7 +790,7 @@ struct FRigVMFunction_MathQuaternionMakeRelative : public FRigVMFunction_MathQua
 };
 
 /**
- * Returns the absolute global transform within a parent's transform
+ * Returns the absolute global rotation within a parent's rotation
  */
 USTRUCT(meta = (DisplayName = "Make Absolute", TemplateName="Make Absolute", Keywords = "Local,Global,Relative"))
 struct FRigVMFunction_MathQuaternionMakeAbsolute : public FRigVMFunction_MathQuaternionBase
diff --git a/Engine/Plugins/Runtime/RigVM/Source/RigVM/Public/RigVMFunctions/Math/RigVMFunction_MathTransform.h b/Engine/Plugins/Runtime/RigVM/Source/RigVM/Public/RigVMFunctions/Math/RigVMFunction_MathTransform.h
--- a/Engine/Plugins/Runtime/RigVM/Source/RigVM/Public/RigVMFunctions/Math/RigVMFunction_MathTransform.h
+++ b/Engine/Plugins/Runtime/RigVM/Source/RigVM/Public/RigVMFunctions/Math/RigVMFunction_MathTransform.h
@@ -362,7 +362,7 @@ struct FRigVMFunction_MathTransformAccumulateArray : public FRigVMFunction_MathT
 };
 
 /**
- * Returns the negative value
+ * Returns the inverted transform of the input
  */
 USTRUCT(meta=(DisplayName="Inverse", TemplateName="Inverse"))
 struct FRigVMFunction_MathTransformInverse : public FRigVMFunction_MathTransformUnaryOp
@@ -579,7 +579,7 @@ struct FRigVMFunction_MathTransformArrayToSRT : public FRigVMFunction_MathTransf
 };
 
 /**
- * Clamps a position using a plane collision, cylindric collision or spherical collision.
+ * Clamps a transform's position using a plane collision, cylindric collision or spherical collision.
  */
 USTRUCT(meta = (DisplayName = "Clamp Spatially", TemplateName="ClampSpatially", Keywords = "Collide,Collision"))
 struct FRigVMFunction_MathTransformClampSpatially : public FRigVMFunction_MathTransformBase
diff --git a/Engine/Plugins/Runtime/RigVM/Source/RigVM/Public/RigVMFunctions/Math/RigVMFunction_MathVector.h b/Engine/Plugins/Runtime/RigVM/Source/RigVM/Public/RigVMFunctions/Math/RigVMFunction_MathVector.h
--- a/Engine/Plugins/Runtime/RigVM/Source/RigVM/Public/RigVMFunctions/Math/RigVMFunction_MathVector.h
+++ b/Engine/Plugins/Runtime/RigVM/Source/RigVM/Public/RigVMFunctions/Math/RigVMFunction_MathVector.h
@@ -350,7 +350,7 @@ struct FRigVMFunction_MathVectorCeil : public FRigVMFunction_MathVectorUnaryOp
 };
 
 /**
- * Returns the closest higher full number (integer) of the value for each component
+ * Returns the closest full number (integer) of the value for each component
  */
 USTRUCT(meta=(DisplayName="Round"))
 struct FRigVMFunction_MathVectorRound : public FRigVMFunction_MathVectorUnaryOp
diff --git a/Engine/Plugins/Runtime/RigVM/Source/RigVM/Public/RigVMFunctions/Math/RigVMFunction_Noise.h b/Engine/Plugins/Runtime/RigVM/Source/RigVM/Public/RigVMFunctions/Math/RigVMFunction_Noise.h
--- a/Engine/Plugins/Runtime/RigVM/Source/RigVM/Public/RigVMFunctions/Math/RigVMFunction_Noise.h
+++ b/Engine/Plugins/Runtime/RigVM/Source/RigVM/Public/RigVMFunctions/Math/RigVMFunction_Noise.h
@@ -52,7 +52,7 @@ struct FRigVMFunction_NoiseFloat : public FRigVMFunction_MathBase
 };
 
 /**
- * Generates a float through a noise fluctuation process between a min and a max through speed
+ * Generates a double through a noise fluctuation process between a min and a max through speed
  */
 USTRUCT(meta=(DisplayName="Noise (Double)", Category="Math|Noise", TemplateName="Noise"))
 struct FRigVMFunction_NoiseDouble : public FRigVMFunction_MathBase
diff --git a/Engine/Plugins/Runtime/RigVM/Source/RigVM/Public/RigVMFunctions/RigVMDispatch_CastEnum.h b/Engine/Plugins/Runtime/RigVM/Source/RigVM/Public/RigVMFunctions/RigVMDispatch_CastEnum.h
--- a/Engine/Plugins/Runtime/RigVM/Source/RigVM/Public/RigVMFunctions/RigVMDispatch_CastEnum.h
+++ b/Engine/Plugins/Runtime/RigVM/Source/RigVM/Public/RigVMFunctions/RigVMDispatch_CastEnum.h
@@ -30,6 +30,7 @@ public:
 #if WITH_EDITOR
 	RIGVM_API virtual FString GetNodeTitle(const FRigVMTemplateTypeMap& InTypes) const override;
 	RIGVM_API virtual FText GetNodeTooltip(const FRigVMTemplateTypeMap& InTypes) const override;
+	RIGVM_API virtual FText GetArgumentTooltip(const FName& InArgumentName, TRigVMTypeIndex InTypeIndex) const override;
 #endif
 	virtual bool IsSingleton() const override { return true; }
 	
@@ -57,6 +58,7 @@ public:
 #if WITH_EDITOR
 	RIGVM_API virtual FString GetNodeTitle(const FRigVMTemplateTypeMap& InTypes) const override;
 	RIGVM_API virtual FText GetNodeTooltip(const FRigVMTemplateTypeMap& InTypes) const override;
+	RIGVM_API virtual FText GetArgumentTooltip(const FName& InArgumentName, TRigVMTypeIndex InTypeIndex) const override;
 #endif
 	virtual bool IsSingleton() const override { return true; }
 	
diff --git a/Engine/Plugins/Runtime/RigVM/Source/RigVM/Public/RigVMFunctions/RigVMDispatch_CastObject.h b/Engine/Plugins/Runtime/RigVM/Source/RigVM/Public/RigVMFunctions/RigVMDispatch_CastObject.h
--- a/Engine/Plugins/Runtime/RigVM/Source/RigVM/Public/RigVMFunctions/RigVMDispatch_CastObject.h
+++ b/Engine/Plugins/Runtime/RigVM/Source/RigVM/Public/RigVMFunctions/RigVMDispatch_CastObject.h
@@ -22,6 +22,7 @@ public:
 #if WITH_EDITOR
 	RIGVM_API virtual FString GetNodeTitle(const FRigVMTemplateTypeMap& InTypes) const override;
 	RIGVM_API virtual FText GetNodeTooltip(const FRigVMTemplateTypeMap& InTypes) const override;
+	RIGVM_API virtual FText GetArgumentTooltip(const FName& InArgumentName, TRigVMTypeIndex InTypeIndex) const override;
 #endif
 	virtual bool IsSingleton() const override { return true; }
 	
@@ -31,4 +32,5 @@ protected:
 
 	static inline const FLazyName ValueName = FLazyName(TEXT("Value"));
 	static inline const FLazyName ResultName = FLazyName(TEXT("Result"));
+	static inline const FLazyName SuccessName = FLazyName(TEXT("Success"));
 };
diff --git a/Engine/Plugins/Runtime/RigVM/Source/RigVM/Public/RigVMFunctions/RigVMDispatch_Constant.h b/Engine/Plugins/Runtime/RigVM/Source/RigVM/Public/RigVMFunctions/RigVMDispatch_Constant.h
--- a/Engine/Plugins/Runtime/RigVM/Source/RigVM/Public/RigVMFunctions/RigVMDispatch_Constant.h
+++ b/Engine/Plugins/Runtime/RigVM/Source/RigVM/Public/RigVMFunctions/RigVMDispatch_Constant.h
@@ -23,6 +23,7 @@ public:
 #if WITH_EDITOR
 	RIGVM_API virtual FString GetNodeTitle(const FRigVMTemplateTypeMap& InTypes) const override;
 	RIGVM_API virtual FText GetNodeTooltip(const FRigVMTemplateTypeMap& InTypes) const override;
+	RIGVM_API virtual FText GetArgumentTooltip(const FName& InArgumentName, TRigVMTypeIndex InTypeIndex) const override;
 #endif
 	virtual bool IsSingleton() const override { return true; }
 	virtual bool SupportsRenaming() const override { return true; }
diff --git a/Engine/Plugins/Runtime/RigVM/Source/RigVM/Public/RigVMFunctions/RigVMDispatch_Core.h b/Engine/Plugins/Runtime/RigVM/Source/RigVM/Public/RigVMFunctions/RigVMDispatch_Core.h
--- a/Engine/Plugins/Runtime/RigVM/Source/RigVM/Public/RigVMFunctions/RigVMDispatch_Core.h
+++ b/Engine/Plugins/Runtime/RigVM/Source/RigVM/Public/RigVMFunctions/RigVMDispatch_Core.h
@@ -30,6 +30,10 @@ public:
 	RIGVM_API virtual const TArray<FRigVMTemplateArgumentInfo>& GetArgumentInfos(FRigVMRegistryHandle& InRegistry) const override;
 	RIGVM_API virtual FRigVMTemplateTypeMap OnNewArgumentType(const FName& InArgumentName, TRigVMTypeIndex InTypeIndex, FRigVMRegistryHandle& InRegistry) const override;
 
+#if WITH_EDITOR
+	RIGVM_API virtual FText GetArgumentTooltip(const FName& InArgumentName, TRigVMTypeIndex InTypeIndex) const override;
+#endif
+
 protected:
 
 	RIGVM_API virtual FRigVMFunctionPtr GetDispatchFunctionImpl(const FRigVMTemplateTypeMap& InTypes, const FRigVMRegistryHandle& InRegistry) const override;
@@ -88,7 +92,7 @@ protected:
 };
 
 /*
- * Compares any two values and return true if they are identical
+ * Compares any two values and return true if they are not identical
  */
 USTRUCT(meta=(DisplayName = "Not Equals", Keywords = "Different,!=,Xor"))
 struct FRigVMDispatch_CoreNotEquals : public FRigVMDispatch_CoreEquals
diff --git a/Engine/Plugins/Runtime/RigVM/Source/RigVM/Public/RigVMFunctions/RigVMDispatch_If.h b/Engine/Plugins/Runtime/RigVM/Source/RigVM/Public/RigVMFunctions/RigVMDispatch_If.h
--- a/Engine/Plugins/Runtime/RigVM/Source/RigVM/Public/RigVMFunctions/RigVMDispatch_If.h
+++ b/Engine/Plugins/Runtime/RigVM/Source/RigVM/Public/RigVMFunctions/RigVMDispatch_If.h
@@ -28,6 +28,7 @@ public:
 
 #if WITH_EDITOR
 	RIGVM_API virtual FString GetArgumentMetaData(const FName& InArgumentName, const FName& InMetaDataKey) const override;
+	RIGVM_API virtual FText GetArgumentTooltip(const FName& InArgumentName, TRigVMTypeIndex InTypeIndex) const override;
 #endif
 
 protected:
diff --git a/Engine/Plugins/Runtime/RigVM/Source/RigVM/Public/RigVMFunctions/RigVMDispatch_MakeStruct.h b/Engine/Plugins/Runtime/RigVM/Source/RigVM/Public/RigVMFunctions/RigVMDispatch_MakeStruct.h
--- a/Engine/Plugins/Runtime/RigVM/Source/RigVM/Public/RigVMFunctions/RigVMDispatch_MakeStruct.h
+++ b/Engine/Plugins/Runtime/RigVM/Source/RigVM/Public/RigVMFunctions/RigVMDispatch_MakeStruct.h
@@ -25,6 +25,7 @@ public:
 	RIGVM_API virtual FText GetNodeTooltip(const FRigVMTemplateTypeMap& InTypes) const override;
 	RIGVM_API virtual FString GetArgumentMetaData(const FName& InArgumentName, const FName& InMetaDataKey) const override;
 	RIGVM_API virtual FString GetKeywords() const override;
+	RIGVM_API virtual FText GetArgumentTooltip(const FName& InArgumentName, TRigVMTypeIndex InTypeIndex) const override;
 #endif
 	virtual bool SupportsRenaming() const override { return true; }
 
@@ -55,6 +56,7 @@ public:
 #if WITH_EDITOR
 	RIGVM_API virtual FText GetNodeTooltip(const FRigVMTemplateTypeMap& InTypes) const override;
 	RIGVM_API virtual FString GetKeywords() const override;
+	RIGVM_API virtual FText GetArgumentTooltip(const FName& InArgumentName, TRigVMTypeIndex InTypeIndex) const override;
 #endif
 	virtual bool SupportsRenaming() const override { return false; }
 };
diff --git a/Engine/Plugins/Runtime/RigVM/Source/RigVM/Public/RigVMFunctions/RigVMDispatch_Print.h b/Engine/Plugins/Runtime/RigVM/Source/RigVM/Public/RigVMFunctions/RigVMDispatch_Print.h
--- a/Engine/Plugins/Runtime/RigVM/Source/RigVM/Public/RigVMFunctions/RigVMDispatch_Print.h
+++ b/Engine/Plugins/Runtime/RigVM/Source/RigVM/Public/RigVMFunctions/RigVMDispatch_Print.h
@@ -29,6 +29,7 @@ public:
 #if WITH_EDITOR
 	RIGVM_API virtual FString GetArgumentDefaultValue(const FName& InArgumentName, TRigVMTypeIndex InTypeIndex) const override;
 	RIGVM_API virtual FString GetArgumentMetaData(const FName& InArgumentName, const FName& InMetaDataKey) const override;
+	RIGVM_API virtual FText GetArgumentTooltip(const FName& InArgumentName, TRigVMTypeIndex InTypeIndex) const override;
 #endif
 
 protected:
diff --git a/Engine/Plugins/Runtime/RigVM/Source/RigVM/Public/RigVMFunctions/RigVMDispatch_Select.h b/Engine/Plugins/Runtime/RigVM/Source/RigVM/Public/RigVMFunctions/RigVMDispatch_Select.h
--- a/Engine/Plugins/Runtime/RigVM/Source/RigVM/Public/RigVMFunctions/RigVMDispatch_Select.h
+++ b/Engine/Plugins/Runtime/RigVM/Source/RigVM/Public/RigVMFunctions/RigVMDispatch_Select.h
@@ -27,6 +27,7 @@ public:
 #if WITH_EDITOR
 	RIGVM_API virtual FString GetArgumentMetaData(const FName& InArgumentName, const FName& InMetaDataKey) const override;
 	RIGVM_API virtual FString GetArgumentDefaultValue(const FName& InArgumentName, TRigVMTypeIndex InTypeIndex) const override;
+	RIGVM_API virtual FText GetArgumentTooltip(const FName& InArgumentName, TRigVMTypeIndex InTypeIndex) const override;
 #endif
 	virtual bool IsSingleton() const override { return true; } 
 
diff --git a/Engine/Plugins/Runtime/RigVM/Source/RigVM/Public/RigVMFunctions/RigVMDispatch_Switch.h b/Engine/Plugins/Runtime/RigVM/Source/RigVM/Public/RigVMFunctions/RigVMDispatch_Switch.h
--- a/Engine/Plugins/Runtime/RigVM/Source/RigVM/Public/RigVMFunctions/RigVMDispatch_Switch.h
+++ b/Engine/Plugins/Runtime/RigVM/Source/RigVM/Public/RigVMFunctions/RigVMDispatch_Switch.h
@@ -30,6 +30,7 @@ public:
 	RIGVM_API virtual FString GetArgumentMetaData(const FName& InArgumentName, const FName& InMetaDataKey) const override;
 	RIGVM_API virtual FString GetArgumentDefaultValue(const FName& InArgumentName, TRigVMTypeIndex InTypeIndex) const override;
 	RIGVM_API virtual FName GetDisplayNameForArgument(const FName& InArgumentName) const override;
+	RIGVM_API virtual FText GetArgumentTooltip(const FName& InArgumentName, TRigVMTypeIndex InTypeIndex) const override;
 #endif
 	
 	RIGVM_API virtual const TArray<FName>& GetControlFlowBlocks_Impl(const FRigVMDispatchContext& InContext) const override;
diff --git a/Engine/Plugins/Runtime/RigVM/Source/RigVM/Public/RigVMFunctions/RigVMFunction_ControlFlow.h b/Engine/Plugins/Runtime/RigVM/Source/RigVM/Public/RigVMFunctions/RigVMFunction_ControlFlow.h
--- a/Engine/Plugins/Runtime/RigVM/Source/RigVM/Public/RigVMFunctions/RigVMFunction_ControlFlow.h
+++ b/Engine/Plugins/Runtime/RigVM/Source/RigVM/Public/RigVMFunctions/RigVMFunction_ControlFlow.h
@@ -41,6 +41,7 @@ struct FRigVMFunction_ControlFlowBranch : public FRigVMFunction_ControlFlowBase
 		return Blocks;
 	}
 
+	// The input execution pin to hook up to the graph
 	UPROPERTY(meta=(Input, DisplayName="Execute"))
 	FRigVMExecuteContext ExecuteContext;
 
diff --git a/Engine/Plugins/Runtime/RigVM/Source/RigVM/Public/RigVMFunctions/RigVMFunction_String.h b/Engine/Plugins/Runtime/RigVM/Source/RigVM/Public/RigVMFunctions/RigVMFunction_String.h
--- a/Engine/Plugins/Runtime/RigVM/Source/RigVM/Public/RigVMFunctions/RigVMFunction_String.h
+++ b/Engine/Plugins/Runtime/RigVM/Source/RigVM/Public/RigVMFunctions/RigVMFunction_String.h
@@ -603,6 +603,10 @@ public:
 	virtual const TArray<FRigVMTemplateArgumentInfo>& GetArgumentInfos(FRigVMRegistryHandle& InRegistry) const override;
 	virtual FRigVMTemplateTypeMap OnNewArgumentType(const FName& InArgumentName, TRigVMTypeIndex InTypeIndex, FRigVMRegistryHandle& InRegistry) const override;
 
+#if WITH_EDITOR
+	virtual FText GetArgumentTooltip(const FName& InArgumentName, TRigVMTypeIndex InTypeIndex) const override;
+#endif
+
 protected:
 
 	virtual FRigVMFunctionPtr GetDispatchFunctionImpl(const FRigVMTemplateTypeMap& InTypes, const FRigVMRegistryHandle& InRegistry) const override;
@@ -622,6 +626,10 @@ public:
 	virtual const TArray<FRigVMTemplateArgumentInfo>& GetArgumentInfos(FRigVMRegistryHandle& InRegistry) const override;
 	virtual FRigVMTemplateTypeMap OnNewArgumentType(const FName& InArgumentName, TRigVMTypeIndex InTypeIndex, FRigVMRegistryHandle& InRegistry) const override;
 
+#if WITH_EDITOR
+	virtual FText GetArgumentTooltip(const FName& InArgumentName, TRigVMTypeIndex InTypeIndex) const override;
+#endif
+
 protected:
 
 	virtual FRigVMFunctionPtr GetDispatchFunctionImpl(const FRigVMTemplateTypeMap& InTypes, const FRigVMRegistryHandle& InRegistry) const override;
diff --git a/Engine/Plugins/Runtime/RigVM/Source/RigVM/Public/RigVMFunctions/Simulation/RigVMFunction_AlphaInterp.h b/Engine/Plugins/Runtime/RigVM/Source/RigVM/Public/RigVMFunctions/Simulation/RigVMFunction_AlphaInterp.h
--- a/Engine/Plugins/Runtime/RigVM/Source/RigVM/Public/RigVMFunctions/Simulation/RigVMFunction_AlphaInterp.h
+++ b/Engine/Plugins/Runtime/RigVM/Source/RigVM/Public/RigVMFunctions/Simulation/RigVMFunction_AlphaInterp.h
@@ -179,7 +179,7 @@ struct FRigVMFunction_AlphaInterpVector : public FRigVMFunction_SimBase
 };
 
 /**
- * Takes in a vector value and outputs an accumulated value with a customized scale and clamp
+ * Takes in a quaternion value and outputs an accumulated value with a customized scale and clamp
  */
 USTRUCT(meta=(DisplayName="Alpha Interpolate", Keywords="Alpha,Lerp,LinearInterpolate", Category = "Simulation|Time", TemplateName="AlphaInterp", MenuDescSuffix = "(Quat)"))
 struct FRigVMFunction_AlphaInterpQuat : public FRigVMFunction_SimBase
diff --git a/Engine/Plugins/Runtime/RigVM/Source/RigVMDeveloper/Private/RigVMJsonUtils.cpp b/Engine/Plugins/Runtime/RigVM/Source/RigVMDeveloper/Private/RigVMJsonUtils.cpp
--- a/Engine/Plugins/Runtime/RigVM/Source/RigVMDeveloper/Private/RigVMJsonUtils.cpp
+++ b/Engine/Plugins/Runtime/RigVM/Source/RigVMDeveloper/Private/RigVMJsonUtils.cpp
@@ -25,17 +25,53 @@ FAutoConsoleCommand FCmdRigVMJsonPrintAvailableNodes
 (
 	TEXT("RigVM.JSON.PrintAvailableNodes"),
 	TEXT("Prints out a json string containing information about all available nodes."),
-	FConsoleCommandDelegate::CreateLambda([]()
+	FConsoleCommandWithArgsDelegate::CreateLambda([](const TArray<FString>& Args)
 	{
 		FRigVMRegistryReadLock Registry;
+		RigVMJsonUtils::FIntrospectionSettings Settings(Registry);
+
+		for (const FString& Arg : Args)
+		{
+			if (Arg.Equals(TEXT("exclude-deprecated"), ESearchCase::IgnoreCase))
+			{
+				Settings.bExcludeDeprecated = true;
+			}
+			else if (Arg.Equals(TEXT("exclude-functions"), ESearchCase::IgnoreCase))
+			{
+				Settings.bExcludeFunctions = true;
+			}
+			else if (Arg.Equals(TEXT("exclude-dispatches"), ESearchCase::IgnoreCase))
+			{
+				Settings.bExcludeDispatches = true;
+			}
+			else
+			{
+				UE_LOG(LogRigVMDeveloper, Display, TEXT("Valid Arguments:\n"
+					"exclude-deprecated - only includes the current / non-deprecated nodes\n"
+					"exclude-functions - excludes RigVM struct based functions\n"
+					"exclude-dispatches - excludes RigVM dispatch factories"
+				));
+				return;
+			}
+		}
 			
-		TSharedPtr<FJsonObject> JsonObject = RigVMJsonUtils::AvailableNodesToJson(Registry, [](const UScriptStruct* InStruct) -> bool
+		TSharedPtr<FJsonObject> JsonObject = RigVMJsonUtils::AvailableNodesToJson(Settings, [](const UScriptStruct* InStruct) -> bool
 		{
 			check(InStruct);
 			const UPackage* Package = InStruct->GetPackage();
-			if (Package->GetName().Contains(TEXT("EngineTest")))
+
+			static const TArray<FString> StringsToFilter = {
+				TEXT("EngineTest"),
+				TEXT("TestPlugin"),
+				TEXT("RigVMExample"),
+				TEXT("UAF"),
+			};
+			for (const FString& StringToFilter : StringsToFilter)
 			{
-				return false;
+				if (Package->GetName().Contains(StringToFilter))
+				{
+					return false;
+				}
 			}
 			return true;
 		});
@@ -48,7 +84,7 @@ FAutoConsoleCommand FCmdRigVMJsonPrintAvailableNodes
 			{
 				UE_LOG(LogRigVMDeveloper, Display, TEXT("Json:\n\n%s"), *Text);
 #if WITH_EDITOR
-				UE_LOG(LogRigVMDeveloper, Display, TEXT("The content has also been copied to the clipboard."), *Text);
+				UE_LOG(LogRigVMDeveloper, Display, TEXT("The content has also been copied to the clipboard."));
 				FPlatformApplicationMisc::ClipboardCopy(*Text);
 #endif
 			}
@@ -76,8 +112,9 @@ FAutoConsoleCommand FCmdRigVMJsonPrintGraph
 			return;
 		}
 			
-		FRigVMRegistryReadLock Registry;
-		if (TSharedPtr<FJsonObject> JsonObject = RigVMJsonUtils::ToJson(Registry, CastChecked<URigVMGraph>(GraphObject)))
+			FRigVMRegistryReadLock Registry;
+			RigVMJsonUtils::FIntrospectionSettings Settings(Registry);
+			if (TSharedPtr<FJsonObject> JsonObject = RigVMJsonUtils::ToJson(Settings, CastChecked<URigVMGraph>(GraphObject)))
 		{
 			FString Text;
 			RigVMJsonUtils::FJsonWriterRef Writer = RigVMJsonUtils::MakeJsonWriter(Text);
@@ -85,7 +122,7 @@ FAutoConsoleCommand FCmdRigVMJsonPrintGraph
 			{
 				UE_LOG(LogRigVMDeveloper, Display, TEXT("Json:\n\n%s"), *Text);
 #if WITH_EDITOR
-				UE_LOG(LogRigVMDeveloper, Display, TEXT("The content has also been copied to the clipboard."), *Text);
+				UE_LOG(LogRigVMDeveloper, Display, TEXT("The content has also been copied to the clipboard."));
 				FPlatformApplicationMisc::ClipboardCopy(*Text);
 #endif
 			}
@@ -121,7 +158,8 @@ FAutoConsoleCommand FCmdRigVMJsonPrintAsset
 		}
 
 		FRigVMRegistryReadLock Registry;
-		if (TSharedPtr<FJsonObject> JsonObject = RigVMJsonUtils::ToJson(Registry, Asset))
+		RigVMJsonUtils::FIntrospectionSettings Settings(Registry);
+		if (TSharedPtr<FJsonObject> JsonObject = RigVMJsonUtils::ToJson(Settings, Asset))
 		{
 			FString Text;
 			RigVMJsonUtils::FJsonWriterRef Writer = RigVMJsonUtils::MakeJsonWriter(Text);
@@ -133,21 +171,33 @@ FAutoConsoleCommand FCmdRigVMJsonPrintAsset
 	})
 );
 
+namespace RigVMJsonUtils
+{
+	static FString PostProcessToolTip(const FString& InToolTip)
+	{
+		FString Result = InToolTip;
+		Result.ReplaceInline(TEXT("\\n* "), TEXT("\\n"));
+		Result.RemoveFromStart(TEXT("* "));
+		Result.RemoveFromStart(TEXT("*"));
+		return Result;
+	}
+}
+
 RigVMJsonUtils::FJsonWriterRef RigVMJsonUtils::MakeJsonWriter(FString& OutJsonText)
 {
 	return FJsonFactory::Create(&OutJsonText);
 }
 
-TSharedPtr<FJsonObject> RigVMJsonUtils::AvailableNodesToJson(FRigVMRegistryHandle& InRegistry, TFunction<bool(const UScriptStruct*)> InFilterCallback)
+TSharedPtr<FJsonObject> RigVMJsonUtils::AvailableNodesToJson(FIntrospectionSettings& InSettings, TFunction<bool(const UScriptStruct*)> InFilterCallback)
 {
-	if (!InRegistry.IsValid())
+	if (!InSettings.Registry.IsValid())
 	{
 		return nullptr;
 	}
 	
 	TArray<TSharedPtr<FJsonValue>> FunctionJsonValues, FactoryJsonValues;
-	const TChunkedArray<FRigVMFunction>& Functions = InRegistry->GetFunctions_NoLock();
-	const TArray<FRigVMDispatchFactory*>& Factories = InRegistry->GetFactories_NoLock();
+	const TChunkedArray<FRigVMFunction>& Functions = InSettings.Registry->GetFunctions_NoLock();
+	const TArray<FRigVMDispatchFactory*>& Factories = InSettings.Registry->GetFactories_NoLock();
 
 	for (const FRigVMFunction& Function : Functions)
 	{
@@ -160,7 +210,7 @@ TSharedPtr<FJsonObject> RigVMJsonUtils::AvailableNodesToJson(FRigVMRegistryHandl
 		{
 			continue;
 		}
-		if (const TSharedPtr<FJsonObject> JsonObject = ToJson(InRegistry, &Function))
+		if (const TSharedPtr<FJsonObject> JsonObject = ToJson(InSettings, &Function))
 		{
 			FunctionJsonValues.Add(MakeShared<FJsonValueObject>(JsonObject));
 		}
@@ -173,7 +223,7 @@ TSharedPtr<FJsonObject> RigVMJsonUtils::AvailableNodesToJson(FRigVMRegistryHandl
 		{
 			continue;
 		}
-		if (const TSharedPtr<FJsonObject> JsonObject = ToJson(InRegistry, Factory))
+		if (const TSharedPtr<FJsonObject> JsonObject = ToJson(InSettings, Factory))
 		{
 			FactoryJsonValues.Add(MakeShared<FJsonValueObject>(JsonObject));
 		}
@@ -185,17 +235,25 @@ TSharedPtr<FJsonObject> RigVMJsonUtils::AvailableNodesToJson(FRigVMRegistryHandl
 	return JsonObject;
 }
 
-TSharedPtr<FJsonObject> RigVMJsonUtils::ToJson(FRigVMRegistryHandle& InRegistry, const FRigVMFunction* InFunction)
+TSharedPtr<FJsonObject> RigVMJsonUtils::ToJson(FIntrospectionSettings& InSettings, const FRigVMFunction* InFunction)
 {
+	if (InSettings.bExcludeFunctions)
+	{
+		return nullptr;
+	}
+	
 	check(InFunction);
 	const UScriptStruct* Struct = InFunction->Struct;
 	check(Struct);
 
 	TSharedPtr<FJsonObject> JsonObject = MakeShared<FJsonObject>();
-	GenericStructAttributesToJson(InRegistry, Struct, JsonObject.Get());
+	if (!GenericStructAttributesToJson(InSettings, Struct, JsonObject.Get()))
+	{
+		return nullptr;
+	}
 
 	FString TemplateName;
-	if (const FRigVMTemplate* Template = InFunction->GetTemplate_NoLock(InRegistry))
+	if (const FRigVMTemplate* Template = InFunction->GetTemplate_NoLock(InSettings.Registry))
 	{
 		TemplateName = Template->GetName().ToString();
 	}
@@ -206,6 +264,15 @@ TSharedPtr<FJsonObject> RigVMJsonUtils::ToJson(FRigVMRegistryHandle& InRegistry,
 	FStructOnScope StructOnScope(Struct);
 	FRigVMStruct* StructInstance = reinterpret_cast<FRigVMStruct*>(StructOnScope.GetStructMemory());
 
+	const FRigVMStructUpgradeInfo UpgradeInfo = StructInstance->GetUpgradeInfo();
+	if (UpgradeInfo.IsValid())
+	{
+		if (const UScriptStruct* NewStruct = UpgradeInfo.GetNewStruct())
+		{
+			JsonObject->SetStringField(TEXT("UpgradedVersion"), NewStruct->GetStructCPPName());
+		}
+	}
+
 	TMap<FName,const FRigVMFunctionArgument*> ArgumentLookup;
 	for (int32 Index = 0; Index < InFunction->Arguments.Num(); ++Index)
 	{
@@ -262,7 +329,12 @@ TSharedPtr<FJsonObject> RigVMJsonUtils::ToJson(FRigVMRegistryHandle& InRegistry,
 		TSharedPtr<FJsonObject> ArgumentJsonObject = MakeShared<FJsonObject>();
 		ArgumentJsonObject->SetStringField(TEXT("Name"), Property->GetName());
 		ArgumentJsonObject->SetStringField(TEXT("Type"), Argument ? FString(Argument->Type) : ExecutePinType);
-		ArgumentJsonObject->SetStringField(TEXT("PropertyToolTip"), Property->GetToolTipText().ToString());
+		FString PropertyToolTip = PostProcessToolTip(Property->GetToolTipText().ToString());
+		if (PropertyToolTip.IsEmpty() && bIsExecutePin)
+		{
+			PropertyToolTip = TEXT("This pin is used to chain multiple mutable units together");
+		}
+		ArgumentJsonObject->SetStringField(TEXT("PropertyToolTip"), PropertyToolTip);
 		ArgumentJsonObject->SetStringField(TEXT("DefaultValue"), DefaultValue);
 
 		if (bIsInput && bIsOutput)
@@ -286,22 +358,44 @@ TSharedPtr<FJsonObject> RigVMJsonUtils::ToJson(FRigVMRegistryHandle& InRegistry,
 	return JsonObject;
 }
 
-TSharedPtr<FJsonObject> RigVMJsonUtils::ToJson(FRigVMRegistryHandle& InRegistry, const FRigVMDispatchFactory* InFactory)
+TSharedPtr<FJsonObject> RigVMJsonUtils::ToJson(FIntrospectionSettings& InSettings, const FRigVMDispatchFactory* InFactory)
 {
+	if (InSettings.bExcludeDispatches)
+	{
+		return nullptr;
+	}
+	
 	check(InFactory);
 	const UScriptStruct* Struct = InFactory->GetScriptStruct();
 	check(Struct);
 
 	TSharedPtr<FJsonObject> JsonObject = MakeShared<FJsonObject>();
-	GenericStructAttributesToJson(InRegistry, Struct, JsonObject.Get());
+	if (!GenericStructAttributesToJson(InSettings, Struct, JsonObject.Get()))
+	{
+		return nullptr;
+	}
 
-	const FRigVMTemplate* Template = InFactory->GetTemplate_NoLock(InRegistry);
+	const FRigVMTemplate* Template = InFactory->GetTemplate_NoLock(InSettings.Registry);
 	JsonObject->SetStringField(TEXT("TemplateName"), Template->GetName().ToString());
+	const FText ToolTip = InFactory->GetNodeTooltip({});
+	if (!ToolTip.IsEmpty())
+	{
+		JsonObject->SetStringField(TEXT("ToolTip"), PostProcessToolTip(ToolTip.ToString()));
+	}
 
 	FRigVMDispatchContext DispatchContext;
 	TArray<TSharedPtr<FJsonValue>> IO, Inputs, Outputs;
 
-	TArray<FRigVMExecuteArgument> ExecuteArguments = InFactory->GetExecuteArguments_NoLock(DispatchContext, InRegistry);
+	const FRigVMStructUpgradeInfo UpgradeInfo = InFactory->GetUpgradeInfo({}, DispatchContext);
+	if (UpgradeInfo.IsValid())
+	{
+		if (const UScriptStruct* NewStruct = UpgradeInfo.GetNewStruct())
+		{
+			JsonObject->SetStringField(TEXT("UpgradedVersion"), NewStruct->GetStructCPPName());
+		}
+	}
+
+	TArray<FRigVMExecuteArgument> ExecuteArguments = InFactory->GetExecuteArguments_NoLock(DispatchContext, InSettings.Registry);
 	for (const FRigVMExecuteArgument& ExecuteArgument : ExecuteArguments)
 	{
 		static const FString ExecutePinType = FRigVMExecuteContext::StaticStruct()->GetStructCPPName();
@@ -309,7 +403,14 @@ TSharedPtr<FJsonObject> RigVMJsonUtils::ToJson(FRigVMRegistryHandle& InRegistry,
 		TSharedPtr<FJsonObject> ArgumentJsonObject = MakeShared<FJsonObject>();
 		ArgumentJsonObject->SetStringField(TEXT("Name"), ExecuteArgument.Name.ToString());
 		ArgumentJsonObject->SetStringField(TEXT("Type"), ExecutePinType);
-		ArgumentJsonObject->SetStringField(TEXT("PropertyToolTip"), FString::Printf(TEXT("The %s execute argument."), *ExecuteArgument.Name.ToString()));
+		if (ExecuteArgument.Name == FRigVMStruct::ExecuteName || ExecuteArgument.Name == FRigVMStruct::ExecuteContextName)
+		{
+			ArgumentJsonObject->SetStringField(TEXT("PropertyToolTip"), TEXT("This pin is used to chain multiple mutable units together"));
+		}
+		else
+		{
+			ArgumentJsonObject->SetStringField(TEXT("PropertyToolTip"), FString::Printf(TEXT("The %s execute argument."), *ExecuteArgument.Name.ToString()));
+		}
 		ArgumentJsonObject->SetStringField(TEXT("DefaultValue"), FString());
 
 		if (ExecuteArgument.Direction == ERigVMPinDirection::IO)
@@ -326,7 +427,7 @@ TSharedPtr<FJsonObject> RigVMJsonUtils::ToJson(FRigVMRegistryHandle& InRegistry,
 		}
 	}
 
-	const TArray<FRigVMTemplateArgumentInfo>& ArgumentInfos = InFactory->GetArgumentInfos(InRegistry);
+	const TArray<FRigVMTemplateArgumentInfo>& ArgumentInfos = InFactory->GetArgumentInfos(InSettings.Registry);
 	for (const FRigVMTemplateArgumentInfo& ArgumentInfo : ArgumentInfos)
 	{
 		if (ArgumentInfo.Direction != ERigVMPinDirection::Input &&
@@ -337,14 +438,14 @@ TSharedPtr<FJsonObject> RigVMJsonUtils::ToJson(FRigVMRegistryHandle& InRegistry,
 			continue;
 		}
 		
-		const FRigVMTemplateArgument& Argument = ArgumentInfo.GetArgument_NoLock(InRegistry);
-		const TArray<TRigVMTypeIndex>& TypeIndices = Argument.GetSupportedTypeIndices_NoLock({}, InRegistry);
+		const FRigVMTemplateArgument& Argument = ArgumentInfo.GetArgument_NoLock(InSettings.Registry);
+		const TArray<TRigVMTypeIndex>& TypeIndices = Argument.GetSupportedTypeIndices_NoLock({}, InSettings.Registry);
 		check(!TypeIndices.IsEmpty());
 
 		FString Type;
 		if (TypeIndices.Num() == 1)
 		{
-			Type = InRegistry->GetType_NoLock(TypeIndices[0]).CPPType.ToString();
+			Type = InSettings.Registry->GetType_NoLock(TypeIndices[0]).CPPType.ToString();
 		}
 		else if (TypeIndices.Num() < 16)
 		{
@@ -352,7 +453,7 @@ TSharedPtr<FJsonObject> RigVMJsonUtils::ToJson(FRigVMRegistryHandle& InRegistry,
 			Types.Reserve(TypeIndices.Num());
 			for (const TRigVMTypeIndex& TypeIndex : TypeIndices)
 			{
-				Types.Add(InRegistry->GetType_NoLock(TypeIndex).CPPType.ToString());
+				Types.Add(InSettings.Registry->GetType_NoLock(TypeIndex).CPPType.ToString());
 			}
 			Type = RigVMStringUtils::JoinStrings(Types, TEXT(","));
 		}
@@ -364,7 +465,7 @@ TSharedPtr<FJsonObject> RigVMJsonUtils::ToJson(FRigVMRegistryHandle& InRegistry,
 			int32 NumArrayTypes = 0;
 			for (const TRigVMTypeIndex& TypeIndex : TypeIndices)
 			{
-				if (InRegistry->IsArrayType_NoLock(TypeIndex))
+				if (InSettings.Registry->IsArrayType_NoLock(TypeIndex))
 				{
 					NumArrayTypes++;
 				}
@@ -382,7 +483,7 @@ TSharedPtr<FJsonObject> RigVMJsonUtils::ToJson(FRigVMRegistryHandle& InRegistry,
 		TSharedPtr<FJsonObject> ArgumentJsonObject = MakeShared<FJsonObject>();
 		ArgumentJsonObject->SetStringField(TEXT("Name"), ArgumentInfo.Name.ToString());
 		ArgumentJsonObject->SetStringField(TEXT("Type"), Type);
-		ArgumentJsonObject->SetStringField(TEXT("PropertyToolTip"), InFactory->GetArgumentTooltip(ArgumentInfo.Name, TypeIndices[0]).ToString());
+		ArgumentJsonObject->SetStringField(TEXT("PropertyToolTip"), PostProcessToolTip(InFactory->GetArgumentTooltip(ArgumentInfo.Name, TypeIndices[0]).ToString()));
 		ArgumentJsonObject->SetStringField(TEXT("DefaultValue"), FString());
 
 		if (Argument.GetDirection() == ERigVMPinDirection::IO)
@@ -406,16 +507,17 @@ TSharedPtr<FJsonObject> RigVMJsonUtils::ToJson(FRigVMRegistryHandle& InRegistry,
 	return JsonObject;
 }
 
-void RigVMJsonUtils::GenericStructAttributesToJson(FRigVMRegistryHandle& InRegistry, const UScriptStruct* InStruct, FJsonObject* OutJsonObject)
+bool RigVMJsonUtils::GenericStructAttributesToJson(FIntrospectionSettings& InSettings, const UScriptStruct* InStruct, FJsonObject* OutJsonObject)
 {
 	check(InStruct);
 	check(OutJsonObject);
 	
-	FString CategoryMetadata, KeywordsMetadata, DisplayNameMetadata;
+	FString CategoryMetadata, KeywordsMetadata, DisplayNameMetadata, DeprecatedMetadata;
 #if WITH_EDITOR
 	InStruct->GetStringMetaDataHierarchical(FRigVMStruct::CategoryMetaName, &CategoryMetadata);
 	InStruct->GetStringMetaDataHierarchical(FRigVMStruct::KeywordsMetaName, &KeywordsMetadata);
 	InStruct->GetStringMetaDataHierarchical(FRigVMStruct::DisplayNameMetaName, &DisplayNameMetadata);
+	InStruct->GetStringMetaDataHierarchical(FRigVMStruct::DeprecatedMetaName, &DeprecatedMetadata);
 #endif
 	if(DisplayNameMetadata.IsEmpty())
 	{
@@ -424,13 +526,22 @@ void RigVMJsonUtils::GenericStructAttributesToJson(FRigVMRegistryHandle& InRegis
 
 	OutJsonObject->SetStringField(TEXT("DisplayName"), DisplayNameMetadata);
 	OutJsonObject->SetStringField(TEXT("Type"), InStruct->GetStructCPPName());
+	if (!DeprecatedMetadata.IsEmpty())
+	{
+		if (InSettings.bExcludeDeprecated && !DeprecatedMetadata.IsEmpty())
+		{
+			return false;
+		}
+		OutJsonObject->SetStringField(TEXT("Deprecated"), DeprecatedMetadata);
+	}
 	OutJsonObject->SetStringField(TEXT("Category"), CategoryMetadata);
 	OutJsonObject->SetStringField(TEXT("Keywords"), KeywordsMetadata);
-	OutJsonObject->SetStringField(TEXT("ToolTip"), InStruct->GetToolTipText().ToString());
+	OutJsonObject->SetStringField(TEXT("ToolTip"), PostProcessToolTip(InStruct->GetToolTipText().ToString()));
 	OutJsonObject->SetStringField(TEXT("PackageName"), InStruct->GetOutermost()->GetPathName());
+	return true;
 }
 
-TSharedPtr<FJsonObject> RigVMJsonUtils::ToJson(FRigVMRegistryHandle& InRegistry, const IRigVMAssetInterface* InAsset)
+TSharedPtr<FJsonObject> RigVMJsonUtils::ToJson(FIntrospectionSettings& InSettings, const IRigVMAssetInterface* InAsset)
 {
 	check(InAsset);
 	
@@ -439,7 +550,7 @@ TSharedPtr<FJsonObject> RigVMJsonUtils::ToJson(FRigVMRegistryHandle& InRegistry,
 	TArray<FRigVMGraphVariableDescription> VariableDescriptions = InAsset->GetAssetVariables();
 	for (const FRigVMGraphVariableDescription& VariableDescription : VariableDescriptions)
 	{
-		if (TSharedPtr<FJsonObject> VariableJsonObject = ToJson(InRegistry, VariableDescription))
+		if (TSharedPtr<FJsonObject> VariableJsonObject = ToJson(InSettings, VariableDescription))
 		{
 			Variables.Add(MakeShared<FJsonValueObject>(VariableJsonObject));
 		}
@@ -450,7 +561,7 @@ TSharedPtr<FJsonObject> RigVMJsonUtils::ToJson(FRigVMRegistryHandle& InRegistry,
 		const TArray<URigVMGraph*> Graphs =Client->GetAllModels(true, false);
 		for (const URigVMGraph* Graph : Graphs)
 		{
-			if (TSharedPtr<FJsonObject> GraphJsonObject = ToJson(InRegistry, Graph))
+			if (TSharedPtr<FJsonObject> GraphJsonObject = ToJson(InSettings, Graph))
 			{
 				GraphValues.Add(MakeShared<FJsonValueObject>(GraphJsonObject));
 			}
@@ -469,7 +580,7 @@ TSharedPtr<FJsonObject> RigVMJsonUtils::ToJson(FRigVMRegistryHandle& InRegistry,
 	return JsonObject;
 }
 
-TSharedPtr<FJsonObject> RigVMJsonUtils::ToJson(FRigVMRegistryHandle& InRegistry, const URigVMGraph* InGraph, bool bOnlySelection)
+TSharedPtr<FJsonObject> RigVMJsonUtils::ToJson(FIntrospectionSettings& InSettings, const URigVMGraph* InGraph, bool bOnlySelection)
 {
 	check(InGraph);
 
@@ -481,7 +592,7 @@ TSharedPtr<FJsonObject> RigVMJsonUtils::ToJson(FRigVMRegistryHandle& InRegistry,
 		TArray<FRigVMGraphVariableDescription> VariableDescriptions = InGraph->GetLocalVariables(false);
 		for (const FRigVMGraphVariableDescription& VariableDescription : VariableDescriptions)
 		{
-			if (TSharedPtr<FJsonObject> VariableJsonObject = ToJson(InRegistry, VariableDescription))
+			if (TSharedPtr<FJsonObject> VariableJsonObject = ToJson(InSettings, VariableDescription))
 			{
 				Variables.Add(MakeShared<FJsonValueObject>(VariableJsonObject));
 			}
@@ -499,7 +610,7 @@ TSharedPtr<FJsonObject> RigVMJsonUtils::ToJson(FRigVMRegistryHandle& InRegistry,
 			}
 		}
 		
-		if (TSharedPtr<FJsonObject> NodeJsonObject = ToJson(InRegistry, Node))
+		if (TSharedPtr<FJsonObject> NodeJsonObject = ToJson(InSettings, Node))
 		{
 			Nodes.Add(MakeShared<FJsonValueObject>(NodeJsonObject));
 			ExportedNodes.Add(Node->GetFName());
@@ -518,7 +629,7 @@ TSharedPtr<FJsonObject> RigVMJsonUtils::ToJson(FRigVMRegistryHandle& InRegistry,
 		{
 			continue;
 		}
-		if (TSharedPtr<FJsonObject> LinkJsonObject = ToJson(InRegistry, Link))
+		if (TSharedPtr<FJsonObject> LinkJsonObject = ToJson(InSettings, Link))
 		{
 			Links.Add(MakeShared<FJsonValueObject>(LinkJsonObject));
 		}
@@ -541,7 +652,7 @@ TSharedPtr<FJsonObject> RigVMJsonUtils::ToJson(FRigVMRegistryHandle& InRegistry,
 	return JsonObject;
 }
 
-TSharedPtr<FJsonObject> RigVMJsonUtils::ToJson(FRigVMRegistryHandle& InRegistry, const URigVMNode* InNode)
+TSharedPtr<FJsonObject> RigVMJsonUtils::ToJson(FIntrospectionSettings& InSettings, const URigVMNode* InNode)
 {
 	check(InNode);
 
@@ -556,7 +667,7 @@ TSharedPtr<FJsonObject> RigVMJsonUtils::ToJson(FRigVMRegistryHandle& InRegistry,
 			continue;
 		}
 
-		if (TSharedPtr<FJsonObject> PinJsonObject = ToJson(InRegistry, Pin))
+		if (TSharedPtr<FJsonObject> PinJsonObject = ToJson(InSettings, Pin))
 		{
 			Pins.Add(MakeShared<FJsonValueObject>(PinJsonObject));
 		}
@@ -607,7 +718,7 @@ TSharedPtr<FJsonObject> RigVMJsonUtils::ToJson(FRigVMRegistryHandle& InRegistry,
 	}
 	else if (const URigVMCollapseNode* CollapseNode = Cast<URigVMCollapseNode>(InNode))
 	{
-		if (TSharedPtr<FJsonObject> ContainedGraphJsonObject = ToJson(InRegistry, CollapseNode->GetContainedGraph()))
+		if (TSharedPtr<FJsonObject> ContainedGraphJsonObject = ToJson(InSettings, CollapseNode->GetContainedGraph()))
 		{
 			if (CollapseNode->GetGraph()->IsA<URigVMFunctionLibrary>())
 			{
@@ -631,7 +742,7 @@ TSharedPtr<FJsonObject> RigVMJsonUtils::ToJson(FRigVMRegistryHandle& InRegistry,
 	return JsonObject;
 }
 
-TSharedPtr<FJsonObject> RigVMJsonUtils::ToJson(FRigVMRegistryHandle& InRegistry, const URigVMPin* InPin)
+TSharedPtr<FJsonObject> RigVMJsonUtils::ToJson(FIntrospectionSettings& InSettings, const URigVMPin* InPin)
 {
 	check(InPin);
 	
@@ -650,7 +761,7 @@ TSharedPtr<FJsonObject> RigVMJsonUtils::ToJson(FRigVMRegistryHandle& InRegistry,
 	return JsonObject;
 }
 
-TSharedPtr<FJsonObject> RigVMJsonUtils::ToJson(FRigVMRegistryHandle& InRegistry, const URigVMLink* InLink)
+TSharedPtr<FJsonObject> RigVMJsonUtils::ToJson(FIntrospectionSettings& InSettings, const URigVMLink* InLink)
 {
 	check(InLink);
 
@@ -667,7 +778,7 @@ TSharedPtr<FJsonObject> RigVMJsonUtils::ToJson(FRigVMRegistryHandle& InRegistry,
 	return JsonObject;
 }
 
-TSharedPtr<FJsonObject> RigVMJsonUtils::ToJson(FRigVMRegistryHandle& InRegistry, const FRigVMGraphVariableDescription& InVariable)
+TSharedPtr<FJsonObject> RigVMJsonUtils::ToJson(FIntrospectionSettings& InSettings, const FRigVMGraphVariableDescription& InVariable)
 {
 	if (InVariable.Name.IsNone())
 	{
@@ -683,7 +794,7 @@ TSharedPtr<FJsonObject> RigVMJsonUtils::ToJson(FRigVMRegistryHandle& InRegistry,
 	JsonObject->SetStringField(TEXT("TypePath"), InVariable.CPPTypeObjectPath.IsNone() ? FString() : InVariable.CPPTypeObjectPath.ToString());
 	JsonObject->SetStringField(TEXT("DefaultValue"), InVariable.DefaultValue);
 	JsonObject->SetStringField(TEXT("Category"), InVariable.Category.ToString());
-	JsonObject->SetStringField(TEXT("Tooltip"), InVariable.Tooltip.ToString());
+	JsonObject->SetStringField(TEXT("ToolTip"), PostProcessToolTip(InVariable.Tooltip.ToString()));
 	JsonObject->SetStringField(TEXT("ExposedOnSpawn"), InVariable.bExposedOnSpawn ? True : False);
 	JsonObject->SetStringField(TEXT("ExposeToCinematics"), InVariable.bExposeToCinematics ? True : False);
 	JsonObject->SetStringField(TEXT("Public"), InVariable.bPublic ? True : False);
diff --git a/Engine/Plugins/Runtime/RigVM/Source/RigVMDeveloper/Public/RigVMJsonUtils.h b/Engine/Plugins/Runtime/RigVM/Source/RigVMDeveloper/Public/RigVMJsonUtils.h
--- a/Engine/Plugins/Runtime/RigVM/Source/RigVMDeveloper/Public/RigVMJsonUtils.h
+++ b/Engine/Plugins/Runtime/RigVM/Source/RigVMDeveloper/Public/RigVMJsonUtils.h
@@ -31,19 +31,32 @@ namespace RigVMJsonUtils
 
 	RIGVMDEVELOPER_API FJsonWriterRef MakeJsonWriter(FString& OutJsonText);
 
+	struct FIntrospectionSettings
+	{
+		FIntrospectionSettings(FRigVMRegistryHandle& InRegistry)
+			: Registry(InRegistry)
+		{
+		}
+
+		FRigVMRegistryHandle& Registry;
+		bool bExcludeFunctions = false;
+		bool bExcludeDispatches = false;
+		bool bExcludeDeprecated = false;
+	};
+
 	// introspection for the registry
-	RIGVMDEVELOPER_API TSharedPtr<FJsonObject> AvailableNodesToJson(FRigVMRegistryHandle& InRegistry, TFunction<bool(const UScriptStruct*)> InFilterCallback);
-	RIGVMDEVELOPER_API TSharedPtr<FJsonObject> ToJson(FRigVMRegistryHandle& InRegistry, const FRigVMFunction* InFunction);
-	RIGVMDEVELOPER_API TSharedPtr<FJsonObject> ToJson(FRigVMRegistryHandle& InRegistry, const FRigVMDispatchFactory* InFactory);
-	RIGVMDEVELOPER_API void GenericStructAttributesToJson(FRigVMRegistryHandle& InRegistry, const UScriptStruct* InStruct, FJsonObject* OutJsonObject);
+	RIGVMDEVELOPER_API TSharedPtr<FJsonObject> AvailableNodesToJson(FIntrospectionSettings& InSettings, TFunction<bool(const UScriptStruct*)> InFilterCallback);
+	RIGVMDEVELOPER_API TSharedPtr<FJsonObject> ToJson(FIntrospectionSettings& InSettings, const FRigVMFunction* InFunction);
+	RIGVMDEVELOPER_API TSharedPtr<FJsonObject> ToJson(FIntrospectionSettings& InSettings, const FRigVMDispatchFactory* InFactory);
+	RIGVMDEVELOPER_API bool GenericStructAttributesToJson(FIntrospectionSettings& InSettings, const UScriptStruct* InStruct, FJsonObject* OutJsonObject);
 
 	// introspection for existing graphs
-	RIGVMDEVELOPER_API TSharedPtr<FJsonObject> ToJson(FRigVMRegistryHandle& InRegistry, const IRigVMAssetInterface* InAsset);
-	RIGVMDEVELOPER_API TSharedPtr<FJsonObject> ToJson(FRigVMRegistryHandle& InRegistry, const URigVMGraph* InGraph, bool bOnlySelection = false);
-	RIGVMDEVELOPER_API TSharedPtr<FJsonObject> ToJson(FRigVMRegistryHandle& InRegistry, const URigVMNode* InNode);
-	RIGVMDEVELOPER_API TSharedPtr<FJsonObject> ToJson(FRigVMRegistryHandle& InRegistry, const URigVMPin* InPin);
-	RIGVMDEVELOPER_API TSharedPtr<FJsonObject> ToJson(FRigVMRegistryHandle& InRegistry, const URigVMLink* InLink);
-	RIGVMDEVELOPER_API TSharedPtr<FJsonObject> ToJson(FRigVMRegistryHandle& InRegistry, const FRigVMGraphVariableDescription& InVariable);
+	RIGVMDEVELOPER_API TSharedPtr<FJsonObject> ToJson(FIntrospectionSettings& InSettings, const IRigVMAssetInterface* InAsset);
+	RIGVMDEVELOPER_API TSharedPtr<FJsonObject> ToJson(FIntrospectionSettings& InSettings, const URigVMGraph* InGraph, bool bOnlySelection = false);
+	RIGVMDEVELOPER_API TSharedPtr<FJsonObject> ToJson(FIntrospectionSettings& InSettings, const URigVMNode* InNode);
+	RIGVMDEVELOPER_API TSharedPtr<FJsonObject> ToJson(FIntrospectionSettings& InSettings, const URigVMPin* InPin);
+	RIGVMDEVELOPER_API TSharedPtr<FJsonObject> ToJson(FIntrospectionSettings& InSettings, const URigVMLink* InLink);
+	RIGVMDEVELOPER_API TSharedPtr<FJsonObject> ToJson(FIntrospectionSettings& InSettings, const FRigVMGraphVariableDescription& InVariable);
 
 	// string conversion
 	RIGVMDEVELOPER_API FString StructValueToJson(const UScriptStruct* InScriptStruct, const void* InMemory);
