RigVM: Fix duplicate LOCTEXT entries

#rnx

diff --git a/Engine/Plugins/Runtime/RigVM/Source/RigVM/Private/RigVMFunctions/RigVMDispatch_CastEnum.cpp b/Engine/Plugins/Runtime/RigVM/Source/RigVM/Private/RigVMFunctions/RigVMDispatch_CastEnum.cpp
--- a/Engine/Plugins/Runtime/RigVM/Source/RigVM/Private/RigVMFunctions/RigVMDispatch_CastEnum.cpp
+++ b/Engine/Plugins/Runtime/RigVM/Source/RigVM/Private/RigVMFunctions/RigVMDispatch_CastEnum.cpp
@@ -46,18 +46,18 @@ FString FRigVMDispatch_CastEnumToInt::GetNodeTitle(const FRigVMTemplateTypeMap&
 
 FText FRigVMDispatch_CastEnumToInt::GetNodeTooltip(const FRigVMTemplateTypeMap& InTypes) const
 {
-	return LOCTEXT("CastEnumToolTip", "Casts from enum to int");
+	return LOCTEXT("CastEnumToIntToolTip", "Casts from enum to int");
 }
 
 FText FRigVMDispatch_CastEnumToInt::GetArgumentTooltip(const FName& InArgumentName, TRigVMTypeIndex InTypeIndex) const
 {
 	if (InArgumentName == ValueName)
 	{
-		return LOCTEXT("ValueToolTip", "The enum value to cast to an integer");
+		return LOCTEXT("CastEnumToIntValueToolTip", "The enum value to cast to an integer");
 	}
 	if (InArgumentName == ResultName)
 	{
-		return LOCTEXT("ResultToolTip", "The resulting integer value");
+		return LOCTEXT("CastEnumToIntResultToolTip", "The resulting integer value");
 	}
 	return FRigVMDispatchFactory::GetArgumentTooltip(InArgumentName, InTypeIndex);
 }
@@ -160,11 +160,11 @@ FText FRigVMDispatch_CastIntToEnum::GetArgumentTooltip(const FName& InArgumentNa
 {
 	if (InArgumentName == ValueName)
 	{
-		return LOCTEXT("ValueToolTip", "The integer value to cast to an enum value");
+		return LOCTEXT("CastIntToEnumValueToolTip", "The integer value to cast to an enum value");
 	}
 	if (InArgumentName == ResultName)
 	{
-		return LOCTEXT("ResultToolTip", "The resulting enum value");
+		return LOCTEXT("CastIntToEnumResultToolTip", "The resulting enum value");
 	}
 	return FRigVMDispatchFactory::GetArgumentTooltip(InArgumentName, InTypeIndex);
 }
diff --git a/Engine/Plugins/Runtime/RigVM/Source/RigVM/Private/RigVMFunctions/RigVMDispatch_Core.cpp b/Engine/Plugins/Runtime/RigVM/Source/RigVM/Private/RigVMFunctions/RigVMDispatch_Core.cpp
--- a/Engine/Plugins/Runtime/RigVM/Source/RigVM/Private/RigVMFunctions/RigVMDispatch_Core.cpp
+++ b/Engine/Plugins/Runtime/RigVM/Source/RigVM/Private/RigVMFunctions/RigVMDispatch_Core.cpp
@@ -54,11 +54,11 @@ FText FRigVMDispatch_CoreEquals::GetArgumentTooltip(const FName& InArgumentName,
 	}
 	if (InArgumentName == BName)
 	{
-		return LOCTEXT("CoreEqualsAToolTip", "The second value to compare");
+		return LOCTEXT("CoreEqualsBToolTip", "The second value to compare");
 	}
 	if (InArgumentName == ResultName)
 	{
-		return LOCTEXT("CoreEqualsAToolTip", "The boolean result of the comparison");
+		return LOCTEXT("CoreEqualsResultToolTip", "The boolean result of the comparison");
 	}
 	return FRigVMDispatch_CoreBase::GetArgumentTooltip(InArgumentName, InTypeIndex);
 }
diff --git a/Engine/Plugins/Runtime/RigVM/Source/RigVM/Private/RigVMFunctions/RigVMDispatch_MakeStruct.cpp b/Engine/Plugins/Runtime/RigVM/Source/RigVM/Private/RigVMFunctions/RigVMDispatch_MakeStruct.cpp
--- a/Engine/Plugins/Runtime/RigVM/Source/RigVM/Private/RigVMFunctions/RigVMDispatch_MakeStruct.cpp
+++ b/Engine/Plugins/Runtime/RigVM/Source/RigVM/Private/RigVMFunctions/RigVMDispatch_MakeStruct.cpp
@@ -93,11 +93,11 @@ FText FRigVMDispatch_MakeStruct::GetArgumentTooltip(const FName& InArgumentName,
 {
 	if (InArgumentName == ElementsName)
 	{
-		return LOCTEXT("ElementsToolTip", "The elements of the structure to make");
+		return LOCTEXT("MakeStructElementsToolTip", "The elements of the structure to make");
 	}
 	if (InArgumentName == StructName)
 	{
-		return LOCTEXT("StructToolTip", "The created structure");
+		return LOCTEXT("MakeStructStructToolTip", "The created structure");
 	}
 	return FRigVMDispatch_CoreBase::GetArgumentTooltip(InArgumentName, InTypeIndex);
 }
@@ -156,11 +156,11 @@ FText FRigVMDispatch_BreakStruct::GetArgumentTooltip(const FName& InArgumentName
 {
 	if (InArgumentName == ElementsName)
 	{
-		return LOCTEXT("ElementsToolTip", "The elements of the structure");
+		return LOCTEXT("BreakStructElementsToolTip", "The elements of the structure");
 	}
 	if (InArgumentName == StructName)
 	{
-		return LOCTEXT("StructToolTip", "The structure to break up into its elements");
+		return LOCTEXT("BreakStructStructToolTip", "The structure to break up into its elements");
 	}
 	return FRigVMDispatch_MakeStruct::GetArgumentTooltip(InArgumentName, InTypeIndex);
 }
diff --git a/Engine/Plugins/Runtime/RigVM/Source/RigVMDeveloper/Private/EdGraph/RigVMEdGraphSchema.cpp b/Engine/Plugins/Runtime/RigVM/Source/RigVMDeveloper/Private/EdGraph/RigVMEdGraphSchema.cpp
--- a/Engine/Plugins/Runtime/RigVM/Source/RigVMDeveloper/Private/EdGraph/RigVMEdGraphSchema.cpp
+++ b/Engine/Plugins/Runtime/RigVM/Source/RigVMDeveloper/Private/EdGraph/RigVMEdGraphSchema.cpp
@@ -444,7 +444,7 @@ UEdGraphNode* FRigVMEdGraphSchemaAction_PromoteToVariable::PerformAction(UEdGrap
 FRigVMEdGraphSchemaAction_PromoteToExposedPin::FRigVMEdGraphSchemaAction_PromoteToExposedPin(UEdGraphPin* InEdGraphPin)
 : FEdGraphSchemaAction(	FText(), 
 						LOCTEXT("PromoteToExposedPin", "Promote to exposed pin"),
-						LOCTEXT("PromoteToExposedPin", "Adds an exposed pin and wires it to this pin. For inputs if you hold shift this will insert a Get node."),
+						LOCTEXT("PromoteToExposedPinTooltip", "Adds an exposed pin and wires it to this pin. For inputs if you hold shift this will insert a Get node."),
 						1)
 , EdGraphPin(InEdGraphPin)
 {
