RigVM: Fix initializing sub pins for external variables for set variable nodes

#rnx

diff --git a/Engine/Plugins/Runtime/RigVM/Source/RigVMDeveloper/Private/RigVMCompiler/RigVMAST.cpp b/Engine/Plugins/Runtime/RigVM/Source/RigVMDeveloper/Private/RigVMCompiler/RigVMAST.cpp
--- a/Engine/Plugins/Runtime/RigVM/Source/RigVMDeveloper/Private/RigVMCompiler/RigVMAST.cpp
+++ b/Engine/Plugins/Runtime/RigVM/Source/RigVMDeveloper/Private/RigVMCompiler/RigVMAST.cpp
@@ -1363,6 +1363,72 @@ FRigVMExprAST* FRigVMParserAST::TraverseNode(const FRigVMASTProxy& InNodeProxy,
 	if (NodeExpr)
 	{
 		TraversePins(InNodeProxy, NodeExpr);
+
+		if (const URigVMVariableNode* VariableNode = Cast<URigVMVariableNode>(Node))
+		{
+			bool bRequiresCloningVariableValue = false;
+			if (VariableNode->IsGetter())
+			{
+				URigVMPin* VariableValuePin = VariableNode->GetValuePin();
+				FRigVMASTProxy VariableValuePinProxy = NodeExpr->GetProxy().GetSibling(VariableValuePin);
+
+				// for variable nodes we need to deal with a specific case where a sub pin
+				// of the value pin is wired to a variable setter node with the same variable
+				// on a different sub pin. so for example Get MyVector.X -> Set MyVector.Y.
+				// In that case we need to inject a clone of the variable in the getter.
+				// Since arrays don't have sub pins on the getter this can only happen to non-arrays.
+				const TArray<int32>& LinkIndices = GetTargetLinkIndices(VariableValuePinProxy, true);
+				for(const int32 LinkIndex : LinkIndices)
+				{
+					const FRigVMASTLinkDescription& Link = Links[LinkIndex];
+					const URigVMPin* SourcePin = Link.SourceProxy.GetSubjectChecked<URigVMPin>();
+					const URigVMPin* TargetPin = Link.TargetProxy.GetSubjectChecked<URigVMPin>();
+					if (TargetPin && !Link.SegmentPath.IsEmpty())
+					{
+						TargetPin = TargetPin->FindSubPin(Link.SegmentPath);
+					}
+					if (!SourcePin || !TargetPin)
+					{
+						continue;
+					}
+					if (TargetPin->IsRootPin())
+					{
+						continue;
+					}
+					const URigVMVariableNode* TargetVariableNode = Cast<URigVMVariableNode>(TargetPin->GetNode());
+					if (!TargetVariableNode)
+					{
+						continue;
+					}
+					if (VariableNode->GetVariableName() != TargetVariableNode->GetVariableName())
+					{
+						continue;
+					}
+
+					if (SourcePin->GetSegmentPath() != TargetPin->GetSegmentPath())
+					{
+						bRequiresCloningVariableValue = true;
+						break;
+					}
+				}
+			}
+
+			if (bRequiresCloningVariableValue && NodeExpr->IsNode())
+			{
+				FRigVMNodeExprAST* VariableNodeExpr = NodeExpr->To<FRigVMNodeExprAST>();
+				const FRigVMVarExprAST* ConstVariableVarExpr = VariableNodeExpr->FindVarWithPinName(URigVMVariableNode::ValueName);
+				FRigVMVarExprAST* VariableVarExpr = const_cast<FRigVMVarExprAST*>(ConstVariableVarExpr);
+				
+				FRigVMASTProxy PinProxy = VariableVarExpr->GetProxy();
+				FRigVMExprAST* ClonedVariableVarExpr = MakeExpr<FRigVMVarExprAST>(FRigVMExprAST::EType::Var, PinProxy);
+				ClonedVariableVarExpr->Name = *FString::Printf(TEXT("%s::Clone"), *VariableVarExpr->Name.ToString());
+				const FRigVMASTLinkDescription CloneLink(PinProxy, PinProxy, FString());
+				FRigVMExprAST* PinCopyExpr = MakeExpr<FRigVMCopyExprAST>(CloneLink);
+				PinCopyExpr->AddParent(ClonedVariableVarExpr);
+				VariableVarExpr->AddParent(PinCopyExpr);
+				VariableNodeExpr->ReplaceChild(VariableVarExpr, ClonedVariableVarExpr);
+			}
+		}
 	}
 
 	return NodeExpr;
@@ -1662,7 +1728,7 @@ FRigVMExprAST* FRigVMParserAST::TraversePin(const FRigVMASTProxy& InPinProxy, FR
 		if (!bHasSourceLinkToRoot && 
 			GetSourceLinkIndices(InPinProxy, false).Num() == 0 &&
 			(Pin->GetDirection() == ERigVMPinDirection::IO ||
-			(LinkIndices.Num() > 0 && (Pin->IsArray() || Pin->GetNode()->IsA<URigVMVariableNode>()))))
+			(LinkIndices.Num() > 0 && Pin->IsArray())))
 		{
 			FRigVMLiteralExprAST* LiteralExpr = MakeExpr<FRigVMLiteralExprAST>(InPinProxy);
 			const FRigVMASTLinkDescription LiteralLink(InPinProxy, InPinProxy, FString());
diff --git a/Engine/Plugins/Runtime/RigVM/Source/RigVMDeveloper/Private/RigVMCompiler/RigVMCompiler.cpp b/Engine/Plugins/Runtime/RigVM/Source/RigVMDeveloper/Private/RigVMCompiler/RigVMCompiler.cpp
--- a/Engine/Plugins/Runtime/RigVM/Source/RigVMDeveloper/Private/RigVMCompiler/RigVMCompiler.cpp
+++ b/Engine/Plugins/Runtime/RigVM/Source/RigVMDeveloper/Private/RigVMCompiler/RigVMCompiler.cpp
@@ -3256,6 +3256,129 @@ bool URigVMCompiler::TraverseInlineFunction(const FRigVMInlineFunctionExprAST* I
 
 bool URigVMCompiler::TraverseNoOp(const FRigVMNoOpExprAST* InExpr, FRigVMCompilerWorkData& WorkData)
 {
+	if (const URigVMVariableNode* VariableNode = Cast<URigVMVariableNode>(InExpr->GetNode()))
+	{
+		if (!VariableNode->IsGetter())
+		{
+			const URigVMPin* ValuePin = VariableNode->GetValuePin();
+
+			// if the value pin has sub pin links
+			if (ValuePin->GetSourceLinks(false).IsEmpty() &&
+				!ValuePin->GetSourceLinks(true).IsEmpty())
+			{
+				// find all the paths that require a copy from the default
+				TArray<FString> PropertyPathsRequiringCopy;
+				TArray<URigVMPin*> SubPins = ValuePin->GetSubPins();
+				for (int32 SubPinIndex = 0; SubPinIndex < SubPins.Num(); ++SubPinIndex)
+				{
+					const URigVMPin* SubPin = SubPins[SubPinIndex];
+					if (!SubPin->GetSourceLinks(false).IsEmpty())
+					{
+						continue;
+					}
+					if (SubPin->GetSourceLinks(true).IsEmpty())
+					{
+						PropertyPathsRequiringCopy.Add(SubPin->GetSubPinPath(ValuePin, false));
+					}
+					else
+					{
+						SubPins.Append(SubPin->GetSubPins());
+					}
+				}
+
+				if (!PropertyPathsRequiringCopy.IsEmpty())
+				{
+					// run the recursion now so that all the registers exist.
+					if (!TraverseChildren(InExpr, WorkData))
+					{
+						return false;
+					}
+					
+					FName VariableName = VariableNode->GetVariableName();
+					FRigVMASTProxy ParentProxy = InExpr->GetProxy();
+					while(ParentProxy.GetCallstack().Num() > 1)
+					{
+						ParentProxy = ParentProxy.GetParent();
+
+						if(URigVMFunctionReferenceNode* FunctionReferenceNode = ParentProxy.GetSubject<URigVMFunctionReferenceNode>())
+						{
+							const FName RemappedVariableName = FunctionReferenceNode->GetOuterVariableName(VariableName);
+							if(!RemappedVariableName.IsNone())
+							{
+								VariableName = RemappedVariableName;
+							}
+						}
+					}
+
+					// Default value literals will be shared with variable sharing the same default
+					const FString DefaultValue = ValuePin->GetDefaultValue(); 
+					const uint32 DefaultValueHash = GetTypeHash(DefaultValue);
+					FString RegisterNameString;
+					if (VariableNode->IsLocalVariable())
+					{
+						RegisterNameString = FString::Printf(TEXT("LocalVariableDefault::%s|%s|%lu::Const"), *VariableNode->GetGraph()->GetGraphName(), *VariableName.ToString(), (unsigned long)DefaultValueHash);
+					}
+					else
+					{
+						RegisterNameString = FString::Printf(TEXT("VariableDefault::%s|%lu::Const"), *VariableName.ToString(), (unsigned long)DefaultValueHash);
+					}
+					const FName RegisterName = FRigVMPropertyDescription::SanitizeName(*RegisterNameString);
+
+					// set up a default value literal based on the value pin configuration
+					const FRigVMOperand RootDefaultValueOperand = WorkData.FindProperty(ERigVMMemoryType::Literal, RegisterName);
+
+					FRigVMOperand VariableOperand;
+					FString VariableRegisterNameString;
+					if (VariableNode->IsLocalVariable())
+					{
+						VariableRegisterNameString = FString::Printf(TEXT("LocalVariable::%s|%s"), *VariableNode->GetGraph()->GetGraphName(), *VariableName.ToString());
+						const FName VariableRegisterName = FRigVMPropertyDescription::SanitizeName(*VariableRegisterNameString);
+						VariableOperand = WorkData.FindProperty(ERigVMMemoryType::Work, VariableRegisterName);
+					}
+					else
+					{
+						VariableRegisterNameString = FString::Printf(TEXT("Variable::%s"), *VariableName.ToString());					
+						if (const FRigVMOperand* VariableOperandPtr = WorkData.PinPathToOperand->Find(VariableRegisterNameString))
+						{
+							VariableOperand = *VariableOperandPtr;
+						}
+					}
+
+					if (WorkData.bSetupMemory)
+					{
+						if (!RootDefaultValueOperand.IsValid() && VariableOperand.IsValid())
+						{
+							const FRigVMOperand DefaultOperand = WorkData.AddProperty(ERigVMMemoryType::Literal, RegisterName, ValuePin->GetCPPType(), ValuePin->GetCPPTypeObject(), DefaultValue);
+							for (const FString& PropertyPath : PropertyPathsRequiringCopy)
+							{
+								// add the property path both to the external variable as well as the default operand
+								WorkData.FindOrAddPropertyPath(DefaultOperand, ValuePin->GetCPPType(), PropertyPath);
+								WorkData.FindOrAddPropertyPath(VariableOperand, ValuePin->GetCPPType(), PropertyPath);
+							}
+						}
+					}
+					else if (ensure(RootDefaultValueOperand.IsValid()) && VariableOperand.IsValid())
+					{
+						for (const FString& PropertyPath : PropertyPathsRequiringCopy)
+						{
+							// set up an operand both for the default literal and the variable
+							// given the same property path.
+							const int32 DefaultRegisterOffset = WorkData.FindOrAddPropertyPath(RootDefaultValueOperand, ValuePin->GetCPPType(), PropertyPath);
+							const int32 VariableRegisterOffset = WorkData.FindOrAddPropertyPath(VariableOperand, ValuePin->GetCPPType(), PropertyPath);
+							const FRigVMOperand DefaultOperand(RootDefaultValueOperand.GetMemoryType(), RootDefaultValueOperand.GetRegisterIndex(), DefaultRegisterOffset);
+							const FRigVMOperand SubPinVariableOperand(VariableOperand.GetMemoryType(), VariableOperand.GetRegisterIndex(), VariableRegisterOffset);
+
+							const FRigVMCopyOp CopyOp = WorkData.VM->GetCopyOpForOperands(DefaultOperand, SubPinVariableOperand);
+							WorkData.VM->GetByteCode().AddCopyOp(CopyOp);
+						}
+					}
+
+					// return here since we've already run traverse children
+					return true;
+				}
+			}
+		}
+	}
 	return TraverseChildren(InExpr, WorkData);
 }
 
@@ -3707,7 +3830,9 @@ FString URigVMCompiler::GetPinHashImpl(const URigVMPin* InPin, const FRigVMVarEx
 	}
 
 	bool bUseFullNodePath = true;
-	if (URigVMVariableNode* VariableNode = Cast<URigVMVariableNode>(Node))
+	URigVMVariableNode* VariableNode = Cast<URigVMVariableNode>(Node);
+	const bool bIsVariableClone = InVarExpr && InVarExpr->Name.ToString().EndsWith(TEXT("::Clone"));
+	if (VariableNode && !bIsVariableClone)
 	{
 		if (InPin->GetName() == TEXT("Value") && !bIsDebugValue)
 		{
@@ -3836,6 +3961,12 @@ FString URigVMCompiler::GetPinHashImpl(const URigVMPin* InPin, const FRigVMVarEx
 				return FString::Printf(TEXT("%s%s%s"), *Prefix, *FullPath, *Suffix);
 			}
 		}
+
+		if (bIsVariableClone)
+		{
+			Suffix = TEXT("::Clone");
+			bUseFullNodePath = true;
+		}
 	}
 
 	if (InPinProxy.IsValid())
