RigVM Json Utils: Don't export default values for outputs

#rnx

diff --git a/Engine/Plugins/Runtime/RigVM/Source/RigVMDeveloper/Private/RigVMJsonUtils.cpp b/Engine/Plugins/Runtime/RigVM/Source/RigVMDeveloper/Private/RigVMJsonUtils.cpp
--- a/Engine/Plugins/Runtime/RigVM/Source/RigVMDeveloper/Private/RigVMJsonUtils.cpp
+++ b/Engine/Plugins/Runtime/RigVM/Source/RigVMDeveloper/Private/RigVMJsonUtils.cpp
@@ -299,12 +299,14 @@ TSharedPtr<FJsonObject> RigVMJsonUtils::ToJson(FIntrospectionSettings& InSetting
 
 		bool bIsInput = Argument ? Argument->Direction == ERigVMFunctionArgumentDirection::Input : true;
 		bool bIsOutput = Argument ? Argument->Direction == ERigVMFunctionArgumentDirection::Output : true;
+		bool bIsConstant = false;
 		
 #if WITH_EDITOR
 		bIsInput =
 			Property->HasMetaData(FRigVMStruct::InputMetaName) ||
 			Property->HasMetaData(FRigVMStruct::VisibleMetaName);
 		bIsOutput = Property->HasMetaData(FRigVMStruct::OutputMetaName);
+		bIsConstant = Property->HasMetaData(FRigVMStruct::ConstantMetaName);
 #endif
 
 		if (!bIsInput && !bIsOutput)
@@ -335,7 +337,11 @@ TSharedPtr<FJsonObject> RigVMJsonUtils::ToJson(FIntrospectionSettings& InSetting
 			PropertyToolTip = TEXT("This pin is used to chain multiple mutable units together");
 		}
 		ArgumentJsonObject->SetStringField(TEXT("PropertyToolTip"), PropertyToolTip);
-		ArgumentJsonObject->SetStringField(TEXT("DefaultValue"), DefaultValue);
+
+		if ((bIsInput || bIsConstant) && !bIsExecutePin)
+		{
+			ArgumentJsonObject->SetStringField(TEXT("DefaultValue"), DefaultValue);
+		}
 
 		if (bIsInput && bIsOutput)
 		{
@@ -749,7 +755,10 @@ TSharedPtr<FJsonObject> RigVMJsonUtils::ToJson(FIntrospectionSettings& InSetting
 	TSharedPtr<FJsonObject> JsonObject = MakeShared<FJsonObject>();
 	JsonObject->SetStringField(TEXT("Name"), InPin->GetName());
 	JsonObject->SetStringField(TEXT("Type"), InPin->GetCPPType());
-	JsonObject->SetStringField(TEXT("DefaultValue"), InPin->GetDefaultValue());
+	if (InPin->GetDirection() != ERigVMPinDirection::Output)
+	{
+		JsonObject->SetStringField(TEXT("DefaultValue"), InPin->GetDefaultValue());
+	}
 
 	const TArray<URigVMInjectionInfo*> InjectedNodes = InPin->GetInjectedNodes();
 	check(InjectedNodes.Num() < 2);
