RigVM: Treat constant nodes as NoOp and constant output as literals

#rnx

diff --git a/Engine/Plugins/Runtime/RigVM/Source/RigVM/Private/RigVMFunctions/Math/RigVMFunction_MathBool.cpp b/Engine/Plugins/Runtime/RigVM/Source/RigVM/Private/RigVMFunctions/Math/RigVMFunction_MathBool.cpp
--- a/Engine/Plugins/Runtime/RigVM/Source/RigVM/Private/RigVMFunctions/Math/RigVMFunction_MathBool.cpp
+++ b/Engine/Plugins/Runtime/RigVM/Source/RigVM/Private/RigVMFunctions/Math/RigVMFunction_MathBool.cpp
@@ -17,12 +17,12 @@ FRigVMStructUpgradeInfo FRigVMFunction_MathBoolMake::GetUpgradeInfo() const
 
 FRigVMFunction_MathBoolConstTrue_Execute()
 {
-	Value = true;
+	// nothing to do here - value is assigned in the constructor
 }
 
 FRigVMFunction_MathBoolConstFalse_Execute()
 {
-	Value = false;
+	// nothing to do here - value is assigned in the constructor
 }
 
 FRigVMFunction_MathBoolNot_Execute()
diff --git a/Engine/Plugins/Runtime/RigVM/Source/RigVM/Private/RigVMFunctions/Math/RigVMFunction_MathDouble.cpp b/Engine/Plugins/Runtime/RigVM/Source/RigVM/Private/RigVMFunctions/Math/RigVMFunction_MathDouble.cpp
--- a/Engine/Plugins/Runtime/RigVM/Source/RigVM/Private/RigVMFunctions/Math/RigVMFunction_MathDouble.cpp
+++ b/Engine/Plugins/Runtime/RigVM/Source/RigVM/Private/RigVMFunctions/Math/RigVMFunction_MathDouble.cpp
@@ -17,22 +17,22 @@ FRigVMStructUpgradeInfo FRigVMFunction_MathDoubleMake::GetUpgradeInfo() const
 
 FRigVMFunction_MathDoubleConstPi_Execute()
 {
-	Value = PI;
+	// nothing to do here - value is assigned in the constructor
 }
 
 FRigVMFunction_MathDoubleConstHalfPi_Execute()
 {
-	Value = HALF_PI;
+	// nothing to do here - value is assigned in the constructor
 }
 
 FRigVMFunction_MathDoubleConstTwoPi_Execute()
 {
-	Value = PI * 2.0;
+	// nothing to do here - value is assigned in the constructor
 }
 
 FRigVMFunction_MathDoubleConstE_Execute()
 {
-	Value = EULERS_NUMBER;
+	// nothing to do here - value is assigned in the constructor
 }
 
 FRigVMFunction_MathDoubleAdd_Execute()
diff --git a/Engine/Plugins/Runtime/RigVM/Source/RigVM/Private/RigVMFunctions/Math/RigVMFunction_MathFloat.cpp b/Engine/Plugins/Runtime/RigVM/Source/RigVM/Private/RigVMFunctions/Math/RigVMFunction_MathFloat.cpp
--- a/Engine/Plugins/Runtime/RigVM/Source/RigVM/Private/RigVMFunctions/Math/RigVMFunction_MathFloat.cpp
+++ b/Engine/Plugins/Runtime/RigVM/Source/RigVM/Private/RigVMFunctions/Math/RigVMFunction_MathFloat.cpp
@@ -17,22 +17,22 @@ FRigVMStructUpgradeInfo FRigVMFunction_MathFloatMake::GetUpgradeInfo() const
 
 FRigVMFunction_MathFloatConstPi_Execute()
 {
-	Value = PI;
+	// nothing to do here - value is assigned in the constructor
 }
 
 FRigVMFunction_MathFloatConstHalfPi_Execute()
 {
-	Value = HALF_PI;
+	// nothing to do here - value is assigned in the constructor
 }
 
 FRigVMFunction_MathFloatConstTwoPi_Execute()
 {
-	Value = PI * 2.f;
+	// nothing to do here - value is assigned in the constructor
 }
 
 FRigVMFunction_MathFloatConstE_Execute()
 {
-	Value = EULERS_NUMBER;
+	// nothing to do here - value is assigned in the constructor
 }
 
 FRigVMFunction_MathFloatAdd_Execute()
diff --git a/Engine/Plugins/Runtime/RigVM/Source/RigVM/Public/RigVMCore/RigVMStruct.h b/Engine/Plugins/Runtime/RigVM/Source/RigVM/Public/RigVMCore/RigVMStruct.h
--- a/Engine/Plugins/Runtime/RigVM/Source/RigVM/Public/RigVMCore/RigVMStruct.h
+++ b/Engine/Plugins/Runtime/RigVM/Source/RigVM/Public/RigVMCore/RigVMStruct.h
@@ -267,6 +267,7 @@ public:
 	static inline const FLazyName ShowVariableNameInTitleMetaName = FLazyName(TEXT("ShowVariableNameInTitle"));
 	static inline const FLazyName CustomWidgetMetaName = FLazyName(TEXT("CustomWidget"));
 	static inline const FLazyName ConstantMetaName = FLazyName(TEXT("Constant"));
+	static inline const FLazyName NoOpMetaName = FLazyName(TEXT("NoOp"));
 	static inline const FLazyName TitleColorMetaName = FLazyName(TEXT("TitleColor"));
 	static inline const FLazyName NodeColorMetaName = FLazyName(TEXT("NodeColor"));
 	// icon meta name format: StyleSetName|StyleName|SmallStyleName|StatusOverlayStyleName
diff --git a/Engine/Plugins/Runtime/RigVM/Source/RigVM/Public/RigVMFunctions/Math/RigVMFunction_MathBool.h b/Engine/Plugins/Runtime/RigVM/Source/RigVM/Public/RigVMFunctions/Math/RigVMFunction_MathBool.h
--- a/Engine/Plugins/Runtime/RigVM/Source/RigVM/Public/RigVMFunctions/Math/RigVMFunction_MathBool.h
+++ b/Engine/Plugins/Runtime/RigVM/Source/RigVM/Public/RigVMFunctions/Math/RigVMFunction_MathBool.h
@@ -17,7 +17,7 @@ struct FRigVMFunction_MathBoolBase : public FRigVMFunction_MathBase
 /*
  * The base class for all bool math constants
  */
-USTRUCT(meta=(Abstract))
+USTRUCT(meta=(Abstract, NoOp))
 struct FRigVMFunction_MathBoolConstant : public FRigVMFunction_MathBoolBase
 {
 	GENERATED_BODY()
@@ -28,7 +28,7 @@ struct FRigVMFunction_MathBoolConstant : public FRigVMFunction_MathBoolBase
 	}
 
 	// The value of the constant
-	UPROPERTY(meta=(Output))
+	UPROPERTY(meta=(Output,Constant))
 	bool Value;
 };
 
diff --git a/Engine/Plugins/Runtime/RigVM/Source/RigVM/Public/RigVMFunctions/Math/RigVMFunction_MathDouble.h b/Engine/Plugins/Runtime/RigVM/Source/RigVM/Public/RigVMFunctions/Math/RigVMFunction_MathDouble.h
--- a/Engine/Plugins/Runtime/RigVM/Source/RigVM/Public/RigVMFunctions/Math/RigVMFunction_MathDouble.h
+++ b/Engine/Plugins/Runtime/RigVM/Source/RigVM/Public/RigVMFunctions/Math/RigVMFunction_MathDouble.h
@@ -17,7 +17,7 @@ struct FRigVMFunction_MathDoubleBase : public FRigVMFunction_MathBase
 /*
  * The base class for all double math constants
  */
-USTRUCT(meta=(Abstract))
+USTRUCT(meta=(Abstract, NoOp))
 struct FRigVMFunction_MathDoubleConstant : public FRigVMFunction_MathDoubleBase
 {
 	GENERATED_BODY()
@@ -28,7 +28,7 @@ struct FRigVMFunction_MathDoubleConstant : public FRigVMFunction_MathDoubleBase
 	}
 
 	// The value of the constant
-	UPROPERTY(meta=(Output))
+	UPROPERTY(meta=(Output,Constant))
 	double Value;
 };
 
@@ -136,6 +136,11 @@ USTRUCT(meta=(DisplayName="Pi", TemplateName="Pi"))
 struct FRigVMFunction_MathDoubleConstPi : public FRigVMFunction_MathDoubleConstant
 {
 	GENERATED_BODY()
+
+	FRigVMFunction_MathDoubleConstPi()
+	{
+		Value = PI;
+	}
 	
 	RIGVM_METHOD()
 	RIGVM_API virtual void Execute() override;
@@ -149,6 +154,11 @@ struct FRigVMFunction_MathDoubleConstHalfPi : public FRigVMFunction_MathDoubleCo
 {
 	GENERATED_BODY()
 
+	FRigVMFunction_MathDoubleConstHalfPi()
+	{
+		Value = HALF_PI;
+	}
+
 	RIGVM_METHOD()
 	RIGVM_API virtual void Execute() override;
 };
@@ -161,6 +171,11 @@ struct FRigVMFunction_MathDoubleConstTwoPi : public FRigVMFunction_MathDoubleCon
 {
 	GENERATED_BODY()
 
+	FRigVMFunction_MathDoubleConstTwoPi()
+	{
+		Value = 2.0 * PI;
+	}
+
 	RIGVM_METHOD()
 	RIGVM_API virtual void Execute() override;
 };
@@ -172,6 +187,11 @@ USTRUCT(meta=(DisplayName="E", TemplateName="E", Keywords="Euler"))
 struct FRigVMFunction_MathDoubleConstE : public FRigVMFunction_MathDoubleConstant
 {
 	GENERATED_BODY()
+
+	FRigVMFunction_MathDoubleConstE()
+	{
+		Value = EULERS_NUMBER;
+	}
 	
 	RIGVM_METHOD()
 	RIGVM_API virtual void Execute() override;
diff --git a/Engine/Plugins/Runtime/RigVM/Source/RigVM/Public/RigVMFunctions/Math/RigVMFunction_MathFloat.h b/Engine/Plugins/Runtime/RigVM/Source/RigVM/Public/RigVMFunctions/Math/RigVMFunction_MathFloat.h
--- a/Engine/Plugins/Runtime/RigVM/Source/RigVM/Public/RigVMFunctions/Math/RigVMFunction_MathFloat.h
+++ b/Engine/Plugins/Runtime/RigVM/Source/RigVM/Public/RigVMFunctions/Math/RigVMFunction_MathFloat.h
@@ -19,7 +19,7 @@ struct FRigVMFunction_MathFloatBase : public FRigVMFunction_MathBase
 /*
  * The base class for all float math constants
  */
-USTRUCT(meta=(Abstract))
+USTRUCT(meta=(Abstract, NoOp))
 struct FRigVMFunction_MathFloatConstant : public FRigVMFunction_MathFloatBase
 {
 	GENERATED_BODY()
@@ -30,7 +30,7 @@ struct FRigVMFunction_MathFloatConstant : public FRigVMFunction_MathFloatBase
 	}
 
 	// The value of the constant
-	UPROPERTY(meta=(Output))
+	UPROPERTY(meta=(Output,Constant))
 	float Value;
 };
 
@@ -139,6 +139,11 @@ struct FRigVMFunction_MathFloatConstPi : public FRigVMFunction_MathFloatConstant
 {
 	GENERATED_BODY()
 	
+	FRigVMFunction_MathFloatConstPi()
+	{
+		Value = PI;
+	}
+	
 	RIGVM_METHOD()
 	UE_API virtual void Execute() override;
 };
@@ -150,6 +155,11 @@ USTRUCT(meta=(DisplayName="Half Pi", TemplateName="HalfPi"))
 struct FRigVMFunction_MathFloatConstHalfPi : public FRigVMFunction_MathFloatConstant
 {
 	GENERATED_BODY()
+	
+	FRigVMFunction_MathFloatConstHalfPi()
+	{
+		Value = HALF_PI;
+	}
 
 	RIGVM_METHOD()
 	UE_API virtual void Execute() override;
@@ -162,6 +172,11 @@ USTRUCT(meta=(DisplayName="Two Pi", TemplateName="TwoPi", Keywords="Tau"))
 struct FRigVMFunction_MathFloatConstTwoPi : public FRigVMFunction_MathFloatConstant
 {
 	GENERATED_BODY()
+	
+	FRigVMFunction_MathFloatConstTwoPi()
+	{
+		Value = PI * 2.f;
+	}
 
 	RIGVM_METHOD()
 	UE_API virtual void Execute() override;
@@ -175,6 +190,11 @@ struct FRigVMFunction_MathFloatConstE : public FRigVMFunction_MathFloatConstant
 {
 	GENERATED_BODY()
 	
+	FRigVMFunction_MathFloatConstE()
+	{
+		Value = EULERS_NUMBER;
+	}
+	
 	RIGVM_METHOD()
 	UE_API virtual void Execute() override;
 };
diff --git a/Engine/Plugins/Runtime/RigVM/Source/RigVMDeveloper/Private/RigVMCompiler/RigVMAST.cpp b/Engine/Plugins/Runtime/RigVM/Source/RigVMDeveloper/Private/RigVMCompiler/RigVMAST.cpp
--- a/Engine/Plugins/Runtime/RigVM/Source/RigVMDeveloper/Private/RigVMCompiler/RigVMAST.cpp
+++ b/Engine/Plugins/Runtime/RigVM/Source/RigVMDeveloper/Private/RigVMCompiler/RigVMAST.cpp
@@ -1477,6 +1477,7 @@ FRigVMExprAST* FRigVMParserAST::CreateExpressionForNode(const FRigVMASTProxy& In
 			InNodeProxy.IsA<URigVMEnumNode>() ||
 			InNodeProxy.IsA<URigVMLibraryNode>() ||
 			InNodeProxy.IsA<URigVMFunctionInterfaceNode>() ||
+			IsNoOpUnitNode(Node) ||
 			IsConstantDispatchNode(Node) ||
 			IsMakeStructDispatchNode(Node) ||
 			IsBreakStructDispatchNode(Node))
@@ -1690,6 +1691,10 @@ FRigVMExprAST* FRigVMParserAST::TraversePin(const FRigVMASTProxy& InPinProxy, FR
 	{
 		PinExpr = MakeExpr<FRigVMLiteralExprAST>(InPinProxy);
 	}
+	else if (IsNoOpUnitNode(Pin->GetNode()))
+	{
+		PinExpr = MakeExpr<FRigVMLiteralExprAST>(InPinProxy);
+	}
 	else
 	{
 		PinExpr = MakeExpr<FRigVMVarExprAST>(FRigVMExprAST::EType::Var, InPinProxy);
@@ -3949,6 +3954,23 @@ void FRigVMParserAST::IncrementCacheVersion() const
 	CacheVersion++;
 }
 
+bool FRigVMParserAST::IsNoOpUnitNode(const URigVMNode* InNode)
+{
+#if WITH_EDITOR
+	if(const URigVMUnitNode* UnitNode = Cast<URigVMUnitNode>(InNode))
+	{
+		if (const UScriptStruct* ScriptStruct = UnitNode->GetScriptStruct())
+		{
+			if (ScriptStruct->HasMetaDataHierarchical(FRigVMStruct::NoOpMetaName))
+			{
+				return true;
+			}
+		}
+	}
+#endif
+	return false;
+}
+
 bool FRigVMParserAST::IsConstantDispatchNode(const URigVMNode* InNode)
 {
 	if(const URigVMDispatchNode* DispatchNode = Cast<URigVMDispatchNode>(InNode))
diff --git a/Engine/Plugins/Runtime/RigVM/Source/RigVMDeveloper/Public/RigVMCompiler/RigVMAST.h b/Engine/Plugins/Runtime/RigVM/Source/RigVMDeveloper/Public/RigVMCompiler/RigVMAST.h
--- a/Engine/Plugins/Runtime/RigVM/Source/RigVMDeveloper/Public/RigVMCompiler/RigVMAST.h
+++ b/Engine/Plugins/Runtime/RigVM/Source/RigVMDeveloper/Public/RigVMCompiler/RigVMAST.h
@@ -1444,6 +1444,7 @@ private:
 
 	UE_API void IncrementCacheVersion() const;
 
+	static bool IsNoOpUnitNode(const URigVMNode* InNode);
 	static bool IsConstantDispatchNode(const URigVMNode* InNode);
 	static bool IsConstantDispatchValuePin(const URigVMPin* InPin);
  	static bool IsMakeStructDispatchNode(const URigVMNode* InNode);
