RigVM: undoing change on Cast Object Dispatch - removing Success template arg.

#rnx

diff --git a/Engine/Plugins/Runtime/RigVM/Source/RigVM/Private/RigVMFunctions/RigVMDispatch_CastObject.cpp b/Engine/Plugins/Runtime/RigVM/Source/RigVM/Private/RigVMFunctions/RigVMDispatch_CastObject.cpp
--- a/Engine/Plugins/Runtime/RigVM/Source/RigVM/Private/RigVMFunctions/RigVMDispatch_CastObject.cpp
+++ b/Engine/Plugins/Runtime/RigVM/Source/RigVM/Private/RigVMFunctions/RigVMDispatch_CastObject.cpp
@@ -19,7 +19,6 @@ const TArray<FRigVMTemplateArgumentInfo>& FRigVMDispatch_CastObject::GetArgument
 		TArray<FRigVMTemplateArgumentInfo> Infos;
 		Infos.Emplace(ValueName, ERigVMPinDirection::Input, ElementCategories);
 		Infos.Emplace(ResultName, ERigVMPinDirection::Output, INDEX_NONE);
-		Infos.Emplace(SuccessName, ERigVMPinDirection::Output, RigVMTypeUtils::TypeIndex::Bool);
 		
 		CachedArgumentInfos = BuildArgumentListFromPrimaryArgument(Infos, ValueName, InRegistry);
 	}
@@ -37,8 +36,7 @@ bool FRigVMDispatch_CastObject::GetPermutationsFromArgumentType(const FName& InA
 			OutPermutations.Add(
 	{
 				{ ValueName, InTypeIndex },
-				{ ResultName, Type },
-				{ SuccessName, RigVMTypeUtils::TypeIndex::Bool }
+				{ ResultName, Type }
 			});
 		}
 	}
@@ -50,8 +48,7 @@ bool FRigVMDispatch_CastObject::GetPermutationsFromArgumentType(const FName& InA
 			OutPermutations.Add(
 	{
 				{ ValueName, Type },
-				{ ResultName, InTypeIndex },
-				{ SuccessName, RigVMTypeUtils::TypeIndex::Bool }
+				{ ResultName, InTypeIndex }
 			});
 		}
 	}
@@ -115,10 +112,6 @@ FText FRigVMDispatch_CastObject::GetArgumentTooltip(const FName& InArgumentName,
 	{
 		return LOCTEXT("ResultToolTip", "The resulting cast object. This may be potentially nullptr as well if the cast was not successful.");
 	}
-	if (InArgumentName == SuccessName)
-	{
-		return LOCTEXT("SuccessToolTip", "A boolean result indicating if the cast was successful. If the cast failed the result will be nullptr.");
-	}
 	return FRigVMDispatchFactory::GetArgumentTooltip(InArgumentName, InTypeIndex);
 }
 
@@ -128,20 +121,17 @@ void FRigVMDispatch_CastObject::Execute(FRigVMExtendedExecuteContext& InContext,
 {
 	UObject** ValueObjectPtr = (UObject**)Handles[0].GetInputData();
 	UObject** ResultObjectPtr = (UObject**)Handles[1].GetOutputData();
-	bool* Success = (bool*)Handles[2].GetOutputData();
 	const FObjectPropertyBase* ValueObjectProperty = CastFieldChecked<FObjectPropertyBase>(Handles[0].GetResolvedProperty());
 	const FObjectPropertyBase* ResultObjectProperty = CastFieldChecked<FObjectPropertyBase>(Handles[1].GetProperty());
-	if(ValueObjectPtr && ResultObjectPtr && ValueObjectProperty && ResultObjectProperty && Success)
+	if(ValueObjectPtr && ResultObjectPtr && ValueObjectProperty && ResultObjectProperty)
 	{
 		if(*ValueObjectPtr && (*ValueObjectPtr)->GetClass()->IsChildOf(ResultObjectProperty->PropertyClass))
 		{
 			*ResultObjectPtr = *ValueObjectPtr;
-			*Success = true;
 		}
 		else
 		{
 			*ResultObjectPtr = nullptr;
-			*Success = false;
 		}
 	}
 }
diff --git a/Engine/Plugins/Runtime/RigVM/Source/RigVM/Public/RigVMFunctions/RigVMDispatch_CastObject.h b/Engine/Plugins/Runtime/RigVM/Source/RigVM/Public/RigVMFunctions/RigVMDispatch_CastObject.h
--- a/Engine/Plugins/Runtime/RigVM/Source/RigVM/Public/RigVMFunctions/RigVMDispatch_CastObject.h
+++ b/Engine/Plugins/Runtime/RigVM/Source/RigVM/Public/RigVMFunctions/RigVMDispatch_CastObject.h
@@ -32,5 +32,4 @@ protected:
 
 	static inline const FLazyName ValueName = FLazyName(TEXT("Value"));
 	static inline const FLazyName ResultName = FLazyName(TEXT("Result"));
-	static inline const FLazyName SuccessName = FLazyName(TEXT("Success"));
 };
