RigVM: Improve serialization of RigVMFunction for localized registry

#rnx

diff --git a/Engine/Plugins/Runtime/RigVM/Source/RigVM/Private/RigVMCore/RigVMFunction.cpp b/Engine/Plugins/Runtime/RigVM/Source/RigVM/Private/RigVMCore/RigVMFunction.cpp
--- a/Engine/Plugins/Runtime/RigVM/Source/RigVM/Private/RigVMCore/RigVMFunction.cpp
+++ b/Engine/Plugins/Runtime/RigVM/Source/RigVM/Private/RigVMCore/RigVMFunction.cpp
@@ -325,10 +325,71 @@ void FRigVMFunction::Serialize_NoLock(FArchive& Ar, FRigVMRegistryHandle& InRegi
 				FunctionPtr = GlobalFunction->FunctionPtr;
 #if WITH_EDITOR
 				check(GlobalFunction->Struct == Struct);
-				check(GlobalFunction->Arguments.Num() == Arguments.Num());
-				for(int32 Arg = 0; Arg< Arguments.Num(); Arg++)
+
+				auto CompareArguments = [GlobalFunction](const FRigVMFunctionArgument& SerializedArg, const FRigVMFunctionArgument& CurrentArg)
+				{
+					if (SerializedArg == CurrentArg)
+					{
+						return;
+					}
+					UE_LOG(LogRigVM, Error, TEXT("Localized Registry Function '%s' mismatch for argument '%s'. Please outdate function instead of changing them in the code and roll back potential changes."), *GlobalFunction->Struct->GetStructCPPName(), CurrentArg.Name);
+				};
+
+				// do a rough pass to see of all arguments match.
+				bool bAllArgumentsMatch = true;
+				if (GlobalFunction->Arguments.Num() == Arguments.Num())
+				{
+					for(int32 Arg = 0; Arg < Arguments.Num(); Arg++)
+					{
+						if(Arguments[Arg] != GlobalFunction->Arguments[Arg])
+						{
+							bAllArgumentsMatch = false;
+						}
+					}
+				}
+				
+				// if we have more arguments in the serialized file compared
+				// to what we have now - we need to let users know about the missing arguments.
+				if(GlobalFunction->Arguments.Num() < Arguments.Num())
+				{
+					TMap<FName, int32> GlobalArgumentToIndex;
+					for(int32 GlobalArgIndex = 0; GlobalArgIndex < GlobalFunction->Arguments.Num(); GlobalArgIndex++)
+					{
+						GlobalArgumentToIndex.Add(GlobalFunction->Arguments[GlobalArgIndex].Name, GlobalArgIndex);
+					}
+					for(int32 SerializedArgIndex = 0; SerializedArgIndex < Arguments.Num(); SerializedArgIndex++)
+					{
+						const FName ArgName = Arguments[SerializedArgIndex].Name;
+						const int32* GlobalArgIndex = GlobalArgumentToIndex.Find(ArgName);
+						if (!GlobalArgIndex)
+						{
+							UE_LOG(LogRigVM, Error, TEXT("Localized Registry Function '%s' is missing formerly existing argument '%s'. Please outdate function instead of changing them in the code and roll back potential changes."), *GlobalFunction->Struct->GetStructCPPName(), *ArgName.ToString());
+							continue;
+						}
+						CompareArguments(Arguments[SerializedArgIndex], GlobalFunction->Arguments[*GlobalArgIndex]);
+					}
+				}
+
+				// if we have more or sufficient arguments now,
+				// we still need to check mismatches.
+				else if (!bAllArgumentsMatch)
 				{
-					check(GlobalFunction->Arguments[Arg] == Arguments[Arg]);
+					TMap<FName, int32> SerializedArgumentToIndex;
+					for(int32 SerializedArgIndex = 0; SerializedArgIndex < Arguments.Num(); SerializedArgIndex++)
+					{
+						SerializedArgumentToIndex.Add(Arguments[SerializedArgIndex].Name, SerializedArgIndex);
+					}
+					for(int32 GlobalArgIndex = 0; GlobalArgIndex < GlobalFunction->Arguments.Num(); GlobalArgIndex++)
+					{
+						const FName ArgName = GlobalFunction->Arguments[GlobalArgIndex].Name;
+						const int32* SerializedArgIndex = SerializedArgumentToIndex.Find(ArgName);
+						if (!SerializedArgIndex)
+						{
+							// if the function has more arguments now - that's not in the file - that's ok.
+							continue;
+						}
+						CompareArguments(Arguments[*SerializedArgIndex], GlobalFunction->Arguments[GlobalArgIndex]);
+					}
 				}
 #endif
 			}
