Improved RigVM Execute function memory handles patching cost by not allocating memory and caching the memory base address of the VM memory storages

#rnx

diff --git a/Engine/Plugins/Runtime/RigVM/Source/RigVM/Private/RigVMCore/RigVM.cpp b/Engine/Plugins/Runtime/RigVM/Source/RigVM/Private/RigVMCore/RigVM.cpp
--- a/Engine/Plugins/Runtime/RigVM/Source/RigVM/Private/RigVMCore/RigVM.cpp
+++ b/Engine/Plugins/Runtime/RigVM/Source/RigVM/Private/RigVMCore/RigVM.cpp
@@ -1821,6 +1821,20 @@ ERigVMExecuteResult URigVM::ExecuteInstructions(FRigVMExtendedExecuteContext& Co
 	TArray<const FRigVMFunction*>& Functions = GetFunctions();
 	TArray<const FRigVMDispatchFactory*>& Factories = GetFactories();
 
+	TArray<FRigVMMemoryHandle, TInlineAllocator<64>> HandlesMemory; // have ready some memory to patch handles
+
+	// have the VM instance memory by type ready for handles patching
+	static_assert((int32)ERigVMMemoryType::Num == 4);
+	const TStaticArray<uint8*, (int32)ERigVMMemoryType::Num> VMMemoryPtrs =
+	{
+		GetMemoryByType(Context, ERigVMMemoryType::Work) != nullptr ? (uint8*)GetMemoryByType(Context, ERigVMMemoryType::Work)->GetContainerPtr() : nullptr,
+		GetMemoryByType(Context, ERigVMMemoryType::Literal) != nullptr ? (uint8*)GetMemoryByType(Context, ERigVMMemoryType::Literal)->GetContainerPtr() : nullptr,
+		GetMemoryByType(Context, ERigVMMemoryType::External) != nullptr ? (uint8*)GetMemoryByType(Context, ERigVMMemoryType::External)->GetContainerPtr() : nullptr,
+		GetMemoryByType(Context, ERigVMMemoryType::Debug) != nullptr ? (uint8*)GetMemoryByType(Context, ERigVMMemoryType::Debug)->GetContainerPtr() : nullptr,
+	};
+	check(PreCachedMemoryHandles.Num() == MemoryHandleCount);
+	check(PreCachedMemoryHandlesMemoryType.Num() == MemoryHandleCount);
+
 #if WITH_EDITOR
 	TArray<FName>& FunctionNames = GetFunctionNames();
 	FRigVMInstructionVisitInfo* RigVMInstructionVisitInfo = Context.GetRigVMInstructionVisitInfo();
@@ -1889,18 +1903,6 @@ ERigVMExecuteResult URigVM::ExecuteInstructions(FRigVMExtendedExecuteContext& Co
 #endif
 #endif
 
-		// have the VM instance memory by type ready for handles patching
-		static_assert((int32)ERigVMMemoryType::Num == 4);
-		const TStaticArray<FRigVMMemoryStorageStruct*, (int32)ERigVMMemoryType::Num> VMMemoryPtrs = 
-		{
-			GetMemoryByType(Context, ERigVMMemoryType::Work),
-			GetMemoryByType(Context, ERigVMMemoryType::Literal),
-			GetMemoryByType(Context, ERigVMMemoryType::External),
-			GetMemoryByType(Context, ERigVMMemoryType::Debug),
-		};
-		check(PreCachedMemoryHandles.Num() == MemoryHandleCount);
-		check(PreCachedMemoryHandlesMemoryType.Num() == MemoryHandleCount);
-
 		switch (Instruction.OpCode)
 		{
 			case ERigVMOpCode::Execute:
@@ -1908,10 +1910,12 @@ ERigVMExecuteResult URigVM::ExecuteInstructions(FRigVMExtendedExecuteContext& Co
 				const FRigVMExecuteOp& Op = ByteCode.GetOpAt<FRigVMExecuteOp>(Instruction);
 				const int32 OperandCount = FirstHandleForInstruction[ContextPublicData.InstructionIndex + 1] - FirstHandleForInstruction[ContextPublicData.InstructionIndex];
 
-				TArray<FRigVMMemoryHandle> Handles;
+				TArrayView<FRigVMMemoryHandle> Handles;
 
 				if(OperandCount > 0)
 				{
+					HandlesMemory.SetNumUninitialized(OperandCount, EAllowShrinking::No);
+					Handles = TArrayView<FRigVMMemoryHandle>(HandlesMemory.GetData(), OperandCount);
 					GeneratePatchedHandles(Handles, Context.ExternalVariableRuntimeData, VMMemoryPtrs, PreCachedMemoryHandles, PreCachedMemoryHandlesMemoryType, FirstHandleForInstruction[ContextPublicData.InstructionIndex], OperandCount);
 				}
 
@@ -3232,15 +3236,15 @@ void URigVM::PreCacheSingleMemoryHandle(const FRigVMRegistryHandle& InRegistry,
 }
 
 void URigVM::GeneratePatchedHandles(
-	TArray<FRigVMMemoryHandle>& OutHandles,
+	const TArrayView<FRigVMMemoryHandle>& OutHandles,
 	const TArray<FRigVMExternalVariableRuntimeData>& ExternalVariablesRuntimeData,
-	const TStaticArray<FRigVMMemoryStorageStruct*, (int32)ERigVMMemoryType::Num>& VMMemoryPtrs,
+	const TStaticArray<uint8*, (int32)ERigVMMemoryType::Num>& VMMemoryPtrs,
 	const TArray<FRigVMMemoryHandle>& PreCachedMemoryHandles,
 	const TArray<ERigVMMemoryType>& PreCachedMemoryHandlesMemoryType,
 	int32 FirstHandleIndex,
 	int32 OperandCount)
 {
-	OutHandles.SetNumUninitialized(OperandCount);
+	check(OutHandles.Num() == OperandCount);
 	for (int32 HandleIndex = 0; HandleIndex < OperandCount; HandleIndex++)
 	{
 		const int32 PreCachedHandleIndex = FirstHandleIndex + HandleIndex;
@@ -3252,7 +3256,7 @@ void URigVM::GeneratePatchedHandles(
 
 FRigVMMemoryHandle URigVM::GeneratePatchedHandle(
 	const TArray<FRigVMExternalVariableRuntimeData>& ExternalVariablesRuntimeData,
-	const TStaticArray<FRigVMMemoryStorageStruct*, (int32)ERigVMMemoryType::Num>& VMMemoryPtrs,
+	const TStaticArray<uint8*, (int32)ERigVMMemoryType::Num>& VMMemoryPtrs,
 	ERigVMMemoryType HandleMemoryType,
 	const FRigVMMemoryHandle& PreCachedHandle)
 {
@@ -3270,9 +3274,9 @@ FRigVMMemoryHandle URigVM::GeneratePatchedHandle(
 		return Handle;
 	}
 
-	if (const FRigVMMemoryStorageStruct* Memory = VMMemoryPtrs[(int32)HandleMemoryType])
+	if (const uint8* Memory = VMMemoryPtrs[(int32)HandleMemoryType])
 	{
-		const PTRINT BaseAddress = (PTRINT)Memory->GetContainerPtr();
+		const PTRINT BaseAddress = (PTRINT)Memory;
 		const PTRINT Offset = (PTRINT)Handle.Ptr;
 		Handle.Ptr = (uint8*)(BaseAddress + Offset);
 	}
diff --git a/Engine/Plugins/Runtime/RigVM/Source/RigVM/Public/RigVMCore/RigVM.h b/Engine/Plugins/Runtime/RigVM/Source/RigVM/Public/RigVMCore/RigVM.h
--- a/Engine/Plugins/Runtime/RigVM/Source/RigVM/Public/RigVMCore/RigVM.h
+++ b/Engine/Plugins/Runtime/RigVM/Source/RigVM/Public/RigVMCore/RigVM.h
@@ -1051,9 +1051,9 @@ private:
 	RIGVM_API void PreCacheSingleMemoryHandle(const FRigVMRegistryHandle& InRegistry, int32 InInstructionIndex, int32 InHandleIndex, const FRigVMBranchInfoKey& InBranchInfoKey, const FRigVMOperand& InArg, bool bForExecute = false);
 
 	RIGVM_API static void GeneratePatchedHandles(
-		TArray<FRigVMMemoryHandle>& OutHandles,
+		const TArrayView<FRigVMMemoryHandle>& OutHandles,
 		const TArray<FRigVMExternalVariableRuntimeData>& ExternalVariablesRuntimeData,
-		const TStaticArray<FRigVMMemoryStorageStruct*, (int32)ERigVMMemoryType::Num>& VMMemoryPtrs,
+		const TStaticArray<uint8*, (int32)ERigVMMemoryType::Num>& VMMemoryPtrs,
 		const TArray<FRigVMMemoryHandle>& PreCachedMemoryHandles,
 		const TArray<ERigVMMemoryType>& PreCachedMemoryHandlesMemoryType,
 		int32 FirstHandleIndex,
@@ -1061,7 +1061,7 @@ private:
 
 	RIGVM_API static FRigVMMemoryHandle GeneratePatchedHandle(
 		const TArray<FRigVMExternalVariableRuntimeData>& ExternalVariablesRuntimeData,
-		const TStaticArray<FRigVMMemoryStorageStruct*, (int32)ERigVMMemoryType::Num>& VMMemoryPtrs,
+		const TStaticArray<uint8*, (int32)ERigVMMemoryType::Num>& VMMemoryPtrs,
 		ERigVMMemoryType HandleMemoryType,
 		const FRigVMMemoryHandle& PreCachedHandle);
 
