RigVM: Localized registry to always allow types

#rnx

diff --git a/Engine/Plugins/Runtime/RigVM/Source/RigVM/Private/RigVMCore/RigVM.cpp b/Engine/Plugins/Runtime/RigVM/Source/RigVM/Private/RigVMCore/RigVM.cpp
--- a/Engine/Plugins/Runtime/RigVM/Source/RigVM/Private/RigVMCore/RigVM.cpp
+++ b/Engine/Plugins/Runtime/RigVM/Source/RigVM/Private/RigVMCore/RigVM.cpp
@@ -2649,7 +2649,7 @@ bool URigVM::CreateLocalizedRegistryIfRequired_NoLock(bool bForce)
 			{
 				TemporaryHandle->AddSpecificallyAllowedType_NoLock(CPPTypeObject);
 			}
-			TemporaryHandle->FindOrAddType_NoLock(FRigVMTemplateArgumentType(CPPTypeName, CPPTypeObject, bIsSerializing));
+			TemporaryHandle->FindOrAddType_NoLock(FRigVMTemplateArgumentType(CPPTypeName, CPPTypeObject, bIsSerializing), true /* force */);
 		}
 	}
 
@@ -2674,7 +2674,7 @@ bool URigVM::CreateLocalizedRegistryIfRequired_NoLock(bool bForce)
 							{
 								TemporaryHandle->AddSpecificallyAllowedType_NoLock(Type.CPPTypeObject);
 							}
-							TemporaryRegistry->FindOrAddType_NoLock(Type);
+							TemporaryRegistry->FindOrAddType_NoLock(Type, true /* force */);
 						}
 						else
 						{
@@ -2739,11 +2739,11 @@ bool URigVM::CreateLocalizedRegistryIfRequired_NoLock(bool bForce)
 						return ResetLocalizedRegistry();
 					}
 
-					if (OriginalType.CPPTypeObject && ReadLock->IsSpecificallyAllowedType_NoLock(OriginalType.CPPTypeObject))
+					if (OriginalType.CPPTypeObject && (OriginalType.CPPTypeObject->IsA<UClass>() || ReadLock->IsSpecificallyAllowedType_NoLock(OriginalType.CPPTypeObject)))
 					{
 						TemporaryRegistry->AddSpecificallyAllowedType_NoLock(OriginalType.CPPTypeObject);
 					}
-					TemporaryTypeIndices.Add(TemporaryRegistry->FindOrAddType_NoLock(OriginalType));
+					TemporaryTypeIndices.Add(TemporaryRegistry->FindOrAddType_NoLock(OriginalType, true /* force */));
 					if (TemporaryTypeIndices.Last() == INDEX_NONE)
 					{
 						UE_LOG(LogRigVM, Error, TEXT("Cannot register type '%s' in dispatch '%s' for VM '%s'\nCannot localize registry."), *OriginalType.CPPType.ToString(), *Factory->GetFactoryName().ToString(), *GetPathName());
