RigVM: Support no lock flavor for RigVMSchema methods

#rnx

diff --git a/Engine/Plugins/Animation/ControlRig/Source/ControlRigDeveloper/Private/ControlRigSchema.cpp b/Engine/Plugins/Animation/ControlRig/Source/ControlRigDeveloper/Private/ControlRigSchema.cpp
--- a/Engine/Plugins/Animation/ControlRig/Source/ControlRigDeveloper/Private/ControlRigSchema.cpp
+++ b/Engine/Plugins/Animation/ControlRig/Source/ControlRigDeveloper/Private/ControlRigSchema.cpp
@@ -46,20 +46,19 @@ bool UControlRigSchema::ShouldUnfoldStruct(URigVMController* InController, const
 	return true;
 }
 
-bool UControlRigSchema::SupportsUnitFunction(URigVMController* InController, const FRigVMFunction* InUnitFunction) const
+bool UControlRigSchema::SupportsUnitFunction_NoLock(URigVMController* InController, const FRigVMFunction* InUnitFunction, FRigVMRegistryHandle& InRegistry) const
 {
 	if (InUnitFunction->Struct == FRigUnit_ConnectorExecution::StaticStruct() ||
 		InUnitFunction->Struct == FRigUnit_GetCandidates::StaticStruct() ||
 		InUnitFunction->Struct == FRigUnit_DiscardMatches::StaticStruct())
 	{
-		if (UControlRigBlueprint* Blueprint = Cast<UControlRigBlueprint>(InController->GetOuter()))
+		if (InController)
 		{
-			if (Blueprint->IsControlRigModule())
+			if (UControlRigBlueprint* Blueprint = Cast<UControlRigBlueprint>(InController->GetOuter()))
 			{
-				return true;
+				return Blueprint->IsControlRigModule();
 			}
-			return false;
 		}
 	}
-	return Super::SupportsUnitFunction(InController, InUnitFunction);
+	return Super::SupportsUnitFunction_NoLock(InController, InUnitFunction, InRegistry);
 }
diff --git a/Engine/Plugins/Animation/ControlRig/Source/ControlRigDeveloper/Public/ControlRigSchema.h b/Engine/Plugins/Animation/ControlRig/Source/ControlRigDeveloper/Public/ControlRigSchema.h
--- a/Engine/Plugins/Animation/ControlRig/Source/ControlRigDeveloper/Public/ControlRigSchema.h
+++ b/Engine/Plugins/Animation/ControlRig/Source/ControlRigDeveloper/Public/ControlRigSchema.h
@@ -3,6 +3,7 @@
 #pragma once
 
 #include "RigVMModel/RigVMSchema.h"
+#include "Graph/ControlRigGraphSchema.h"
 #include "ControlRigSchema.generated.h"
 
 #define UE_API CONTROLRIGDEVELOPER_API
@@ -14,9 +15,14 @@ class UControlRigSchema : public URigVMSchema
 
 public:
 
+	virtual UClass* GetEdGraphSchemaClass() const override
+	{
+		return UControlRigGraphSchema::StaticClass();
+	}
+
 	UE_API virtual bool ShouldUnfoldStruct(URigVMController* InController, const UStruct* InStruct) const override;
 
-	UE_API virtual bool SupportsUnitFunction(URigVMController* InController, const FRigVMFunction* InUnitFunction) const override;
+	UE_API virtual bool SupportsUnitFunction_NoLock(URigVMController* InController, const FRigVMFunction* InUnitFunction, FRigVMRegistryHandle& InRegistry) const override;
 };
 
 #undef UE_API
diff --git a/Engine/Plugins/Experimental/UAF/UAF/Source/UAFUncookedOnly/Internal/AnimNextRigVMAssetSchema.h b/Engine/Plugins/Experimental/UAF/UAF/Source/UAFUncookedOnly/Internal/AnimNextRigVMAssetSchema.h
--- a/Engine/Plugins/Experimental/UAF/UAF/Source/UAFUncookedOnly/Internal/AnimNextRigVMAssetSchema.h
+++ b/Engine/Plugins/Experimental/UAF/UAF/Source/UAFUncookedOnly/Internal/AnimNextRigVMAssetSchema.h
@@ -3,6 +3,7 @@
 #pragma once
 
 #include "CoreMinimal.h"
+#include "AnimNextEdGraphSchema.h"
 #include "RigVMModel/RigVMSchema.h"
 #include "AnimNextRigVMAssetSchema.generated.h"
 
@@ -13,6 +14,11 @@ class UAnimNextRigVMAssetSchema : public URigVMSchema
 {
 	GENERATED_BODY()
 
+	virtual UClass* GetEdGraphSchemaClass() const override
+	{
+		return UAnimNextEdGraphSchema::StaticClass();
+	}
+
 	virtual bool SupportsNodeLayouts(const URigVMGraph* InGraph) const override
 	{
 		return true;
diff --git a/Engine/Plugins/Experimental/UAF/UAF/Source/UAFUncookedOnly/Internal/Module/AnimNextEventGraphSchema.h b/Engine/Plugins/Experimental/UAF/UAF/Source/UAFUncookedOnly/Internal/Module/AnimNextEventGraphSchema.h
--- a/Engine/Plugins/Experimental/UAF/UAF/Source/UAFUncookedOnly/Internal/Module/AnimNextEventGraphSchema.h
+++ b/Engine/Plugins/Experimental/UAF/UAF/Source/UAFUncookedOnly/Internal/Module/AnimNextEventGraphSchema.h
@@ -12,6 +12,6 @@ class UAnimNextEventGraphSchema : public UAnimNextRigVMAssetSchema
 	GENERATED_BODY()
 
 	// URigVMSchema interface
-	virtual bool SupportsUnitFunction(URigVMController* InController, const FRigVMFunction* InUnitFunction) const override;
+	virtual bool SupportsUnitFunction_NoLock(URigVMController* InController, const FRigVMFunction* InUnitFunction, FRigVMRegistryHandle& InRegistry) const override;
 };
 
diff --git a/Engine/Plugins/Experimental/UAF/UAF/Source/UAFUncookedOnly/Private/Module/AnimNextModule_EventGraphSchema.cpp b/Engine/Plugins/Experimental/UAF/UAF/Source/UAFUncookedOnly/Private/Module/AnimNextModule_EventGraphSchema.cpp
--- a/Engine/Plugins/Experimental/UAF/UAF/Source/UAFUncookedOnly/Private/Module/AnimNextModule_EventGraphSchema.cpp
+++ b/Engine/Plugins/Experimental/UAF/UAF/Source/UAFUncookedOnly/Private/Module/AnimNextModule_EventGraphSchema.cpp
@@ -3,7 +3,7 @@
 #include "Module/AnimNextEventGraphSchema.h"
 #include "AnimNextExecuteContext.h"
 
-bool UAnimNextEventGraphSchema::SupportsUnitFunction(URigVMController* InController, const FRigVMFunction* InUnitFunction) const
+bool UAnimNextEventGraphSchema::SupportsUnitFunction_NoLock(URigVMController* InController, const FRigVMFunction* InUnitFunction, FRigVMRegistryHandle& InRegistry) const
 {
 	if(const UScriptStruct* FunctionExecuteContextStruct = InUnitFunction->GetExecuteContextStruct())
 	{
@@ -19,5 +19,5 @@ bool UAnimNextEventGraphSchema::SupportsUnitFunction(URigVMController* InControl
 		}
 	}
 
-	return Super::SupportsUnitFunction(InController, InUnitFunction);
+	return Super::SupportsUnitFunction_NoLock(InController, InUnitFunction, InRegistry);
 }
diff --git a/Engine/Plugins/Experimental/UAF/UAFTests/Source/UAFTestSuite/Private/AnimNextVariablesTest.h b/Engine/Plugins/Experimental/UAF/UAFTests/Source/UAFTestSuite/Private/AnimNextVariablesTest.h
--- a/Engine/Plugins/Experimental/UAF/UAFTests/Source/UAFTestSuite/Private/AnimNextVariablesTest.h
+++ b/Engine/Plugins/Experimental/UAF/UAFTests/Source/UAFTestSuite/Private/AnimNextVariablesTest.h
@@ -3,6 +3,7 @@
 #pragma once
 
 #include "CoreMinimal.h"
+#include "AnimNextExecuteContext.h"
 #include "Kismet/BlueprintFunctionLibrary.h"
 #include "Param/ParamType.h"
 #include "RigVMCore/RigVMStruct.h"
@@ -139,7 +140,7 @@ struct FAnimNextParamTypeTestStruct
 //	TArray<TSoftClassPtr<UObject>> SoftClassPtrArray;
 };
 
-USTRUCT(meta=(Hidden, DocumentationPolicy="None"))
+USTRUCT(meta=(ExecuteContext="FAnimNextExecuteContext", Hidden, DocumentationPolicy="None"))
 struct FAnimNextTests_TestBase : public FRigVMStructMutable
 {
 	GENERATED_BODY()
diff --git a/Engine/Plugins/Runtime/RigVM/Source/RigVMDeveloper/Private/RigVMModel/RigVMSchema.cpp b/Engine/Plugins/Runtime/RigVM/Source/RigVMDeveloper/Private/RigVMModel/RigVMSchema.cpp
--- a/Engine/Plugins/Runtime/RigVM/Source/RigVMDeveloper/Private/RigVMModel/RigVMSchema.cpp
+++ b/Engine/Plugins/Runtime/RigVM/Source/RigVMDeveloper/Private/RigVMModel/RigVMSchema.cpp
@@ -2,6 +2,7 @@
 
 #include "RigVMModel/RigVMSchema.h"
 #include "RigVMStringUtils.h"
+#include "EdGraph/RigVMEdGraphSchema.h"
 #include "RigVMModel/RigVMClient.h"
 #include "RigVMModel/RigVMController.h"
 #include "RigVMFunctions/RigVMDispatch_Constant.h"
@@ -25,16 +26,31 @@ URigVMSchema::URigVMSchema(const FObjectInitializer& ObjectInitializer)
 	SetExecuteContextStruct(FRigVMExecuteContext::StaticStruct());
 }
 
+UClass* URigVMSchema::GetEdGraphSchemaClass() const
+{
+	return URigVMEdGraphSchema::StaticClass();
+}
+
 bool URigVMSchema::SupportsType(URigVMController* InController, TRigVMTypeIndex InTypeIndex) const
+{
+	check(Registry);
+	FRigVMRegistryWriteLock Lock(*Registry);
+	return SupportsType_NoLock(InController, InTypeIndex, Lock);
+}
+
+bool URigVMSchema::SupportsType_NoLock(URigVMController* InController, TRigVMTypeIndex InTypeIndex, FRigVMRegistryHandle& InRegistry) const
 {
 	// filter out incompatible execute types
-	if(Registry->IsExecuteType(InTypeIndex))
+	if(InRegistry->IsExecuteType_NoLock(InTypeIndex))
 	{
-		const UStruct* Struct = CastChecked<UStruct>(Registry->GetType(InTypeIndex).CPPTypeObject);
+		const UStruct* Struct = CastChecked<UStruct>(InRegistry->GetType_NoLock(InTypeIndex).CPPTypeObject);
 		if(!ValidExecuteContextStructs.Contains(Struct))
 		{
-			static constexpr TCHAR Format[] = TEXT("ExecuteContext struct '%s' is not supported.");
-			InController->ReportErrorf(Format, *Struct->GetName());
+			if (InController)
+			{
+				static constexpr TCHAR Format[] = TEXT("ExecuteContext struct '%s' is not supported.");
+				InController->ReportErrorf(Format, *Struct->GetName());
+			}
 			return false;
 		}
 	}
@@ -42,26 +58,36 @@ bool URigVMSchema::SupportsType(URigVMController* InController, TRigVMTypeIndex
 }
 
 bool URigVMSchema::SupportsUnitFunction(URigVMController* InController, const FRigVMFunction* InUnitFunction) const
+{
+	check(Registry);
+	FRigVMRegistryWriteLock Lock(*Registry);
+	return SupportsUnitFunction_NoLock(InController, InUnitFunction, Lock);
+}
+
+bool URigVMSchema::SupportsUnitFunction_NoLock(URigVMController* InController, const FRigVMFunction* InUnitFunction, FRigVMRegistryHandle& InRegistry) const
 {
 	if(InUnitFunction)
 	{
-		if(const UScriptStruct* FunctionExecuteContextStruct = InUnitFunction->GetExecuteContextStruct())
+		if(const UScriptStruct* FunctionExecuteContextStruct = InUnitFunction->GetExecuteContextStruct_NoLock(InRegistry))
 		{
-			const TRigVMTypeIndex TypeIndex = Registry->GetTypeIndexFromCPPType(FunctionExecuteContextStruct->GetStructCPPName());
-			if(!SupportsType(InController, TypeIndex))
+			const TRigVMTypeIndex TypeIndex = InRegistry->GetTypeIndexFromCPPType_NoLock(FunctionExecuteContextStruct->GetStructCPPName());
+			if(!SupportsType_NoLock(InController, TypeIndex, InRegistry))
 			{
 				return false;
 			}
 		}
 		
-		const TArray<TRigVMTypeIndex>& ArgumentTypes = InUnitFunction->GetArgumentTypeIndices();
+		const TArray<TRigVMTypeIndex>& ArgumentTypes = InUnitFunction->GetArgumentTypeIndices_NoLock(InRegistry);
 		for(const TRigVMTypeIndex& ArgumentType : ArgumentTypes)
 		{
-			if(!SupportsType(InController, ArgumentType))
+			if(!SupportsType_NoLock(InController, ArgumentType, InRegistry))
 			{
-				const FString CPPTypeString = Registry->GetType(ArgumentType).CPPType.ToString();
-				static constexpr TCHAR Format[] = TEXT("Unit function '%s' is not supported since type '%s' is not supported.");
-				InController->ReportErrorf(Format, *InUnitFunction->GetName(), *CPPTypeString);
+				if (InController)
+				{
+					const FString CPPTypeString = InRegistry->GetType_NoLock(ArgumentType).CPPType.ToString();
+					static constexpr TCHAR Format[] = TEXT("Unit function '%s' is not supported since type '%s' is not supported.");
+					InController->ReportErrorf(Format, *InUnitFunction->GetName(), *CPPTypeString);
+				}
 				return false;
 			}
 		}
@@ -72,13 +98,20 @@ bool URigVMSchema::SupportsUnitFunction(URigVMController* InController, const FR
 }
 
 bool URigVMSchema::SupportsDispatchFactory(URigVMController* InController, const FRigVMDispatchFactory* InDispatchFactory) const
+{
+	check(Registry);
+	FRigVMRegistryWriteLock Lock(*Registry);
+	return SupportsDispatchFactory_NoLock(InController, InDispatchFactory, Lock);
+}
+
+bool URigVMSchema::SupportsDispatchFactory_NoLock(URigVMController* InController, const FRigVMDispatchFactory* InDispatchFactory, FRigVMRegistryHandle& InRegistry) const
 {
 	if(InDispatchFactory)
 	{
 		if(const UScriptStruct* DispatchExecuteContextStruct = InDispatchFactory->GetExecuteContextStruct())
 		{
-			const TRigVMTypeIndex TypeIndex = Registry->GetTypeIndexFromCPPType(DispatchExecuteContextStruct->GetStructCPPName());
-			if(!SupportsType(InController, TypeIndex))
+			const TRigVMTypeIndex TypeIndex = InRegistry->GetTypeIndexFromCPPType_NoLock(DispatchExecuteContextStruct->GetStructCPPName());
+			if(!SupportsType_NoLock(InController, TypeIndex, InRegistry))
 			{
 				return false;
 			}
@@ -86,7 +119,7 @@ bool URigVMSchema::SupportsDispatchFactory(URigVMController* InController, const
 
 		if(const FRigVMTemplate* Template = InDispatchFactory->GetTemplate())
 		{
-			if(!SupportsTemplate(InController, Template))
+			if(!SupportsTemplate_NoLock(InController, Template, InRegistry))
 			{
 				return false;
 			}
@@ -98,6 +131,13 @@ bool URigVMSchema::SupportsDispatchFactory(URigVMController* InController, const
 }
 
 bool URigVMSchema::SupportsTemplate(URigVMController* InController, const FRigVMTemplate* InTemplate) const
+{
+	check(Registry);
+	FRigVMRegistryWriteLock Lock(*Registry);
+	return SupportsTemplate_NoLock(InController, InTemplate, Lock);
+}
+
+bool URigVMSchema::SupportsTemplate_NoLock(URigVMController* InController, const FRigVMTemplate* InTemplate, FRigVMRegistryHandle& InRegistry) const
 {
 	if(InTemplate)
 	{
@@ -109,11 +149,14 @@ bool URigVMSchema::SupportsTemplate(URigVMController* InController, const FRigVM
 		for(int32 Index = 0; Index < InTemplate->NumExecuteArguments(DispatchContext); Index++)
 		{
 			const TRigVMTypeIndex& TypeIndex = InTemplate->GetExecuteArgument(Index, DispatchContext)->TypeIndex;
-			if(!SupportsType(InController, TypeIndex))
+			if(!SupportsType_NoLock(InController, TypeIndex, InRegistry))
 			{
-				const FString CPPTypeString = Registry->GetType(TypeIndex).CPPType.ToString();
-				static constexpr TCHAR Format[] = TEXT("Template '%s' is not supported since type '%s' is not supported.");
-				InController->ReportErrorf(Format, *InTemplate->GetNotation().ToString(), *CPPTypeString);
+				if (InController)
+				{
+					const FString CPPTypeString = InRegistry->GetType_NoLock(TypeIndex).CPPType.ToString();
+					static constexpr TCHAR Format[] = TEXT("Template '%s' is not supported since type '%s' is not supported.");
+					InController->ReportErrorf(Format, *InTemplate->GetNotation().ToString(), *CPPTypeString);
+				}
 				return false;
 			}
 		}
@@ -121,10 +164,16 @@ bool URigVMSchema::SupportsTemplate(URigVMController* InController, const FRigVM
 	}
 
 	return false;
-
 }
 
 bool URigVMSchema::SupportsGraphFunction(URigVMController* InController, const FRigVMGraphFunctionHeader* InGraphFunction) const
+{
+	check(Registry);
+	FRigVMRegistryWriteLock Lock(*Registry);
+	return SupportsGraphFunction_NoLock(InController, InGraphFunction, Lock);
+}
+
+bool URigVMSchema::SupportsGraphFunction_NoLock(URigVMController* InController, const FRigVMGraphFunctionHeader* InGraphFunction, FRigVMRegistryHandle& InRegistry) const
 {
 	RIGVMSCHEMA_DEFAULT_FUNCTION_BODY
 
@@ -132,32 +181,38 @@ bool URigVMSchema::SupportsGraphFunction(URigVMController* InController, const F
 
 	if(Graph->IsA<URigVMFunctionLibrary>())
 	{
-		static const FString Message = TEXT("You cannot place functions inside of the top level function library graph.");
-		InController->ReportError(Message);
+		if (InController)
+		{
+			static const FString Message = TEXT("You cannot place functions inside of the top level function library graph.");
+			InController->ReportError(Message);
+		}
 		return false;
 	}
 
 	// Make sure all the argument types are supported before creating the node
 	for (const FRigVMGraphFunctionArgument& Argument : InGraphFunction->Arguments)
 	{
-		TRigVMTypeIndex Type = Registry->GetTypeIndexFromCPPType(Argument.CPPType.ToString());
+		TRigVMTypeIndex Type = InRegistry->GetTypeIndexFromCPPType_NoLock(Argument.CPPType.ToString());
 		if (Type == INDEX_NONE)
 		{
 			if (Argument.IsCPPTypeObjectValid())
 			{
 				FRigVMTemplateArgumentType ArgumentType(Argument.CPPType, Argument.CPPTypeObject.Get());
-				Type = Registry->FindOrAddType(ArgumentType);
+				Type = InRegistry->FindOrAddType_NoLock(ArgumentType);
 			}
 		}
 		
 		if (Type == INDEX_NONE)
 		{
-			InController->ReportErrorf(TEXT("Cannot add function reference to %s because argument %s has invalid type %s."), *InGraphFunction->Name.ToString(), *Argument.Name.ToString(), *Argument.CPPType.ToString());
+			if (InController)
+			{
+				InController->ReportErrorf(TEXT("Cannot add function reference to %s because argument %s has invalid type %s."), *InGraphFunction->Name.ToString(), *Argument.Name.ToString(), *Argument.CPPType.ToString());
+			}
 			return false;
 		}
 	}
 
-	if (!InController->bAllowPrivateFunctions)
+	if (!InController || !InController->bAllowPrivateFunctions)
 	{
 		if(IRigVMClientHost* ClientHost = Graph->GetImplementingOuter<IRigVMClientHost>())
 		{
@@ -172,7 +227,10 @@ bool URigVMSchema::SupportsGraphFunction(URigVMController* InController, const F
 			}
 			if (!bIsAvailable)
 			{
-				InController->ReportError(TEXT("Function is not available for placement in another graph host."));
+				if (InController)
+				{
+					InController->ReportError(TEXT("Function is not available for placement in another graph host."));
+				}
 				return false;
 			}
 		}
@@ -184,8 +242,11 @@ bool URigVMSchema::SupportsGraphFunction(URigVMController* InController, const F
 		{
 			if (InGraphFunction->Dependencies.Contains(LibraryNode->GetFunctionIdentifier()))
 			{
-				static const FString Message = TEXT("Function is not available for placement in this graph host due to dependency cycles."); 
-				InController->ReportError(Message);
+				if (InController)
+				{
+					static const FString Message = TEXT("Function is not available for placement in this graph host due to dependency cycles."); 
+					InController->ReportError(Message);
+				}
 				return false;
 			}
 		}
@@ -196,8 +257,11 @@ bool URigVMSchema::SupportsGraphFunction(URigVMController* InController, const F
 	{
 		if (TSoftObjectPtr<UObject>(ParentLibraryNode).ToSoftObjectPath() == InGraphFunction->LibraryPointer.GetNodeSoftPath())
 		{
-			static const FString Message = TEXT("You cannot place functions inside of itself or an indirect recursion.");
-			InController->ReportError(Message);
+			if (InController)
+			{
+				static const FString Message = TEXT("You cannot place functions inside of itself or an indirect recursion.");
+				InController->ReportError(Message);
+			}
 			return false;
 		}
 		ParentLibraryNode = Cast<URigVMLibraryNode>(ParentLibraryNode->GetGraph()->GetOuter());
@@ -207,17 +271,24 @@ bool URigVMSchema::SupportsGraphFunction(URigVMController* InController, const F
 }
 
 bool URigVMSchema::SupportsExternalVariable(URigVMController* InController, const FRigVMExternalVariable* InExternalVariable) const
+{
+	check(Registry);
+	FRigVMRegistryWriteLock Lock(*Registry);
+	return SupportsExternalVariable_NoLock(InController, InExternalVariable, Lock);
+}
+
+bool URigVMSchema::SupportsExternalVariable_NoLock(URigVMController* InController, const FRigVMExternalVariable* InExternalVariable, FRigVMRegistryHandle& InRegistry) const
 {
 	RIGVMSCHEMA_DEFAULT_FUNCTION_BODY
 
 	if(InExternalVariable)
 	{
-		TRigVMTypeIndex Type = Registry->GetTypeIndex(InExternalVariable->GetExtendedCPPType(), InExternalVariable->TypeObject);
+		TRigVMTypeIndex Type = InRegistry->GetTypeIndex_NoLock(InExternalVariable->GetExtendedCPPType(), InExternalVariable->TypeObject);
 		if(Type == INDEX_NONE)
 		{
 			return false;
 		}
-		return SupportsType(InController, Type);
+		return SupportsType_NoLock(InController, Type, InRegistry);
 	}
 
 	return false;
diff --git a/Engine/Plugins/Runtime/RigVM/Source/RigVMDeveloper/Public/RigVMModel/RigVMSchema.h b/Engine/Plugins/Runtime/RigVM/Source/RigVMDeveloper/Public/RigVMModel/RigVMSchema.h
--- a/Engine/Plugins/Runtime/RigVM/Source/RigVMDeveloper/Public/RigVMModel/RigVMSchema.h
+++ b/Engine/Plugins/Runtime/RigVM/Source/RigVMDeveloper/Public/RigVMModel/RigVMSchema.h
@@ -39,23 +39,32 @@ public:
 	// Returns the execute context struct this schema is using
 	UScriptStruct* GetExecuteContextStruct() const { return ExecuteContextStruct; }
 
+	// Returns the corresponding Ed Graph Schema class
+	UE_API virtual UClass* GetEdGraphSchemaClass() const;
+
 	// Returns true if a graph supports a given type
-	UE_API virtual bool SupportsType(URigVMController* InController, TRigVMTypeIndex InTypeIndex) const;
+	UE_API bool SupportsType(URigVMController* InController, TRigVMTypeIndex InTypeIndex) const;
+	UE_API virtual bool SupportsType_NoLock(URigVMController* InController, TRigVMTypeIndex InTypeIndex, FRigVMRegistryHandle& InRegistry) const;
 
 	// Returns true if a graph supports a given unit function
-	UE_API virtual bool SupportsUnitFunction(URigVMController* InController, const FRigVMFunction* InUnitFunction) const;
+	UE_API bool SupportsUnitFunction(URigVMController* InController, const FRigVMFunction* InUnitFunction) const;
+	UE_API virtual bool SupportsUnitFunction_NoLock(URigVMController* InController, const FRigVMFunction* InUnitFunction, FRigVMRegistryHandle& InRegistry) const;
 
 	// Returns true if a graph supports a given dispatch factory
-	UE_API virtual bool SupportsDispatchFactory(URigVMController* InController, const FRigVMDispatchFactory* InDispatchFactory) const;
+	UE_API bool SupportsDispatchFactory(URigVMController* InController, const FRigVMDispatchFactory* InDispatchFactory) const;
+	UE_API virtual bool SupportsDispatchFactory_NoLock(URigVMController* InController, const FRigVMDispatchFactory* InDispatchFactory, FRigVMRegistryHandle& InRegistry) const;
 
 	// Returns true if a graph supports a given template
-	UE_API virtual bool SupportsTemplate(URigVMController* InController, const FRigVMTemplate* InTemplate) const;
+	UE_API bool SupportsTemplate(URigVMController* InController, const FRigVMTemplate* InTemplate) const;
+	UE_API virtual bool SupportsTemplate_NoLock(URigVMController* InController, const FRigVMTemplate* InTemplate, FRigVMRegistryHandle& InRegistry) const;
 
 	// Returns true if a graph supports a given graph based function
-	UE_API virtual bool SupportsGraphFunction(URigVMController* InController, const FRigVMGraphFunctionHeader* InGraphFunction) const;
+	UE_API bool SupportsGraphFunction(URigVMController* InController, const FRigVMGraphFunctionHeader* InGraphFunction) const;
+	UE_API virtual bool SupportsGraphFunction_NoLock(URigVMController* InController, const FRigVMGraphFunctionHeader* InGraphFunction, FRigVMRegistryHandle& InRegistry) const;
 
 	// Returns true if a graph supports a given external variable
-	UE_API virtual bool SupportsExternalVariable(URigVMController* InController, const FRigVMExternalVariable* InExternalVariable) const;
+	UE_API bool SupportsExternalVariable(URigVMController* InController, const FRigVMExternalVariable* InExternalVariable) const;
+	UE_API virtual bool SupportsExternalVariable_NoLock(URigVMController* InController, const FRigVMExternalVariable* InExternalVariable, FRigVMRegistryHandle& InRegistry) const;
 
 	// Returns true if a graph supports pin categories / node layouts
 	UE_API virtual bool SupportsNodeLayouts(const URigVMGraph* InGraph) const;
