RigVM: Follow up for static analysis fix for 47954234

#rnx

diff --git a/Engine/Plugins/Runtime/RigVM/Source/RigVMDeveloper/Private/RigVMModel/RigVMSchema.cpp b/Engine/Plugins/Runtime/RigVM/Source/RigVMDeveloper/Private/RigVMModel/RigVMSchema.cpp
--- a/Engine/Plugins/Runtime/RigVM/Source/RigVMDeveloper/Private/RigVMModel/RigVMSchema.cpp
+++ b/Engine/Plugins/Runtime/RigVM/Source/RigVMDeveloper/Private/RigVMModel/RigVMSchema.cpp
@@ -212,7 +212,13 @@ bool URigVMSchema::SupportsGraphFunction_NoLock(URigVMController* InController,
 		}
 	}
 
-	if (!InController || !InController->bAllowPrivateFunctions)
+	bool bAllowPrivateFunctions = true;
+	if (InController)
+	{
+		bAllowPrivateFunctions = InController->bAllowPrivateFunctions;
+	}
+
+	if (!bAllowPrivateFunctions)
 	{
 		if(IRigVMClientHost* ClientHost = Graph->GetImplementingOuter<IRigVMClientHost>())
 		{
