RigVM: Array is empty dispatch

#rnx

diff --git a/Engine/Plugins/Runtime/RigVM/Source/RigVM/Private/RigVMFunctions/RigVMDispatch_Array.cpp b/Engine/Plugins/Runtime/RigVM/Source/RigVM/Private/RigVMFunctions/RigVMDispatch_Array.cpp
--- a/Engine/Plugins/Runtime/RigVM/Source/RigVM/Private/RigVMFunctions/RigVMDispatch_Array.cpp
+++ b/Engine/Plugins/Runtime/RigVM/Source/RigVM/Private/RigVMFunctions/RigVMDispatch_Array.cpp
@@ -137,7 +137,7 @@ FRigVMTemplateArgumentInfo FRigVMDispatch_ArrayBase::CreateArgumentInfo(const FN
 	{
 		return FRigVMTemplateArgumentInfo(InName, InDirection, RigVMTypeUtils::TypeIndex::Int32);
 	}
-	if(InName == SuccessName)
+	if(InName == SuccessName || InName == IsEmptyName)
 	{
 		return FRigVMTemplateArgumentInfo(InName, InDirection, RigVMTypeUtils::TypeIndex::Bool);
 	}
@@ -473,6 +473,67 @@ FText FRigVMDispatch_ArraySetNum::GetArgumentTooltip(const FName& InArgumentName
 
 ////////////////////////////////////////////////////////////////////////////////////////////////////////////////////////
 
+FName FRigVMDispatch_ArrayIsEmpty::GetArgumentNameForOperandIndex(int32 InOperandIndex, int32 InTotalOperands, FRigVMRegistryHandle& InRegistry) const
+{
+	static const FLazyName ArgumentNames[] = {
+		ArrayName,
+		IsEmptyName
+	};
+	check(InTotalOperands == UE_ARRAY_COUNT(ArgumentNames));
+	return ArgumentNames[InOperandIndex];
+}
+
+const TArray<FRigVMTemplateArgumentInfo>& FRigVMDispatch_ArrayIsEmpty::GetArgumentInfos(FRigVMRegistryHandle& InRegistry) const
+{
+	if(CachedArgumentInfos.IsEmpty())
+	{
+		CachedArgumentInfos.Emplace( CreateArgumentInfo(ArrayName, ERigVMPinDirection::Input) );
+		CachedArgumentInfos.Emplace( CreateArgumentInfo(IsEmptyName, ERigVMPinDirection::Output) );
+	}
+	return CachedArgumentInfos;
+}
+
+FRigVMTemplateTypeMap FRigVMDispatch_ArrayIsEmpty::OnNewArgumentType(const FName& InArgumentName, TRigVMTypeIndex InTypeIndex, FRigVMRegistryHandle& InRegistry) const
+{
+	check(InArgumentName == ArrayName);
+	return {
+		{ArrayName, InTypeIndex},
+		{IsEmptyName, RigVMTypeUtils::TypeIndex::Bool}
+	};
+}
+
+void FRigVMDispatch_ArrayIsEmpty::Execute(FRigVMExtendedExecuteContext& InContext, FRigVMMemoryHandleArray Handles, FRigVMPredicateBranchArray Predicates)
+{
+	const FArrayProperty* ArrayProperty = CastFieldChecked<FArrayProperty>(Handles[0].GetResolvedProperty());
+	const FScriptArrayHelper ArrayHelper(ArrayProperty, Handles[0].GetInputData());
+	bool* IsEmpty = (bool*)Handles[1].GetOutputData();
+	*IsEmpty = ArrayHelper.Num() == 0;
+}
+
+#if WITH_EDITOR
+
+FText FRigVMDispatch_ArrayIsEmpty::GetNodeTooltip(const FRigVMTemplateTypeMap& InTypes) const
+{
+    return LOCTEXT("ArrayIsEmptyToolTip", "Returns true if the array is empty");
+}
+
+FText FRigVMDispatch_ArrayIsEmpty::GetArgumentTooltip(const FName& InArgumentName, TRigVMTypeIndex InTypeIndex) const
+{
+	if(InArgumentName == ArrayName)
+	{
+		return LOCTEXT("ArrayIsEmpty_ArrayArgumentToolTip", "The array to check for emptiness.");
+	}
+	if(InArgumentName == IsEmptyName)
+	{
+		return LOCTEXT("ArrayIsEmpty_NumArgumentToolTip", "True if the array is empty.");
+	}
+	return FText();
+}
+
+#endif
+
+////////////////////////////////////////////////////////////////////////////////////////////////////////////////////////
+
 FName FRigVMDispatch_ArrayInit::GetArgumentNameForOperandIndex(int32 InOperandIndex, int32 InTotalOperands, FRigVMRegistryHandle& InRegistry) const
 {
 	static const FLazyName ArgumentNames[] = {
diff --git a/Engine/Plugins/Runtime/RigVM/Source/RigVM/Public/RigVMFunctions/RigVMDispatch_Array.h b/Engine/Plugins/Runtime/RigVM/Source/RigVM/Public/RigVMFunctions/RigVMDispatch_Array.h
--- a/Engine/Plugins/Runtime/RigVM/Source/RigVM/Public/RigVMFunctions/RigVMDispatch_Array.h
+++ b/Engine/Plugins/Runtime/RigVM/Source/RigVM/Public/RigVMFunctions/RigVMDispatch_Array.h
@@ -28,6 +28,7 @@ protected:
 	static inline FLazyName ArrayName = FLazyName(TEXT("Array"));
 	static inline FLazyName ValuesName = FLazyName(TEXT("Values"));
 	static inline FLazyName NumName = FLazyName(TEXT("Num"));
+	static inline FLazyName IsEmptyName = FLazyName(TEXT("IsEmpty"));
 	static inline FLazyName IndexName = FLazyName(TEXT("Index"));
 	static inline FLazyName ElementName = FLazyName(TEXT("Element"));
 	static inline FLazyName SuccessName = FLazyName(TEXT("Success"));
@@ -150,6 +151,30 @@ protected:
 	static RIGVM_API void Execute(FRigVMExtendedExecuteContext& InContext, FRigVMMemoryHandleArray Handles, FRigVMPredicateBranchArray Predicates);
 };
 
+USTRUCT(meta=(DisplayName = "IsEmpty"))
+struct FRigVMDispatch_ArrayIsEmpty : public FRigVMDispatch_ArrayBase
+{
+	GENERATED_BODY()
+
+public:
+	FRigVMDispatch_ArrayIsEmpty()
+	{
+		FactoryScriptStruct = StaticStruct();
+	}
+
+	RIGVM_API virtual FName GetArgumentNameForOperandIndex(int32 InOperandIndex, int32 InTotalOperands, FRigVMRegistryHandle& InRegistry) const override;
+	RIGVM_API virtual const TArray<FRigVMTemplateArgumentInfo>& GetArgumentInfos(FRigVMRegistryHandle& InRegistry) const override;
+	RIGVM_API virtual FRigVMTemplateTypeMap OnNewArgumentType(const FName& InArgumentName, TRigVMTypeIndex InTypeIndex, FRigVMRegistryHandle& InRegistry) const override;
+#if WITH_EDITOR
+	RIGVM_API virtual FText GetNodeTooltip(const FRigVMTemplateTypeMap& InTypes) const override;
+	RIGVM_API virtual FText GetArgumentTooltip(const FName& InArgumentName, TRigVMTypeIndex InTypeIndex) const override;
+#endif
+
+protected:
+	virtual FRigVMFunctionPtr GetDispatchFunctionImpl(const FRigVMTemplateTypeMap& InTypes, const FRigVMRegistryHandle& InRegistry) const override { return &FRigVMDispatch_ArrayIsEmpty::Execute; }
+	static RIGVM_API void Execute(FRigVMExtendedExecuteContext& InContext, FRigVMMemoryHandleArray Handles, FRigVMPredicateBranchArray Predicates);
+};
+
 USTRUCT(meta=(DisplayName = "Init"))
 struct FRigVMDispatch_ArrayInit : public FRigVMDispatch_ArrayBaseMutable
 {
