RigVM: Improve json export with more options

#rnx

diff --git a/Engine/Plugins/Runtime/RigVM/Source/RigVM/Private/RigVMCore/RigVMRegistry.cpp b/Engine/Plugins/Runtime/RigVM/Source/RigVM/Private/RigVMCore/RigVMRegistry.cpp
--- a/Engine/Plugins/Runtime/RigVM/Source/RigVM/Private/RigVMCore/RigVMRegistry.cpp
+++ b/Engine/Plugins/Runtime/RigVM/Source/RigVM/Private/RigVMCore/RigVMRegistry.cpp
@@ -2715,7 +2715,7 @@ const TChunkedArray<FRigVMFunction>& FRigVMRegistry_NoLock::GetFunctions_NoLock(
 	return Functions;
 }
 
-const FRigVMTemplate* FRigVMRegistry_NoLock::FindTemplate_NoLock(const FName& InNotation, bool bIncludeDeprecated) const
+const FRigVMTemplate* FRigVMRegistry_NoLock::FindTemplate_NoLock(const FName& InNotation, bool bIncludeDeprecated, bool bRemapToCoreDispatch) const
 {
 	if (InNotation.IsNone())
 	{
@@ -2751,9 +2751,12 @@ const FRigVMTemplate* FRigVMRegistry_NoLock::FindTemplate_NoLock(const FName& In
 			},
 		};
 
-		if(const FString* RemappedDispatch = CoreDispatchMap.Find(FactoryName))
+		if (bRemapToCoreDispatch)
 		{
-			FactoryName = *RemappedDispatch;
+			if(const FString* RemappedDispatch = CoreDispatchMap.Find(FactoryName))
+			{
+				FactoryName = *RemappedDispatch;
+			}
 		}
 		
 		if(const FRigVMDispatchFactory* Factory = FindDispatchFactory_NoLock(*FactoryName))
diff --git a/Engine/Plugins/Runtime/RigVM/Source/RigVM/Private/RigVMCore/RigVMTemplate.cpp b/Engine/Plugins/Runtime/RigVM/Source/RigVM/Private/RigVMCore/RigVMTemplate.cpp
--- a/Engine/Plugins/Runtime/RigVM/Source/RigVM/Private/RigVMCore/RigVMTemplate.cpp
+++ b/Engine/Plugins/Runtime/RigVM/Source/RigVM/Private/RigVMCore/RigVMTemplate.cpp
@@ -1710,7 +1710,7 @@ const FRigVMFunction* FRigVMTemplate::GetPermutation_NoLock(int32 InIndex, const
 	if(InRegistry->GetFunctions_NoLock().IsValidIndex(FunctionIndex))
 	{
 		const FRigVMFunction* Function = &InRegistry->GetFunctions_NoLock()[Permutations[InIndex]];
-		if (Function)
+		if (Function && Function->TemplateIndex != INDEX_NONE)
 		{
 			check(Function->PermutationIndex == InIndex);
 		}
diff --git a/Engine/Plugins/Runtime/RigVM/Source/RigVM/Public/RigVMCore/RigVMRegistry.h b/Engine/Plugins/Runtime/RigVM/Source/RigVM/Public/RigVMCore/RigVMRegistry.h
--- a/Engine/Plugins/Runtime/RigVM/Source/RigVM/Public/RigVMCore/RigVMRegistry.h
+++ b/Engine/Plugins/Runtime/RigVM/Source/RigVM/Public/RigVMCore/RigVMRegistry.h
@@ -303,7 +303,7 @@ public:
 	UE_API virtual const TChunkedArray<FRigVMFunction>& GetFunctions_NoLock() const;
 
 	// Returns a template pointer given its notation (or nullptr)
-	UE_API virtual const FRigVMTemplate* FindTemplate_NoLock(const FName& InNotation, bool bIncludeDeprecated = false) const;
+	UE_API virtual const FRigVMTemplate* FindTemplate_NoLock(const FName& InNotation, bool bIncludeDeprecated = false, bool bRemapToCoreDispatch = true) const;
 
 	// Returns all current RigVM functions
 	UE_API virtual const TChunkedArray<FRigVMTemplate>& GetTemplates_NoLock() const;
@@ -1004,10 +1004,10 @@ public:
 	}
 
 	// Returns a template pointer given its notation (or nullptr)
-	const FRigVMTemplate* FindTemplate(const FName& InNotation, bool bIncludeDeprecated = false) const
+	const FRigVMTemplate* FindTemplate(const FName& InNotation, bool bIncludeDeprecated = false, bool bRemapToCoreDispatch = true) const
 	{
 		FConditionalWriteScopeLock _(*this);
-		return Super::FindTemplate_NoLock(InNotation, bIncludeDeprecated);
+		return Super::FindTemplate_NoLock(InNotation, bIncludeDeprecated, bRemapToCoreDispatch);
 	}
 
 	// Returns all current RigVM functions
@@ -1434,9 +1434,9 @@ protected:
 
 	// Returns a template pointer given its notation (or nullptr)
 	// Note: Only call NoLock methods on the FRigVMRegistry_NoLock
-	virtual const FRigVMTemplate* FindTemplate_NoLock(const FName& InNotation, bool bIncludeDeprecated = false) const override
+	virtual const FRigVMTemplate* FindTemplate_NoLock(const FName& InNotation, bool bIncludeDeprecated = false, bool bRemapToCoreDispatch = true) const override
 	{
-		return Super::FindTemplate_NoLock(InNotation, bIncludeDeprecated);
+		return Super::FindTemplate_NoLock(InNotation, bIncludeDeprecated, bRemapToCoreDispatch);
 	}
 
 	// Returns all current RigVM functions
diff --git a/Engine/Plugins/Runtime/RigVM/Source/RigVM/Public/RigVMCore/RigVMTemplate.h b/Engine/Plugins/Runtime/RigVM/Source/RigVM/Public/RigVMCore/RigVMTemplate.h
--- a/Engine/Plugins/Runtime/RigVM/Source/RigVM/Public/RigVMCore/RigVMTemplate.h
+++ b/Engine/Plugins/Runtime/RigVM/Source/RigVM/Public/RigVMCore/RigVMTemplate.h
@@ -605,14 +605,14 @@ public:
 	RIGVM_API TArray<FRigVMTemplateArgumentInfo> GetFlattenedArgumentInfos_NoLock(FRigVMRegistryHandle& InRegistry) const;
 	RIGVM_API TArray<FRigVMTemplateArgumentInfo> GetFlattenedArgumentInfosForFunctions_NoLock(const TArray<const FRigVMFunction*>& InFunctions, FRigVMRegistryHandle& InRegistry) const;
 
-private:
-
 	// Constructor from a struct, a template name and a function index
-	FRigVMTemplate(UScriptStruct* InStruct, const FString& InTemplateName, int32 InFunctionIndex, FRigVMRegistryHandle& InRegistry);
+	RIGVM_API FRigVMTemplate(UScriptStruct* InStruct, const FString& InTemplateName, int32 InFunctionIndex, FRigVMRegistryHandle& InRegistry);
 
 	// Constructor from a template name and argument infos
-	FRigVMTemplate(const FName& InTemplateName, const TArray<FRigVMTemplateArgumentInfo>& InInfos, const FRigVMRegistryHandle& InRegistry);
+	RIGVM_API FRigVMTemplate(const FName& InTemplateName, const TArray<FRigVMTemplateArgumentInfo>& InInfos, const FRigVMRegistryHandle& InRegistry);
 
+private:
+	
 	static FLinearColor GetColorFromMetadata(FString InMetadata);
 
 	void UpdateTypesIndicesToPermutationIndices(int32 InPermutationIndex, const FRigVMRegistryHandle& InRegistry);
diff --git a/Engine/Plugins/Runtime/RigVM/Source/RigVMDeveloper/Private/RigVMJsonUtils.cpp b/Engine/Plugins/Runtime/RigVM/Source/RigVMDeveloper/Private/RigVMJsonUtils.cpp
--- a/Engine/Plugins/Runtime/RigVM/Source/RigVMDeveloper/Private/RigVMJsonUtils.cpp
+++ b/Engine/Plugins/Runtime/RigVM/Source/RigVMDeveloper/Private/RigVMJsonUtils.cpp
@@ -5,7 +5,9 @@
 #include "RigVMDeveloperModule.h"
 #include "RigVMCore/RigVMStruct.h"
 #include "RigVMAsset.h"
+#include "EdGraph/RigVMEdGraphSchema.h"
 #include "RigVMModel/RigVMGraph.h"
+#include "RigVMModel/RigVMSchema.h"
 #include "RigVMModel/Nodes/RigVMUnitNode.h"
 #include "RigVMModel/Nodes/RigVMDispatchNode.h"
 #include "RigVMModel/Nodes/RigVMVariableNode.h"
@@ -14,6 +16,7 @@
 #include "RigVMModel/Nodes/RigVMCommentNode.h"
 #include "RigVMModel/RigVMVariableDescription.h"
 #include "Serialization/JsonSerializer.h"
+#include "UObject/UObjectIterator.h"
 
 #define LOCTEXT_NAMESPACE "RigVMJsonUtils"
 
@@ -29,7 +32,7 @@ FAutoConsoleCommand FCmdRigVMJsonPrintAvailableNodes
 	{
 		FRigVMRegistryReadLock Registry;
 		RigVMJsonUtils::FIntrospectionSettings Settings(Registry);
-
+		
 		for (const FString& Arg : Args)
 		{
 			if (Arg.Equals(TEXT("exclude-deprecated"), ESearchCase::IgnoreCase))
@@ -44,13 +47,93 @@ FAutoConsoleCommand FCmdRigVMJsonPrintAvailableNodes
 			{
 				Settings.bExcludeDispatches = true;
 			}
+			else if (Arg.Equals(TEXT("merge-templates"), ESearchCase::IgnoreCase))
+			{
+				Settings.bMergeTemplates = true;
+			}
+			else if(Arg.Len() > 8 && Arg.StartsWith(TEXT("context="), ESearchCase::IgnoreCase))
+			{
+				const FString ArgContextName = Arg.Mid(8);
+				UClass* ArgSchemaClass = nullptr;
+				for (TObjectIterator<UClass> ClassIterator; ClassIterator; ++ClassIterator)
+				{
+					if (*ClassIterator == URigVMSchema::StaticClass())
+					{
+						continue;
+					}
+					if (ClassIterator->IsChildOf(URigVMSchema::StaticClass()))
+					{
+						const FString ClassName = ClassIterator->GetName();
+						const FString ContextName = ClassName.Replace(TEXT("Schema"), TEXT(""));
+						if (ArgContextName.Equals(ContextName, ESearchCase::IgnoreCase))
+						{
+							ArgSchemaClass = *ClassIterator;
+							break;
+						}
+					}
+				}
+
+				if (ArgSchemaClass == nullptr)
+				{
+					UE_LOG(LogRigVMDeveloper, Error, TEXT("Cannot find RigVMSchema class for context '%s'"), *ArgContextName);
+					return;
+				}
+
+				Settings.Schema = ArgSchemaClass->GetDefaultObject<URigVMSchema>();
+				check(Settings.Schema->GetEdGraphSchemaClass());
+				Settings.EdGraphSchema = Settings.Schema->GetEdGraphSchemaClass()->GetDefaultObject<URigVMEdGraphSchema>();
+			}
 			else
 			{
-				UE_LOG(LogRigVMDeveloper, Display, TEXT("Valid Arguments:\n"
+				FString ContextNamesJoined;
+				TArray<FString> ContextNames;
+				
+				for (TObjectIterator<UClass> ClassIterator; ClassIterator; ++ClassIterator)
+				{
+					if (*ClassIterator == URigVMSchema::StaticClass())
+					{
+						continue;
+					}
+					if (!ClassIterator->IsChildOf(URigVMSchema::StaticClass()))
+					{
+						continue;
+					}
+					const FString ClassName = ClassIterator->GetName();
+					const FString ContextName = ClassName.Replace(TEXT("Schema"), TEXT(""));
+
+					static const TArray<FString> StringsToFilter = {
+						TEXT("EngineTest"),
+						TEXT("TestPlugin"),
+						TEXT("RigVMExample"),
+					};
+
+					bool bValidContextName = true;
+					for (const FString& StringToFilter : StringsToFilter)
+					{
+						if (ContextName.Contains(StringToFilter, ESearchCase::IgnoreCase))
+						{
+							bValidContextName = false;
+							break;
+						}
+					}
+					if (!bValidContextName)
+					{
+						continue;
+					}
+					ContextNames.Add(ContextName);
+				}
+
+				if (!ContextNames.IsEmpty())
+				{
+					ContextNamesJoined = RigVMStringUtils::JoinStrings(ContextNames, TEXT(", "));
+				}
+
+				UE_LOG(LogRigVMDeveloper, Warning, TEXT("Valid Arguments:\n"
+					"context=MyContextName - filter the nodes by context in question (%s)\n"
 					"exclude-deprecated - only includes the current / non-deprecated nodes\n"
 					"exclude-functions - excludes RigVM struct based functions\n"
 					"exclude-dispatches - excludes RigVM dispatch factories"
-				));
+				), *ContextNamesJoined);
 				return;
 			}
 		}
@@ -64,7 +147,6 @@ FAutoConsoleCommand FCmdRigVMJsonPrintAvailableNodes
 				TEXT("EngineTest"),
 				TEXT("TestPlugin"),
 				TEXT("RigVMExample"),
-				TEXT("UAF"),
 			};
 			for (const FString& StringToFilter : StringsToFilter)
 			{
@@ -82,9 +164,8 @@ FAutoConsoleCommand FCmdRigVMJsonPrintAvailableNodes
 			RigVMJsonUtils::FJsonWriterRef Writer = RigVMJsonUtils::MakeJsonWriter(Text);
 			if (FJsonSerializer::Serialize(JsonObject, Writer))
 			{
-				UE_LOG(LogRigVMDeveloper, Display, TEXT("Json:\n\n%s"), *Text);
 #if WITH_EDITOR
-				UE_LOG(LogRigVMDeveloper, Display, TEXT("The content has also been copied to the clipboard."));
+				UE_LOG(LogRigVMDeveloper, Display, TEXT("The content has been copied to the clipboard."));
 				FPlatformApplicationMisc::ClipboardCopy(*Text);
 #endif
 			}
@@ -120,9 +201,8 @@ FAutoConsoleCommand FCmdRigVMJsonPrintGraph
 			RigVMJsonUtils::FJsonWriterRef Writer = RigVMJsonUtils::MakeJsonWriter(Text);
 			if (FJsonSerializer::Serialize(JsonObject, Writer))
 			{
-				UE_LOG(LogRigVMDeveloper, Display, TEXT("Json:\n\n%s"), *Text);
 #if WITH_EDITOR
-				UE_LOG(LogRigVMDeveloper, Display, TEXT("The content has also been copied to the clipboard."));
+				UE_LOG(LogRigVMDeveloper, Display, TEXT("The content has been copied to the clipboard."));
 				FPlatformApplicationMisc::ClipboardCopy(*Text);
 #endif
 			}
@@ -165,7 +245,10 @@ FAutoConsoleCommand FCmdRigVMJsonPrintAsset
 			RigVMJsonUtils::FJsonWriterRef Writer = RigVMJsonUtils::MakeJsonWriter(Text);
 			if (FJsonSerializer::Serialize(JsonObject, Writer))
 			{
-				UE_LOG(LogRigVMDeveloper, Display, TEXT("Json:\n\n%s"), *Text);
+#if WITH_EDITOR
+				UE_LOG(LogRigVMDeveloper, Display, TEXT("The content has been copied to the clipboard."));
+				FPlatformApplicationMisc::ClipboardCopy(*Text);
+#endif
 			}
 		}
 	})
@@ -175,6 +258,10 @@ namespace RigVMJsonUtils
 {
 	static FString PostProcessToolTip(const FString& InToolTip)
 	{
+		if (InToolTip.StartsWith(TEXT("#define")))
+		{
+			return FString();
+		}
 		FString Result = InToolTip;
 		Result.ReplaceInline(TEXT("\\n* "), TEXT("\\n"));
 		Result.RemoveFromStart(TEXT("* "));
@@ -183,6 +270,38 @@ namespace RigVMJsonUtils
 	}
 }
 
+RigVMJsonUtils::FIntrospectionSettings::FIntrospectionSettings(FRigVMRegistryHandle& InRegistry)
+: Registry(InRegistry)
+, Schema(URigVMSchema::StaticClass()->GetDefaultObject<URigVMSchema>())
+, EdGraphSchema(URigVMEdGraphSchema::StaticClass()->GetDefaultObject<URigVMEdGraphSchema>())
+{
+}
+
+void RigVMJsonUtils::FIntrospectionSettings::AddDiscoveredType(const TRigVMTypeIndex& InTypeIndex)
+{
+	if (InTypeIndex == INDEX_NONE)
+	{
+		return;
+	}
+	if (DiscoveredTypes.Contains(InTypeIndex))
+	{
+		return;
+	}
+	if (Registry->IsArrayType_NoLock(InTypeIndex))
+	{
+		AddDiscoveredType(Registry->GetBaseTypeFromArrayTypeIndex_NoLock(InTypeIndex));
+		return;
+	}
+	if (Schema)
+	{
+		if (!Schema->SupportsType_NoLock(nullptr, InTypeIndex, Registry))
+		{
+			return;
+		}
+	}
+	DiscoveredTypes.Add(InTypeIndex);
+}
+
 RigVMJsonUtils::FJsonWriterRef RigVMJsonUtils::MakeJsonWriter(FString& OutJsonText)
 {
 	return FJsonFactory::Create(&OutJsonText);
@@ -194,8 +313,9 @@ TSharedPtr<FJsonObject> RigVMJsonUtils::AvailableNodesToJson(FIntrospectionSetti
 	{
 		return nullptr;
 	}
-	
-	TArray<TSharedPtr<FJsonValue>> FunctionJsonValues, FactoryJsonValues;
+
+	InSettings.DiscoveredTypes.Reset();
+	TArray<TSharedPtr<FJsonValue>> TypeValues, FunctionJsonValues, FactoryJsonValues;
 	const TChunkedArray<FRigVMFunction>& Functions = InSettings.Registry->GetFunctions_NoLock();
 	const TArray<FRigVMDispatchFactory*>& Factories = InSettings.Registry->GetFactories_NoLock();
 
@@ -229,7 +349,18 @@ TSharedPtr<FJsonObject> RigVMJsonUtils::AvailableNodesToJson(FIntrospectionSetti
 		}
 	}
 
+	TArray<TRigVMTypeIndex> DiscoveredTypes = InSettings.DiscoveredTypes.Array();
+	DiscoveredTypes.Sort();
+	for (const TRigVMTypeIndex& Type : DiscoveredTypes)
+	{
+		if (const TSharedPtr<FJsonObject> JsonObject = ToJson(InSettings, Type))
+		{
+			TypeValues.Add(MakeShared<FJsonValueObject>(JsonObject));
+		}
+	}
+
 	TSharedPtr<FJsonObject> JsonObject = MakeShared<FJsonObject>();
+	JsonObject->SetArrayField(TEXT("Types"), TypeValues);
 	JsonObject->SetArrayField(TEXT("Functions"), FunctionJsonValues);
 	JsonObject->SetArrayField(TEXT("Dispatches"), FactoryJsonValues);
 	return JsonObject;
@@ -241,7 +372,68 @@ TSharedPtr<FJsonObject> RigVMJsonUtils::ToJson(FIntrospectionSettings& InSetting
 	{
 		return nullptr;
 	}
+
+	check(InFunction);
+	check(InFunction->Struct);
+
+	TArray<const FRigVMFunction*> Permutations = {InFunction};
+	
+	const FRigVMTemplate* Template = InFunction->GetTemplate_NoLock(InSettings.Registry);
+
+#if WITH_EDITOR
+	// this may be a deprecated template
+	if (!Template && InSettings.bMergeTemplates && InFunction->Struct->HasMetaDataHierarchical(FRigVMStruct::DeprecatedMetaName))
+	{
+		FString TemplateMetadata;
+		if (InFunction->Struct->GetStringMetaDataHierarchical(FRigVMRegistry::TemplateNameMetaName, &TemplateMetadata))
+		{
+			const FString TemplateName = RigVMStringUtils::JoinStrings(TemplateMetadata, InFunction->GetMethodName().ToString(), TEXT("::"));
+			const FName Notation = FRigVMTemplate(InFunction->Struct, TemplateName, InFunction->Index, InSettings.Registry).GetNotation();
+			if (!Notation.IsNone())
+			{
+				Template = InSettings.Registry->FindTemplate_NoLock(Notation, true /* include deprecated */, false /* remap */);
+			}
+		}
+	}
+#endif
 	
+	if (Template && InSettings.bMergeTemplates)
+	{
+		if (Template->GetPermutation_NoLock(0, InSettings.Registry) != InFunction)
+		{
+			return nullptr;
+		}
+
+		Permutations.Reset();
+		const int32 NumPermutations = Template->NumPermutations_NoLock(InSettings.Registry);
+		for (int32 PermutationIndex = 0; PermutationIndex < NumPermutations; PermutationIndex++)
+		{
+			if (const FRigVMFunction* Permutation = Template->GetPermutation_NoLock(PermutationIndex, InSettings.Registry))
+			{
+				if (!InSettings.Schema->SupportsUnitFunction_NoLock(nullptr, Permutation, InSettings.Registry))
+				{
+					continue;
+				}
+				Permutations.Add(Permutation);
+			}
+		}
+		if (Permutations.IsEmpty())
+		{
+			return nullptr;
+		}
+	}
+	else if (InSettings.Schema)
+	{
+		if (!InSettings.Schema->SupportsUnitFunction_NoLock(nullptr, InFunction, InSettings.Registry))
+		{
+			return nullptr;
+		}
+	}
+
+	// rely on the first valid permutation
+	// (which by default is the passed in function)
+	InFunction = Permutations[0];
+
 	check(InFunction);
 	const UScriptStruct* Struct = InFunction->Struct;
 	check(Struct);
@@ -253,17 +445,21 @@ TSharedPtr<FJsonObject> RigVMJsonUtils::ToJson(FIntrospectionSettings& InSetting
 	}
 
 	FString TemplateName;
-	if (const FRigVMTemplate* Template = InFunction->GetTemplate_NoLock(InSettings.Registry))
+	if (Template)
 	{
 		TemplateName = Template->GetName().ToString();
+		if (Permutations.Num() > 1)
+		{
+			JsonObject->RemoveField(TEXT("Type"));
+		}
 	}
 	JsonObject->SetStringField(TEXT("TemplateName"), TemplateName);
 
-	TArray<TSharedPtr<FJsonValue>> IO, Inputs, Outputs;
-
 	FStructOnScope StructOnScope(Struct);
 	FRigVMStruct* StructInstance = reinterpret_cast<FRigVMStruct*>(StructOnScope.GetStructMemory());
 
+	TArray<TSharedPtr<FJsonValue>> IO, Inputs, Outputs;
+
 	const FRigVMStructUpgradeInfo UpgradeInfo = StructInstance->GetUpgradeInfo();
 	if (UpgradeInfo.IsValid())
 	{
@@ -273,10 +469,10 @@ TSharedPtr<FJsonObject> RigVMJsonUtils::ToJson(FIntrospectionSettings& InSetting
 		}
 	}
 
-	TMap<FName,const FRigVMFunctionArgument*> ArgumentLookup;
+	TMap<FName,int32> ArgumentLookup;
 	for (int32 Index = 0; Index < InFunction->Arguments.Num(); ++Index)
 	{
-		ArgumentLookup.Add(InFunction->Arguments[Index].Name, &InFunction->Arguments[Index]);
+		ArgumentLookup.Add(InFunction->Arguments[Index].Name, Index);
 	}
 	
 	for (TFieldIterator<FProperty> PropertyIt(Struct); PropertyIt; ++PropertyIt)
@@ -287,10 +483,20 @@ TSharedPtr<FJsonObject> RigVMJsonUtils::ToJson(FIntrospectionSettings& InSetting
 		{
 			bIsExecutePin = StructProperty->Struct->IsChildOf(FRigVMExecutePin::StaticStruct());
 		}
+
 		static const FString ExecutePinType = FRigVMExecuteContext::StaticStruct()->GetStructCPPName();
-		
-		const FRigVMFunctionArgument** ArgumentPtr = ArgumentLookup.Find(Property->GetFName());
-		const FRigVMFunctionArgument* Argument = ArgumentPtr ? *ArgumentPtr : nullptr;
+		if (bIsExecutePin)
+		{
+			InSettings.AddDiscoveredType(InSettings.Registry->GetTypeIndex_NoLock<FRigVMExecuteContext>());
+		}
+
+		int32 ArgumentIndex = INDEX_NONE;
+		const FRigVMFunctionArgument* Argument = nullptr;
+		if (const int32* ArgumentIndexPtr = ArgumentLookup.Find(Property->GetFName()))
+		{
+			ArgumentIndex = *ArgumentIndexPtr;
+			Argument = &InFunction->Arguments[ArgumentIndex];
+		}
 
 		if (!Argument && !bIsExecutePin)
 		{
@@ -325,12 +531,48 @@ TSharedPtr<FJsonObject> RigVMJsonUtils::ToJson(FIntrospectionSettings& InSetting
 				{
 					DefaultValue = TEXT("()");
 				}
+				else if (Property->IsA<FNameProperty>())
+				{
+					DefaultValue = TEXT("None");
+				}
 			}
+			InSettings.AddDiscoveredType(InSettings.Registry->GetTypeIndexFromCPPType_NoLock(Argument->Type));
 		}
 
 		TSharedPtr<FJsonObject> ArgumentJsonObject = MakeShared<FJsonObject>();
 		ArgumentJsonObject->SetStringField(TEXT("Name"), Property->GetName());
-		ArgumentJsonObject->SetStringField(TEXT("Type"), Argument ? FString(Argument->Type) : ExecutePinType);
+
+		if (Argument)
+		{
+			if (Permutations.Num() == 1)
+			{
+				ArgumentJsonObject->SetStringField(TEXT("Type"), Argument->Type);
+			}
+			else
+			{
+				TArray<FString> Types;
+				for (const FRigVMFunction* Permutation : Permutations)
+				{
+					Types.Add(Permutation->Arguments[ArgumentIndex].Type);
+				}
+				Types.Sort([&InSettings](const FString& A, const FString& B) -> bool
+				{
+					const TRigVMTypeIndex TypeIndexA = InSettings.Registry->GetTypeIndexFromCPPType_NoLock(A);
+					const TRigVMTypeIndex TypeIndexB = InSettings.Registry->GetTypeIndexFromCPPType_NoLock(B);
+					if (TypeIndexA != INDEX_NONE && TypeIndexB != INDEX_NONE)
+					{
+						return TypeIndexA < TypeIndexB;
+					}
+					return A.Compare(B, ESearchCase::IgnoreCase) < 0;
+				});
+				ArgumentJsonObject->SetStringField(TEXT("Type"), RigVMStringUtils::JoinStrings(Types, TEXT(",")));
+			}
+		}
+		else
+		{
+			ArgumentJsonObject->SetStringField(TEXT("Type"), ExecutePinType);
+		}
+
 		FString PropertyToolTip = PostProcessToolTip(Property->GetToolTipText().ToString());
 		if (PropertyToolTip.IsEmpty() && bIsExecutePin)
 		{
@@ -375,6 +617,14 @@ TSharedPtr<FJsonObject> RigVMJsonUtils::ToJson(FIntrospectionSettings& InSetting
 	const UScriptStruct* Struct = InFactory->GetScriptStruct();
 	check(Struct);
 
+	if (InSettings.Schema)
+	{
+		if (!InSettings.Schema->SupportsDispatchFactory_NoLock(nullptr, InFactory, InSettings.Registry))
+		{
+			return nullptr;
+		}
+	}
+
 	TSharedPtr<FJsonObject> JsonObject = MakeShared<FJsonObject>();
 	if (!GenericStructAttributesToJson(InSettings, Struct, JsonObject.Get()))
 	{
@@ -417,7 +667,6 @@ TSharedPtr<FJsonObject> RigVMJsonUtils::ToJson(FIntrospectionSettings& InSetting
 		{
 			ArgumentJsonObject->SetStringField(TEXT("PropertyToolTip"), FString::Printf(TEXT("The %s execute argument."), *ExecuteArgument.Name.ToString()));
 		}
-		ArgumentJsonObject->SetStringField(TEXT("DefaultValue"), FString());
 
 		if (ExecuteArgument.Direction == ERigVMPinDirection::IO)
 		{
@@ -445,21 +694,55 @@ TSharedPtr<FJsonObject> RigVMJsonUtils::ToJson(FIntrospectionSettings& InSetting
 		}
 		
 		const FRigVMTemplateArgument& Argument = ArgumentInfo.GetArgument_NoLock(InSettings.Registry);
-		const TArray<TRigVMTypeIndex>& TypeIndices = Argument.GetSupportedTypeIndices_NoLock({}, InSettings.Registry);
+		TArray<TRigVMTypeIndex> TypeIndices = Argument.GetSupportedTypeIndices_NoLock({}, InSettings.Registry);
 		check(!TypeIndices.IsEmpty());
 
+		bool bIsWildCard = false;
+
 		FString Type;
 		if (TypeIndices.Num() == 1)
 		{
+			if (!InSettings.Schema->SupportsType_NoLock(nullptr, TypeIndices[0], InSettings.Registry))
+			{
+				continue;
+			}
 			Type = InSettings.Registry->GetType_NoLock(TypeIndices[0]).CPPType.ToString();
+			InSettings.AddDiscoveredType(TypeIndices[0]);
 		}
 		else if (TypeIndices.Num() < 16)
 		{
 			TArray<FString> Types;
 			Types.Reserve(TypeIndices.Num());
+			TypeIndices.Sort();
 			for (const TRigVMTypeIndex& TypeIndex : TypeIndices)
 			{
+				if (InSettings.Schema)
+				{
+					if (!InSettings.Schema->SupportsType_NoLock(nullptr, TypeIndex, InSettings.Registry))
+					{
+						continue;
+					}
+				}
+				const FRigVMTemplateArgumentType& RigVMType = InSettings.Registry->GetType_NoLock(TypeIndex);
+				if (RigVMType.CPPTypeObject)
+				{
+					if (RigVMType.CPPTypeObject->GetPathName().Contains(TEXT("EngineTest")))
+					{
+						continue;
+					}
+				}
+				if (Cast<UUserDefinedStruct>(RigVMType.CPPTypeObject))
+				{
+					Types.AddUnique(TEXT("User Defined Structs"));
+					continue;
+				}
+				if (Cast<UUserDefinedEnum>(RigVMType.CPPTypeObject))
+				{
+					Types.AddUnique(TEXT("User Defined Enums"));
+					continue;
+				}
 				Types.Add(InSettings.Registry->GetType_NoLock(TypeIndex).CPPType.ToString());
+				InSettings.AddDiscoveredType(TypeIndex);
 			}
 			Type = RigVMStringUtils::JoinStrings(Types, TEXT(","));
 		}
@@ -479,18 +762,46 @@ TSharedPtr<FJsonObject> RigVMJsonUtils::ToJson(FIntrospectionSettings& InSetting
 			if (NumArrayTypes == TypeIndices.Num())
 			{
 				Type = WildcardArrayType;
+				InSettings.AddDiscoveredType(RigVMTypeUtils::TypeIndex::WildCardArray);
 			}
 			else
 			{
 				Type = WildcardType;
+				InSettings.AddDiscoveredType(RigVMTypeUtils::TypeIndex::WildCard);
 			}
+
+			bIsWildCard = true;
 		}
 		
 		TSharedPtr<FJsonObject> ArgumentJsonObject = MakeShared<FJsonObject>();
 		ArgumentJsonObject->SetStringField(TEXT("Name"), ArgumentInfo.Name.ToString());
 		ArgumentJsonObject->SetStringField(TEXT("Type"), Type);
 		ArgumentJsonObject->SetStringField(TEXT("PropertyToolTip"), PostProcessToolTip(InFactory->GetArgumentTooltip(ArgumentInfo.Name, TypeIndices[0]).ToString()));
-		ArgumentJsonObject->SetStringField(TEXT("DefaultValue"), FString());
+
+		if ((Argument.GetDirection() != ERigVMPinDirection::Output) && !bIsWildCard)
+		{
+			if (Argument.GetNumTypes_NoLock(InSettings.Registry) == 1)
+			{
+				const TRigVMTypeIndex FirstTypeIndex = Argument.GetTypeIndex_NoLock(0, InSettings.Registry);
+				FString DefaultValue = InFactory->GetArgumentDefaultValue(Argument.GetName(), FirstTypeIndex);
+				if (DefaultValue.IsEmpty())
+				{
+					if (FirstTypeIndex == RigVMTypeUtils::TypeIndex::FName)
+					{
+						DefaultValue = TEXT("None");
+					}
+					else if(InSettings.Registry->IsArrayType_NoLock(FirstTypeIndex))
+					{
+						DefaultValue = TEXT("()");
+					}
+					else if(Cast<UStruct>(InSettings.Registry->GetType_NoLock(FirstTypeIndex).CPPTypeObject))
+					{
+						DefaultValue = TEXT("()");
+					}
+				}
+				ArgumentJsonObject->SetStringField(TEXT("DefaultValue"), DefaultValue);
+			}
+		}
 
 		if (Argument.GetDirection() == ERigVMPinDirection::IO)
 		{
@@ -513,6 +824,63 @@ TSharedPtr<FJsonObject> RigVMJsonUtils::ToJson(FIntrospectionSettings& InSetting
 	return JsonObject;
 }
 
+TSharedPtr<FJsonObject> RigVMJsonUtils::ToJson(FIntrospectionSettings& InSettings, const TRigVMTypeIndex& InTypeIndex)
+{
+	if (InTypeIndex == INDEX_NONE)
+	{
+		return nullptr;
+	}
+	if (InSettings.Schema)
+	{
+		if (!InSettings.Schema->SupportsType_NoLock(nullptr, InTypeIndex, InSettings.Registry))
+		{
+			return nullptr;
+		}
+	}
+
+	const FRigVMTemplateArgumentType& Type = InSettings.Registry->GetType_NoLock(InTypeIndex);
+	if (!Type.IsValid())
+	{
+		return nullptr;
+	}
+
+	TSharedPtr<FJsonObject> JsonObject = MakeShared<FJsonObject>();
+	JsonObject->SetStringField(TEXT("Type"), Type.CPPType.ToString());
+	if (Type.CPPTypeObject)
+	{
+		if (const UField* Field = Cast<UField>(Type.CPPTypeObject))
+		{
+			FString Label = Field->GetDisplayNameText().ToString();
+			FString ToolTip = PostProcessToolTip(Field->GetToolTipText().ToString());
+			if (ToolTip.Equals(Label))
+			{
+				ToolTip.Reset();
+			}
+			if (Type.CPPTypeObject->IsA<UEnum>())
+			{
+				if (Label.StartsWith(TEXT("E")) &&
+					Label.StartsWith(Type.CPPType.ToString().Left(2), ESearchCase::CaseSensitive))
+				{
+					Label.RightChopInline(1);
+				}
+			}
+			JsonObject->SetStringField(TEXT("Label"), Label);
+			if (!ToolTip.IsEmpty())
+			{
+				JsonObject->SetStringField(TEXT("ToolTip"), ToolTip);
+			}
+		}
+		JsonObject->SetStringField(TEXT("PackageName"), Type.CPPTypeObject->GetOutermost()->GetPathName());
+	}
+	if (InSettings.EdGraphSchema)
+	{
+		const FEdGraphPinType PinType = RigVMTypeUtils::PinTypeFromCPPType(Type.CPPType, Type.CPPTypeObject);
+		const FLinearColor PinColor = InSettings.EdGraphSchema->GetPinTypeColor(PinType);
+		JsonObject->SetStringField(TEXT("Color"), PinColor.ToString());
+	}
+	return JsonObject;
+}
+
 bool RigVMJsonUtils::GenericStructAttributesToJson(FIntrospectionSettings& InSettings, const UScriptStruct* InStruct, FJsonObject* OutJsonObject)
 {
 	check(InStruct);
diff --git a/Engine/Plugins/Runtime/RigVM/Source/RigVMDeveloper/Public/RigVMJsonUtils.h b/Engine/Plugins/Runtime/RigVM/Source/RigVMDeveloper/Public/RigVMJsonUtils.h
--- a/Engine/Plugins/Runtime/RigVM/Source/RigVMDeveloper/Public/RigVMJsonUtils.h
+++ b/Engine/Plugins/Runtime/RigVM/Source/RigVMDeveloper/Public/RigVMJsonUtils.h
@@ -15,6 +15,8 @@ class URigVMGraph;
 class URigVMNode;
 class URigVMPin;
 class URigVMLink;
+class URigVMSchema;
+class URigVMEdGraphSchema;
 struct FRigVMGraphVariableDescription;
 
 namespace RigVMJsonUtils
@@ -33,21 +35,25 @@ namespace RigVMJsonUtils
 
 	struct FIntrospectionSettings
 	{
-		FIntrospectionSettings(FRigVMRegistryHandle& InRegistry)
-			: Registry(InRegistry)
-		{
-		}
+		FIntrospectionSettings(FRigVMRegistryHandle& InRegistry);
+
+		void AddDiscoveredType(const TRigVMTypeIndex& InTypeIndex);
 
 		FRigVMRegistryHandle& Registry;
+		TObjectPtr<URigVMSchema> Schema = nullptr;
+		TObjectPtr<URigVMEdGraphSchema> EdGraphSchema = nullptr;
 		bool bExcludeFunctions = false;
 		bool bExcludeDispatches = false;
 		bool bExcludeDeprecated = false;
+		bool bMergeTemplates = false;
+		TSet<TRigVMTypeIndex> DiscoveredTypes;
 	};
 
 	// introspection for the registry
 	RIGVMDEVELOPER_API TSharedPtr<FJsonObject> AvailableNodesToJson(FIntrospectionSettings& InSettings, TFunction<bool(const UScriptStruct*)> InFilterCallback);
 	RIGVMDEVELOPER_API TSharedPtr<FJsonObject> ToJson(FIntrospectionSettings& InSettings, const FRigVMFunction* InFunction);
 	RIGVMDEVELOPER_API TSharedPtr<FJsonObject> ToJson(FIntrospectionSettings& InSettings, const FRigVMDispatchFactory* InFactory);
+	RIGVMDEVELOPER_API TSharedPtr<FJsonObject> ToJson(FIntrospectionSettings& InSettings, const TRigVMTypeIndex& InTypeIndex);
 	RIGVMDEVELOPER_API bool GenericStructAttributesToJson(FIntrospectionSettings& InSettings, const UScriptStruct* InStruct, FJsonObject* OutJsonObject);
 
 	// introspection for existing graphs
