Add LLM tags for UAF memory usage
* Tags the root of the UAF worker jobs with UAF to flood tag any nested allocations.
* The UAF/AnimGraph tracks graph/node allocations.
* The UAF/Debug and RigVM/Debug tags are for debug allocs such as debug draw.

diff --git a/Engine/Plugins/Experimental/UAF/UAF/Source/UAF/Private/Module/ModuleTickFunction.cpp b/Engine/Plugins/Experimental/UAF/UAF/Source/UAF/Private/Module/ModuleTickFunction.cpp
--- a/Engine/Plugins/Experimental/UAF/UAF/Source/UAF/Private/Module/ModuleTickFunction.cpp
+++ b/Engine/Plugins/Experimental/UAF/UAF/Source/UAF/Private/Module/ModuleTickFunction.cpp
@@ -120,6 +120,8 @@ void FModuleEventTickFunction::EndTick(float DeltaTime)
 
 void FModuleEventTickFunction::ExecuteTick(float DeltaTime, ELevelTick TickType, ENamedThreads::Type CurrentThread, const FGraphEventRef& MyCompletionGraphEvent)
 {
+	LLM_SCOPE_BYNAME(TEXT("UAF"));
+	
 	check(PostExecuteGameThreadTaskList.IsEmpty());	// Should always be empty since we move it to the game thread task
 
 	if (BatchedWorkHandle.IsValid() && GetActualEndTickGroup() != TickGroup)
@@ -180,7 +182,7 @@ void FModuleEventTickFunction::ExecuteTick(float DeltaTime, ELevelTick TickType,
 void FModuleEventTickFunction::Run(float InDeltaTime)
 {
 	SCOPED_NAMED_EVENT(UAF_Module_EventTick, FColor::Orange);
-
+	
 	while (!PreExecuteTasks.IsEmpty())
 	{
 		TOptional<TUniqueFunction<void(const FModuleTaskContext&)>> Function = PreExecuteTasks.Dequeue();
diff --git a/Engine/Plugins/Experimental/UAF/UAF/Source/UAF/Private/Module/SystemReference.cpp b/Engine/Plugins/Experimental/UAF/UAF/Source/UAF/Private/Module/SystemReference.cpp
--- a/Engine/Plugins/Experimental/UAF/UAF/Source/UAF/Private/Module/SystemReference.cpp
+++ b/Engine/Plugins/Experimental/UAF/UAF/Source/UAF/Private/Module/SystemReference.cpp
@@ -118,6 +118,8 @@ FSystemReference::FSystemReference(const UObject* InAsset, UObject* InObject, EA
 {
 	check(IsInGameThread());
 	{
+		LLM_SCOPE_BYNAME(TEXT("UAF"));
+		
 		const UAnimNextModule* Module = Cast<UAnimNextModule>(InAsset);
 		FUAFSystemFactoryParams FactoryParams;
 		if (Module == nullptr)
diff --git a/Engine/Plugins/Experimental/UAF/UAFAnimGraph/Source/UAFAnimGraph/Private/Graph/AnimNextAnimationGraph.cpp b/Engine/Plugins/Experimental/UAF/UAFAnimGraph/Source/UAFAnimGraph/Private/Graph/AnimNextAnimationGraph.cpp
--- a/Engine/Plugins/Experimental/UAF/UAFAnimGraph/Source/UAFAnimGraph/Private/Graph/AnimNextAnimationGraph.cpp
+++ b/Engine/Plugins/Experimental/UAF/UAFAnimGraph/Source/UAFAnimGraph/Private/Graph/AnimNextAnimationGraph.cpp
@@ -45,6 +45,7 @@ UAnimNextAnimationGraph::UAnimNextAnimationGraph(const FObjectInitializer& Objec
 
 TSharedPtr<FAnimNextGraphInstance> UAnimNextAnimationGraph::AllocateInstance(const UE::UAF::FGraphAllocationParams& InParams) const
 {
+	LLM_SCOPE_BYNAME(TEXT("UAF/AnimGraph"))
 	SCOPE_CYCLE_COUNTER(STAT_AnimNext_Graph_AllocateInstance);
 
 	PRAGMA_DISABLE_DEPRECATION_WARNINGS
diff --git a/Engine/Plugins/Experimental/UAF/UAFAnimGraph/Source/UAFAnimGraph/Private/TraitCore/ExecutionContext.cpp b/Engine/Plugins/Experimental/UAF/UAFAnimGraph/Source/UAFAnimGraph/Private/TraitCore/ExecutionContext.cpp
--- a/Engine/Plugins/Experimental/UAF/UAFAnimGraph/Source/UAFAnimGraph/Private/TraitCore/ExecutionContext.cpp
+++ b/Engine/Plugins/Experimental/UAF/UAFAnimGraph/Source/UAFAnimGraph/Private/TraitCore/ExecutionContext.cpp
@@ -122,6 +122,8 @@ namespace UE::UAF
 
 	FTraitPtr FExecutionContext::AllocateNodeInstance(FAnimNextGraphInstance& GraphInstance, FAnimNextTraitHandle ChildTraitHandle) const
 	{
+		LLM_SCOPE_BYNAME(TEXT("UAF/AnimGraph"));
+
 		if (!ChildTraitHandle.IsValid())
 		{
 			return FTraitPtr();	// Attempting to allocate a node using an invalid trait handle
@@ -205,6 +207,8 @@ namespace UE::UAF
 
 	FTraitPtr FExecutionContext::AllocateNodeInstance(const FWeakTraitPtr& ParentBinding, FAnimNextTraitHandle ChildTraitHandle) const
 	{
+		LLM_SCOPE_BYNAME(TEXT("UAF/AnimGraph"));
+		
 		if (!ChildTraitHandle.IsValid())
 		{
 			return FTraitPtr();	// Attempting to allocate a node using an invalid trait handle
diff --git a/Engine/Plugins/Runtime/RigVM/Source/RigVM/Private/RigVMCore/RigVMDrawInterface.cpp b/Engine/Plugins/Runtime/RigVM/Source/RigVM/Private/RigVMCore/RigVMDrawInterface.cpp
--- a/Engine/Plugins/Runtime/RigVM/Source/RigVM/Private/RigVMCore/RigVMDrawInterface.cpp
+++ b/Engine/Plugins/Runtime/RigVM/Source/RigVM/Private/RigVMCore/RigVMDrawInterface.cpp
@@ -10,6 +10,8 @@ TAutoConsoleVariable<int32> CVarEnableRigVMDrawInterfaceInShipping(TEXT("RigVM.E
 
 void FRigVMDrawInterface::DrawInstruction(const FRigVMDrawInstruction& InInstruction)
 {
+	LLM_SCOPE_BYNAME(TEXT("RigVM/Debug"));
+
 	if (!IsEnabled())
 	{
 		return;
@@ -20,6 +22,8 @@ void FRigVMDrawInterface::DrawInstruction(const FRigVMDrawInstruction& InInstruc
 
 void FRigVMDrawInterface::DrawPoint(const FTransform& WorldOffset, const FVector& Position, float Size, const FLinearColor& Color, ESceneDepthPriorityGroup DepthPriority, float Lifetime)
 {
+	LLM_SCOPE_BYNAME(TEXT("RigVM/Debug"));
+
 	if (!IsEnabled())
 	{
 		return;
@@ -32,6 +36,8 @@ void FRigVMDrawInterface::DrawPoint(const FTransform& WorldOffset, const FVector
 
 void FRigVMDrawInterface::DrawPoints(const FTransform& WorldOffset, const TArrayView<const FVector>& Points, float Size, const FLinearColor& Color, ESceneDepthPriorityGroup DepthPriority, float Lifetime)
 {
+	LLM_SCOPE_BYNAME(TEXT("RigVM/Debug"));
+
 	if (!IsEnabled())
 	{
 		return;
@@ -44,6 +50,8 @@ void FRigVMDrawInterface::DrawPoints(const FTransform& WorldOffset, const TArray
 
 void FRigVMDrawInterface::DrawLine(const FTransform& WorldOffset, const FVector& LineStart, const FVector& LineEnd, const FLinearColor& Color, float Thickness, ESceneDepthPriorityGroup DepthPriority, float Lifetime)
 {
+	LLM_SCOPE_BYNAME(TEXT("RigVM/Debug"));
+
 	if (!IsEnabled())
 	{
 		return;
@@ -57,6 +65,8 @@ void FRigVMDrawInterface::DrawLine(const FTransform& WorldOffset, const FVector&
 
 void FRigVMDrawInterface::DrawLines(const FTransform& WorldOffset, const TArrayView<const FVector>& Positions, const FLinearColor& Color, float Thickness, ESceneDepthPriorityGroup DepthPriority, float Lifetime)
 {
+	LLM_SCOPE_BYNAME(TEXT("RigVM/Debug"));
+
 	if (!IsEnabled())
 	{
 		return;
@@ -69,6 +79,8 @@ void FRigVMDrawInterface::DrawLines(const FTransform& WorldOffset, const TArrayV
 
 void FRigVMDrawInterface::DrawLineStrip(const FTransform& WorldOffset, const TArrayView<const FVector>& Positions, const FLinearColor& Color, float Thickness, ESceneDepthPriorityGroup DepthPriority, float Lifetime)
 {
+	LLM_SCOPE_BYNAME(TEXT("RigVM/Debug"));
+
 	if (!IsEnabled())
 	{
 		return;
@@ -81,6 +93,8 @@ void FRigVMDrawInterface::DrawLineStrip(const FTransform& WorldOffset, const TAr
 
 void FRigVMDrawInterface::DrawBox(const FTransform& WorldOffset, const FTransform& Transform, const FLinearColor& Color, float Thickness, ESceneDepthPriorityGroup DepthPriority, float Lifetime)
 {
+	LLM_SCOPE_BYNAME(TEXT("RigVM/Debug"));
+
 	if (!IsEnabled())
 	{
 		return;
@@ -122,6 +136,8 @@ void FRigVMDrawInterface::DrawBox(const FTransform& WorldOffset, const FTransfor
 
 void FRigVMDrawInterface::DrawAxes(const FTransform& WorldOffset, const FTransform& Transform, float Size, float Thickness, ESceneDepthPriorityGroup DepthPriority, float Lifetime)
 {
+	LLM_SCOPE_BYNAME(TEXT("RigVM/Debug"));
+
 	if (!IsEnabled())
 	{
 		return;
@@ -134,6 +150,8 @@ void FRigVMDrawInterface::DrawAxes(const FTransform& WorldOffset, const FTransfo
 
 void FRigVMDrawInterface::DrawAxes(const FTransform& WorldOffset, const TArrayView<const FTransform>& Transforms, float Size, float Thickness, ESceneDepthPriorityGroup DepthPriority, float Lifetime)
 {
+	LLM_SCOPE_BYNAME(TEXT("RigVM/Debug"));
+
 	if (!IsEnabled())
 	{
 		return;
@@ -170,6 +188,8 @@ void FRigVMDrawInterface::DrawAxes(const FTransform& WorldOffset, const TArrayVi
 void FRigVMDrawInterface::DrawAxes(const FTransform& WorldOffset, const FTransform& Transform,
 	const FLinearColor& InColor, float Size, float Thickness, ESceneDepthPriorityGroup DepthPriority, float Lifetime)
 {
+	LLM_SCOPE_BYNAME(TEXT("RigVM/Debug"));
+
 	if (!IsEnabled())
 	{
 		return;
@@ -191,6 +211,8 @@ void FRigVMDrawInterface::DrawAxes(const FTransform& WorldOffset, const FTransfo
 void FRigVMDrawInterface::DrawAxes(const FTransform& WorldOffset, const TArrayView<const FTransform>& Transforms,
 	const FLinearColor& InColor, float Size, float Thickness, ESceneDepthPriorityGroup DepthPriority, float Lifetime)
 {
+	LLM_SCOPE_BYNAME(TEXT("RigVM/Debug"));
+
 	if (!IsEnabled())
 	{
 		return;
@@ -220,6 +242,8 @@ void FRigVMDrawInterface::DrawAxes(const FTransform& WorldOffset, const TArrayVi
 
 void FRigVMDrawInterface::DrawRectangle(const FTransform& WorldOffset, const FTransform& Transform, float Size, const FLinearColor& Color, float Thickness, ESceneDepthPriorityGroup DepthPriority, float Lifetime)
 {
+	LLM_SCOPE_BYNAME(TEXT("RigVM/Debug"));
+
 	if (!IsEnabled())
 	{
 		return;
@@ -242,6 +266,8 @@ void FRigVMDrawInterface::DrawRectangle(const FTransform& WorldOffset, const FTr
 
 void FRigVMDrawInterface::DrawArc(const FTransform& WorldOffset, const FTransform& Transform, float Radius, float MinimumAngle, float MaximumAngle, const FLinearColor& Color, float Thickness, int32 Detail, ESceneDepthPriorityGroup DepthPriority, float Lifetime)
 {
+	LLM_SCOPE_BYNAME(TEXT("RigVM/Debug"));
+
 	if (!IsEnabled())
 	{
 		return;
@@ -276,12 +302,16 @@ void FRigVMDrawInterface::DrawArc(const FTransform& WorldOffset, const FTransfor
 void FRigVMDrawInterface::DrawCircle(const FTransform& WorldOffset, const FTransform& Transform, float Radius,
 	const FLinearColor& Color, float Thickness, int32 Detail, ESceneDepthPriorityGroup DepthPriority, float Lifetime)
 {
+	LLM_SCOPE_BYNAME(TEXT("RigVM/Debug"));
+
 	DrawArc(WorldOffset, Transform, Radius, 0, TWO_PI, Color, Thickness, Detail);
 }
 
 void FRigVMDrawInterface::DrawSphere(const FTransform& WorldOffset, const FTransform& Transform, float Radius, 
 	const FLinearColor& Color, float Thickness, int32 Detail, ESceneDepthPriorityGroup DepthPriority, float Lifetime)
 {
+	LLM_SCOPE_BYNAME(TEXT("RigVM/Debug"));
+
 	// Circles are drawn in the XY plane
 	FTransform R1(FQuat(FVector(1.f, 0.f, 0.f), HALF_PI));
 	FTransform R2(FQuat(FVector(0.f, 1.f, 0.f), HALF_PI));
@@ -293,6 +323,8 @@ void FRigVMDrawInterface::DrawSphere(const FTransform& WorldOffset, const FTrans
 void FRigVMDrawInterface::DrawHemisphere(const FTransform& WorldOffset, const FTransform& Transform, float Radius,
 	const FLinearColor& Color, float Thickness, int32 Detail, ESceneDepthPriorityGroup DepthPriority, float Lifetime)
 {
+	LLM_SCOPE_BYNAME(TEXT("RigVM/Debug"));
+
 	// Circles are drawn in the XY plane
 	FTransform R1(FQuat(FVector(1.f, 0.f, 0.f), HALF_PI));
 	FTransform R2(FQuat(FVector(0.f, 1.f, 0.f), -HALF_PI));
@@ -304,6 +336,8 @@ void FRigVMDrawInterface::DrawHemisphere(const FTransform& WorldOffset, const FT
 void FRigVMDrawInterface::DrawCapsule(const FTransform& WorldOffset, const FTransform& Transform,
 	float Radius, float Length, const FLinearColor& Color, float Thickness, int32 Detail, ESceneDepthPriorityGroup DepthPriority, float Lifetime)
 {
+	LLM_SCOPE_BYNAME(TEXT("RigVM/Debug"));
+
 	FTransform T1(FVector(0, 0, Length * 0.5));
 	FTransform T2(FQuat(FVector(0.f, 1.f, 0.f), PI), FVector(0, 0, -Length * 0.5));
 
@@ -329,6 +363,8 @@ void FRigVMDrawInterface::DrawCone(const FTransform& WorldOffset, const FTransfo
 	bool bDrawSideLines, const FLinearColor& SideLineColor, FMaterialRenderProxy* const MaterialRenderProxy, float SideLineThickness,
 	ESceneDepthPriorityGroup DepthPriority, float Lifetime)
 {
+	LLM_SCOPE_BYNAME(TEXT("RigVM/Debug"));
+
 	if (!IsEnabled())
 	{
 		return;
@@ -368,6 +404,8 @@ void FRigVMDrawInterface::DrawCone(const FTransform& WorldOffset, const FTransfo
 void FRigVMDrawInterface::DrawArrow(const FTransform& WorldOffset, const FVector& Direction, const FVector& Side,
 	const FLinearColor& Color, float Thickness, ESceneDepthPriorityGroup DepthPriority, float Lifetime)
 {
+	LLM_SCOPE_BYNAME(TEXT("RigVM/Debug"));
+
 	if (!IsEnabled())
 	{
 		return;
@@ -389,6 +427,8 @@ void FRigVMDrawInterface::DrawArrow(const FTransform& WorldOffset, const FVector
 void FRigVMDrawInterface::DrawPlane(const FTransform& WorldOffset, const FVector2D& Scale, const FLinearColor& MeshColor,
 	bool bDrawLines, const FLinearColor& LineColor, FMaterialRenderProxy* const MaterialRenderProxy, ESceneDepthPriorityGroup DepthPriority, float Lifetime)
 {
+	LLM_SCOPE_BYNAME(TEXT("RigVM/Debug"));
+
 	if (!IsEnabled())
 	{
 		return;
