RigVM: More follow up for static analysis breakage in RigVMSchema.cpp

#rnx

diff --git a/Engine/Plugins/Runtime/RigVM/Source/RigVMDeveloper/Private/RigVMModel/RigVMSchema.cpp b/Engine/Plugins/Runtime/RigVM/Source/RigVMDeveloper/Private/RigVMModel/RigVMSchema.cpp
--- a/Engine/Plugins/Runtime/RigVM/Source/RigVMDeveloper/Private/RigVMModel/RigVMSchema.cpp
+++ b/Engine/Plugins/Runtime/RigVM/Source/RigVMDeveloper/Private/RigVMModel/RigVMSchema.cpp
@@ -177,6 +177,11 @@ bool URigVMSchema::SupportsGraphFunction_NoLock(URigVMController* InController,
 {
 	RIGVMSCHEMA_DEFAULT_FUNCTION_BODY
 
+	if (!Graph)
+	{
+		return false;
+	}
+
 	check(InGraphFunction);
 
 	if(Graph->IsA<URigVMFunctionLibrary>())
@@ -338,6 +343,11 @@ bool URigVMSchema::CanAddNode(URigVMController* InController, const URigVMNode*
 {
 	RIGVMSCHEMA_DEFAULT_FUNCTION_BODY
 
+	if (!Graph)
+	{
+		return false;
+	}
+
 	if(InNode == nullptr)
 	{
 		return false;
@@ -584,6 +594,11 @@ bool URigVMSchema::CanAddLink(URigVMController* InController, const URigVMPin* I
 {
 	RIGVMSCHEMA_DEFAULT_FUNCTION_BODY
 
+	if (!Graph)
+	{
+		return false;
+	}
+
 	if (Graph->IsA<URigVMFunctionLibrary>())
 	{
 		static const FString FailureReason = TEXT("Cannot add links in function library graphs."); 
@@ -636,6 +651,11 @@ bool URigVMSchema::CanBreakLink(URigVMController* InController, const URigVMPin*
 {
 	RIGVMSCHEMA_DEFAULT_FUNCTION_BODY
 
+	if (!Graph)
+	{
+		return false;
+	}
+
 	if (Graph->IsA<URigVMFunctionLibrary>())
 	{
 		static const FString FailureReason = TEXT("Cannot break links in function library graphs."); 
@@ -682,6 +702,11 @@ bool URigVMSchema::CanCollapseNodes(URigVMController* InController, const TArray
 
 	RIGVMSCHEMA_DEFAULT_FUNCTION_BODY
 
+	if (!Graph)
+	{
+		return false;
+	}
+
 	for(int32 Index = 0; Index < InNodesToCollapse.Num(); Index++)
 	{
 		if(InNodesToCollapse[Index]->GetGraph() != Graph)
@@ -701,6 +726,11 @@ bool URigVMSchema::CanExpandNode(URigVMController* InController, const URigVMNod
 
 	RIGVMSCHEMA_DEFAULT_FUNCTION_BODY
 
+	if (!Graph)
+	{
+		return false;
+	}
+
 	if(Graph->IsA<URigVMFunctionLibrary>())
 	{
 		static const FString Message = TEXT("You cannot expand subgraphs within a function library graph.");
@@ -732,6 +762,11 @@ bool URigVMSchema::CanBindVariable(URigVMController* InController, const URigVMP
 {
 	RIGVMSCHEMA_DEFAULT_FUNCTION_BODY
 
+	if (!Graph)
+	{
+		return false;
+	}
+
 	if (Graph->IsA<URigVMFunctionLibrary>())
 	{
 		InController->ReportError(TEXT("Cannot bind pins to variables in function library graphs."));
@@ -767,6 +802,11 @@ bool URigVMSchema::CanUnbindVariable(URigVMController* InController, const URigV
 {
 	RIGVMSCHEMA_DEFAULT_FUNCTION_BODY
 
+	if (!Graph)
+	{
+		return false;
+	}
+
 	check(InBoundPin);
 
 	if (Graph->IsA<URigVMFunctionLibrary>())
@@ -787,6 +827,11 @@ bool URigVMSchema::CanAddFunction(URigVMController* InController, const URigVMNo
 {
 	RIGVMSCHEMA_DEFAULT_FUNCTION_BODY
 
+	if (!Graph)
+	{
+		return false;
+	}
+
 	if (!Graph->IsA<URigVMFunctionLibrary>())
 	{
 		InController->ReportError(TEXT("Can only add function definitions to function library graphs."));
@@ -808,6 +853,11 @@ bool URigVMSchema::CanRemoveFunction(URigVMController* InController, const URigV
 {
 	RIGVMSCHEMA_DEFAULT_FUNCTION_BODY
 
+	if (!Graph)
+	{
+		return false;
+	}
+
 	check(InFunctionNode);
 
 	if (!Graph->IsA<URigVMFunctionLibrary>())
diff --git a/Engine/Plugins/Runtime/RigVM/Source/RigVMDeveloper/Public/RigVMModel/RigVMSchema.h b/Engine/Plugins/Runtime/RigVM/Source/RigVMDeveloper/Public/RigVMModel/RigVMSchema.h
--- a/Engine/Plugins/Runtime/RigVM/Source/RigVMDeveloper/Public/RigVMModel/RigVMSchema.h
+++ b/Engine/Plugins/Runtime/RigVM/Source/RigVMDeveloper/Public/RigVMModel/RigVMSchema.h
@@ -13,11 +13,15 @@
 #define UE_API RIGVMDEVELOPER_API
 
 #define RIGVMSCHEMA_DEFAULT_FUNCTION_BODY \
-const URigVMGraph* Graph = InController->GetGraph(); \
-check(Graph); \
-if (!InController->IsTransacting() && !IsGraphEditable(Graph)) \
+const URigVMGraph* Graph = nullptr; \
+if (InController) \
 { \
-	return false; \
+	Graph = InController->GetGraph(); \
+	check(Graph); \
+	if (!InController->IsTransacting() && !IsGraphEditable(Graph)) \
+	{ \
+		return false; \
+	} \
 }
 
 class URigVMController;
