RigVM: Fix for static analysis issue in RigVMJsonUtils.cpp

#rnx

diff --git a/Engine/Plugins/Runtime/RigVM/Source/RigVMDeveloper/Private/RigVMJsonUtils.cpp b/Engine/Plugins/Runtime/RigVM/Source/RigVMDeveloper/Private/RigVMJsonUtils.cpp
--- a/Engine/Plugins/Runtime/RigVM/Source/RigVMDeveloper/Private/RigVMJsonUtils.cpp
+++ b/Engine/Plugins/Runtime/RigVM/Source/RigVMDeveloper/Private/RigVMJsonUtils.cpp
@@ -536,7 +536,10 @@ TSharedPtr<FJsonObject> RigVMJsonUtils::ToJson(FIntrospectionSettings& InSetting
 					DefaultValue = TEXT("None");
 				}
 			}
-			InSettings.AddDiscoveredType(InSettings.Registry->GetTypeIndexFromCPPType_NoLock(Argument->Type));
+			if (Argument)
+			{
+				InSettings.AddDiscoveredType(InSettings.Registry->GetTypeIndexFromCPPType_NoLock(Argument->Type));
+			}
 		}
 
 		TSharedPtr<FJsonObject> ArgumentJsonObject = MakeShared<FJsonObject>();
