RigVM: Json export to support merged templates

#rnx

diff --git a/Engine/Plugins/Runtime/RigVM/Source/RigVMDeveloper/Private/RigVMJsonUtils.cpp b/Engine/Plugins/Runtime/RigVM/Source/RigVMDeveloper/Private/RigVMJsonUtils.cpp
--- a/Engine/Plugins/Runtime/RigVM/Source/RigVMDeveloper/Private/RigVMJsonUtils.cpp
+++ b/Engine/Plugins/Runtime/RigVM/Source/RigVMDeveloper/Private/RigVMJsonUtils.cpp
@@ -452,6 +452,46 @@ TSharedPtr<FJsonObject> RigVMJsonUtils::ToJson(FIntrospectionSettings& InSetting
 		{
 			JsonObject->RemoveField(TEXT("Type"));
 		}
+
+#if WITH_EDITOR
+		JsonObject->SetStringField(TEXT("DisplayName"), Template->GetNodeName().ToString());
+		JsonObject->SetStringField(TEXT("Category"), Template->GetCategory());
+		FString KeywordsJoined = Template->GetKeywords();
+		TArray<FString> Keywords;
+		if (RigVMStringUtils::SplitString(KeywordsJoined, TEXT(","), Keywords))
+		{
+			TSet<FString> UniqueKeywords(Keywords);
+			Keywords = UniqueKeywords.Array();
+			KeywordsJoined = RigVMStringUtils::JoinStrings(Keywords, TEXT(","));
+		}
+		JsonObject->SetStringField(TEXT("Keywords"), KeywordsJoined);
+		FString ToolTip = Template->GetTooltipText().ToString();
+		if (ToolTip == Template->GetName().ToString())
+		{
+			TArray<FString> ToolTips;
+			for (const FRigVMFunction* Permutation : Permutations)
+			{
+				if (!Permutation->Struct)
+				{
+					continue;
+				}
+				FString PermutationToolTip = PostProcessToolTip(Permutation->Struct->GetToolTipText().ToString());
+				if (PermutationToolTip.IsEmpty())
+				{
+					continue;
+				}
+				if (PermutationToolTip == Permutation->Struct->GetDisplayNameText().ToString())
+				{
+					continue;
+				}
+				ToolTips.Add(PermutationToolTip);
+			}
+			TSet<FString> UniqueToolTips(ToolTips);
+			ToolTips = UniqueToolTips.Array();
+			ToolTip = RigVMStringUtils::JoinStrings(ToolTips, TEXT("\n"));
+		}
+		JsonObject->SetStringField(TEXT("ToolTip"), ToolTip);
+#endif
 	}
 	JsonObject->SetStringField(TEXT("TemplateName"), TemplateName);
 
