Control Rig: Fix mixing of Keywords and TemplateName in unit structs
#rnx

diff --git a/Engine/Plugins/Animation/ControlRig/Source/ControlRig/Public/Units/Hierarchy/RigUnit_SendEvent.h b/Engine/Plugins/Animation/ControlRig/Source/ControlRig/Public/Units/Hierarchy/RigUnit_SendEvent.h
--- a/Engine/Plugins/Animation/ControlRig/Source/ControlRig/Public/Units/Hierarchy/RigUnit_SendEvent.h
+++ b/Engine/Plugins/Animation/ControlRig/Source/ControlRig/Public/Units/Hierarchy/RigUnit_SendEvent.h
@@ -10,7 +10,7 @@
 /**
  * SendEvent is used to notify the engine / editor of a change that happend within the Control Rig.
  */
-USTRUCT(meta=(DisplayName="Send Event", Category="Hierarchy", DocumentationPolicy="Strict", Keywords = "SendEvent", TemplateName="Event,Notify,Notification", NodeColor="1, 0, 0"))
+USTRUCT(meta=(DisplayName="Send Event", Category="Hierarchy", DocumentationPolicy="Strict", TemplateName = "SendEvent", Keywords="Event,Notify,Notification", NodeColor="1, 0, 0"))
 struct FRigUnit_SendEvent : public FRigUnitMutable
 {
 	GENERATED_BODY()
diff --git a/Engine/Plugins/Runtime/RigVM/Source/RigVM/Public/RigVMFunctions/Math/RigVMFunction_MathBox.h b/Engine/Plugins/Runtime/RigVM/Source/RigVM/Public/RigVMFunctions/Math/RigVMFunction_MathBox.h
--- a/Engine/Plugins/Runtime/RigVM/Source/RigVM/Public/RigVMFunctions/Math/RigVMFunction_MathBox.h
+++ b/Engine/Plugins/Runtime/RigVM/Source/RigVM/Public/RigVMFunctions/Math/RigVMFunction_MathBox.h
@@ -18,7 +18,7 @@ struct FRigVMFunction_MathBoxBase : public FRigVMFunction_MathBase
 /**
  * Returns bounding box of the given array of positions
  */
-USTRUCT(meta = (DisplayName = "Box from Array", TemplateName = "ArrayBounds,CreateBox,CreateBoundingBox,NewBox,NewBoundingBoxMakeBox,MakeBoundingBox,BoundingBox,Bbox,Bounds"))
+USTRUCT(meta = (DisplayName = "Box from Array", Keywords = "ArrayBounds,CreateBox,CreateBoundingBox,NewBox,NewBoundingBoxMakeBox,MakeBoundingBox,BoundingBox,Bbox,Bounds"))
 struct FRigVMFunction_MathBoxFromArray : public FRigVMFunction_MathBoxBase
 {
 	GENERATED_BODY()
@@ -61,7 +61,7 @@ struct FRigVMFunction_MathBoxFromArray : public FRigVMFunction_MathBoxBase
 /**
  * Returns true if the box has any content / is valid
  */
-USTRUCT(meta = (DisplayName = "Is Box Valid", TemplateName = "IsValid,HasVolume,ContainsPoints,Bounds,BoundingBox,Bbox"))
+USTRUCT(meta = (DisplayName = "Is Box Valid", Keywords = "IsValid,HasVolume,ContainsPoints,Bounds,BoundingBox,Bbox"))
 struct FRigVMFunction_MathBoxIsValid : public FRigVMFunction_MathBoxBase
 {
 	GENERATED_BODY()
@@ -88,7 +88,7 @@ struct FRigVMFunction_MathBoxIsValid : public FRigVMFunction_MathBoxBase
 /**
  * Returns the center of a bounding box
  */
-USTRUCT(meta = (DisplayName = "Get Box Center", TemplateName = "Middle,Origin,Bounds,BoundingBox,Bbox"))
+USTRUCT(meta = (DisplayName = "Get Box Center", Keywords = "Middle,Origin,Bounds,BoundingBox,Bbox"))
 struct FRigVMFunction_MathBoxGetCenter : public FRigVMFunction_MathBoxBase
 {
 	GENERATED_BODY()
@@ -115,7 +115,7 @@ struct FRigVMFunction_MathBoxGetCenter : public FRigVMFunction_MathBoxBase
 /**
  * Returns the size of a bounding box
  */
-USTRUCT(meta = (DisplayName = "Get Box Size", TemplateName = "Middle,Origin,Bounds,BoundingBox,Bbox"))
+USTRUCT(meta = (DisplayName = "Get Box Size", Keywords = "Middle,Origin,Bounds,BoundingBox,Bbox"))
 struct FRigVMFunction_MathBoxGetSize : public FRigVMFunction_MathBoxBase
 {
 	GENERATED_BODY()
@@ -146,7 +146,7 @@ struct FRigVMFunction_MathBoxGetSize : public FRigVMFunction_MathBoxBase
 /**
  * Move the box by a certain amount 
  */
-USTRUCT(meta = (DisplayName = "Shift Box", TemplateName = "Bbox,Translate,Move,BoundingBox"))
+USTRUCT(meta = (DisplayName = "Shift Box", Keywords = "Bbox,Translate,Move,BoundingBox"))
 struct FRigVMFunction_MathBoxShift : public FRigVMFunction_MathBoxBase
 {
 	GENERATED_BODY()
@@ -177,7 +177,7 @@ struct FRigVMFunction_MathBoxShift : public FRigVMFunction_MathBoxBase
 /**
  * Moves the center of the box to a new location
  */
-USTRUCT(meta = (DisplayName = "Move Box To", TemplateName = "Bbox,Translate,Move,BoundingBox"))
+USTRUCT(meta = (DisplayName = "Move Box To", Keywords = "Bbox,Translate,Move,BoundingBox"))
 struct FRigVMFunction_MathBoxMoveTo : public FRigVMFunction_MathBoxBase
 {
 	GENERATED_BODY()
@@ -208,7 +208,7 @@ struct FRigVMFunction_MathBoxMoveTo : public FRigVMFunction_MathBoxBase
 /**
  * Expands the size of the box by a given amount
  */
-USTRUCT(meta = (DisplayName = "Expand Box", TemplateName = "Bbox,Scale,Grow,Shrink,BoundingBox"))
+USTRUCT(meta = (DisplayName = "Expand Box", Keywords = "Bbox,Scale,Grow,Shrink,BoundingBox"))
 struct FRigVMFunction_MathBoxExpand : public FRigVMFunction_MathBoxBase
 {
 	GENERATED_BODY()
@@ -239,7 +239,7 @@ struct FRigVMFunction_MathBoxExpand : public FRigVMFunction_MathBoxBase
 /**
  * Transforms the box by a given transform
  */
-USTRUCT(meta = (DisplayName = "Transform Box", TemplateName = "Bbox,BoundingBox"))
+USTRUCT(meta = (DisplayName = "Transform Box", Keywords = "Bbox,BoundingBox"))
 struct FRigVMFunction_MathBoxTransform : public FRigVMFunction_MathBoxBase
 {
 	GENERATED_BODY()
@@ -270,7 +270,7 @@ struct FRigVMFunction_MathBoxTransform : public FRigVMFunction_MathBoxBase
 /**
  * Returns the distance to a given box
  */
-USTRUCT(meta = (DisplayName = "Get Distance to Box", TemplateName = "Bbox,Closest,BoundingBox"))
+USTRUCT(meta = (DisplayName = "Get Distance to Box", Keywords = "Bbox,Closest,BoundingBox"))
 struct FRigVMFunction_MathBoxGetDistance : public FRigVMFunction_MathBoxBase
 {
 	GENERATED_BODY()
@@ -312,7 +312,7 @@ struct FRigVMFunction_MathBoxGetDistance : public FRigVMFunction_MathBoxBase
 /**
  * Returns true if a point is inside a given box
  */
-USTRUCT(meta = (DisplayName = "Is Inside Box", TemplateName = "Bbox,Contains,Encompasses,BoundingBox"))
+USTRUCT(meta = (DisplayName = "Is Inside Box", Keywords = "Bbox,Contains,Encompasses,BoundingBox"))
 struct FRigVMFunction_MathBoxIsInside : public FRigVMFunction_MathBoxBase
 {
 	GENERATED_BODY()
@@ -345,7 +345,7 @@ struct FRigVMFunction_MathBoxIsInside : public FRigVMFunction_MathBoxBase
 /**
  * Returns the volume of a given box
  */
-USTRUCT(meta = (DisplayName = "Get Box Volume", TemplateName = "Bbox,BoundingBox"))
+USTRUCT(meta = (DisplayName = "Get Box Volume", Keywords = "Bbox,BoundingBox"))
 struct FRigVMFunction_MathBoxGetVolume : public FRigVMFunction_MathBoxBase
 {
 	GENERATED_BODY()
diff --git a/Engine/Plugins/Runtime/RigVM/Source/RigVM/Public/RigVMFunctions/Math/RigVMFunction_MathQuaternion.h b/Engine/Plugins/Runtime/RigVM/Source/RigVM/Public/RigVMFunctions/Math/RigVMFunction_MathQuaternion.h
--- a/Engine/Plugins/Runtime/RigVM/Source/RigVM/Public/RigVMFunctions/Math/RigVMFunction_MathQuaternion.h
+++ b/Engine/Plugins/Runtime/RigVM/Source/RigVM/Public/RigVMFunctions/Math/RigVMFunction_MathQuaternion.h
@@ -604,7 +604,7 @@ struct FRigVMFunction_MathQuaternionSelectBool : public FRigVMFunction_MathQuate
 /**
  * Returns the dot product between two quaternions
  */
-USTRUCT(meta=(DisplayName="Dot", TemplateName="Dot,|"))
+USTRUCT(meta=(DisplayName="Dot", Keywords="Dot,|"))
 struct FRigVMFunction_MathQuaternionDot : public FRigVMFunction_MathQuaternionBase
 {
 	GENERATED_BODY()
diff --git a/Engine/Plugins/Runtime/RigVM/Source/RigVM/Public/RigVMFunctions/Math/RigVMFunction_MathVector.h b/Engine/Plugins/Runtime/RigVM/Source/RigVM/Public/RigVMFunctions/Math/RigVMFunction_MathVector.h
--- a/Engine/Plugins/Runtime/RigVM/Source/RigVM/Public/RigVMFunctions/Math/RigVMFunction_MathVector.h
+++ b/Engine/Plugins/Runtime/RigVM/Source/RigVM/Public/RigVMFunctions/Math/RigVMFunction_MathVector.h
@@ -768,7 +768,7 @@ struct FRigVMFunction_MathVectorCross : public FRigVMFunction_MathVectorBinaryOp
 /**
  * Returns the dot product between two vectors
  */
-USTRUCT(meta=(DisplayName="Dot", TemplateName="Dot,|"))
+USTRUCT(meta=(DisplayName="Dot", Keywords="Dot,|"))
 struct FRigVMFunction_MathVectorDot : public FRigVMFunction_MathVectorBase
 {
 	GENERATED_BODY()
