RigVM Json Utils: Improve type compacting for repetitive types

#rnx

diff --git a/Engine/Plugins/Runtime/RigVM/Source/RigVMDeveloper/Private/RigVMJsonUtils.cpp b/Engine/Plugins/Runtime/RigVM/Source/RigVMDeveloper/Private/RigVMJsonUtils.cpp
--- a/Engine/Plugins/Runtime/RigVM/Source/RigVMDeveloper/Private/RigVMJsonUtils.cpp
+++ b/Engine/Plugins/Runtime/RigVM/Source/RigVMDeveloper/Private/RigVMJsonUtils.cpp
@@ -608,6 +608,8 @@ TSharedPtr<FJsonObject> RigVMJsonUtils::ToJson(FIntrospectionSettings& InSetting
 					}
 					return A.Compare(B, ESearchCase::IgnoreCase) < 0;
 				});
+				TSet<FString> UniqueTypes(Types);
+				Types = UniqueTypes.Array();
 				ArgumentJsonObject->SetStringField(TEXT("Type"), RigVMStringUtils::JoinStrings(Types, TEXT(",")));
 			}
 		}
@@ -787,6 +789,8 @@ TSharedPtr<FJsonObject> RigVMJsonUtils::ToJson(FIntrospectionSettings& InSetting
 				Types.Add(InSettings.Registry->GetType_NoLock(TypeIndex).CPPType.ToString());
 				InSettings.AddDiscoveredType(TypeIndex);
 			}
+			TSet<FString> UniqueTypes(Types);
+			Types = UniqueTypes.Array();
 			Type = RigVMStringUtils::JoinStrings(Types, TEXT(","));
 		}
 		else
