Clarify RigVM help text for some trig functions. Adds a maths "Wrap" function to RigVM.

#Synced-CL 48349740

diff --git a/Engine/Plugins/Runtime/RigVM/Source/RigVM/Private/RigVMFunctions/Math/RigVMFunction_MathDouble.cpp b/Engine/Plugins/Runtime/RigVM/Source/RigVM/Private/RigVMFunctions/Math/RigVMFunction_MathDouble.cpp
--- a/Engine/Plugins/Runtime/RigVM/Source/RigVM/Private/RigVMFunctions/Math/RigVMFunction_MathDouble.cpp
+++ b/Engine/Plugins/Runtime/RigVM/Source/RigVM/Private/RigVMFunctions/Math/RigVMFunction_MathDouble.cpp
@@ -146,6 +146,11 @@ FRigVMFunction_MathDoubleClamp_Execute()
 	Result = FMath::Clamp<double>(Value, Minimum, Maximum);
 }
 
+FRigVMFunction_MathDoubleWrap_Execute()
+{
+	Result = FMath::Wrap<double>(Value, Minimum, Maximum);
+}
+
 FRigVMFunction_MathDoubleLerp_Execute()
 {
 	Result = FMath::Lerp<double>(A, B, T);
diff --git a/Engine/Plugins/Runtime/RigVM/Source/RigVM/Private/RigVMFunctions/Math/RigVMFunction_MathFloat.cpp b/Engine/Plugins/Runtime/RigVM/Source/RigVM/Private/RigVMFunctions/Math/RigVMFunction_MathFloat.cpp
--- a/Engine/Plugins/Runtime/RigVM/Source/RigVM/Private/RigVMFunctions/Math/RigVMFunction_MathFloat.cpp
+++ b/Engine/Plugins/Runtime/RigVM/Source/RigVM/Private/RigVMFunctions/Math/RigVMFunction_MathFloat.cpp
@@ -146,6 +146,11 @@ FRigVMFunction_MathFloatClamp_Execute()
 	Result = FMath::Clamp<float>(Value, Minimum, Maximum);
 }
 
+FRigVMFunction_MathFloatWrap_Execute()
+{
+	Result = FMath::Wrap<float>(Value, Minimum, Maximum);
+}
+
 FRigVMFunction_MathFloatLerp_Execute()
 {
 	Result = FMath::Lerp<float>(A, B, T);
diff --git a/Engine/Plugins/Runtime/RigVM/Source/RigVM/Private/RigVMFunctions/Math/RigVMFunction_MathFloatTests.cpp b/Engine/Plugins/Runtime/RigVM/Source/RigVM/Private/RigVMFunctions/Math/RigVMFunction_MathFloatTests.cpp
--- a/Engine/Plugins/Runtime/RigVM/Source/RigVM/Private/RigVMFunctions/Math/RigVMFunction_MathFloatTests.cpp
+++ b/Engine/Plugins/Runtime/RigVM/Source/RigVM/Private/RigVMFunctions/Math/RigVMFunction_MathFloatTests.cpp
@@ -166,6 +166,24 @@ IMPLEMENT_RIGVMSTRUCT_AUTOMATION_TEST(FRigVMFunction_MathFloatClamp)
 	return true;
 }
 
+IMPLEMENT_RIGVMSTRUCT_AUTOMATION_TEST(FRigVMFunction_MathFloatWrap)
+{
+	Unit.Minimum = 1.f;
+	Unit.Maximum = 3.f;
+	Unit.Value = 2.f;
+	Execute();
+	AddErrorIfFalse(FMath::IsNearlyEqual(Unit.Result, 2.f), TEXT("unexpected result"));
+
+	Unit.Value = -0.5f;
+	Execute();
+	AddErrorIfFalse(FMath::IsNearlyEqual(Unit.Result, 2.5f), TEXT("unexpected result"));
+
+	Unit.Value = 3.5f;
+	Execute();
+	AddErrorIfFalse(FMath::IsNearlyEqual(Unit.Result, 1.5f), TEXT("unexpected result"));
+	return true;
+}
+
 IMPLEMENT_RIGVMSTRUCT_AUTOMATION_TEST(FRigVMFunction_MathFloatLerp)
 {
 	Unit.A = 2.f;
diff --git a/Engine/Plugins/Runtime/RigVM/Source/RigVM/Public/RigVMFunctions/Math/RigVMFunction_MathDouble.h b/Engine/Plugins/Runtime/RigVM/Source/RigVM/Public/RigVMFunctions/Math/RigVMFunction_MathDouble.h
--- a/Engine/Plugins/Runtime/RigVM/Source/RigVM/Public/RigVMFunctions/Math/RigVMFunction_MathDouble.h
+++ b/Engine/Plugins/Runtime/RigVM/Source/RigVM/Public/RigVMFunctions/Math/RigVMFunction_MathDouble.h
@@ -539,6 +539,42 @@ struct FRigVMFunction_MathDoubleClamp : public FRigVMFunction_MathDoubleBase
 	double Result;
 };
 
+/**
+ * Wraps the given value to be within minimum and maximum, inclusive
+ * When the value can wrap to both Min and Max, it will wrap to Min if it lies below the range and wrap to Max if it is above the range.
+ */
+USTRUCT(meta = (DisplayName = "Wrap", TemplateName = "Wrap", Keywords = "Range,Remap"))
+struct FRigVMFunction_MathDoubleWrap : public FRigVMFunction_MathDoubleBase
+{
+	GENERATED_BODY()
+
+	FRigVMFunction_MathDoubleWrap()
+	{
+		Value = Result = 0.0;
+		Minimum = -1.0;
+		Maximum = 1.0;
+	}
+
+	RIGVM_METHOD()
+	RIGVM_API virtual void Execute() override;
+
+	// The number to wrap
+	UPROPERTY(meta = (Input))
+	double Value;
+
+	// The Minimum for the resulting range
+	UPROPERTY(meta = (Input))
+	double Minimum;
+
+	// The Maximum for the resulting range
+	UPROPERTY(meta = (Input))
+	double Maximum;
+
+	// The resulting wrapped value
+	UPROPERTY(meta = (Output))
+	double Result;
+};
+
 /**
  * Linearly interpolates between A and B using the ratio T
  */
@@ -887,7 +923,7 @@ struct FRigVMFunction_MathDoubleRad : public FRigVMFunction_MathDoubleUnaryOp
 };
 
 /**
- * Returns the sinus value of the given value (in radians)
+ * Returns the trigonometric sine value of the given angle (in radians)
  */
 USTRUCT(meta=(DisplayName="Sin", TemplateName="Sin"))
 struct FRigVMFunction_MathDoubleSin : public FRigVMFunction_MathDoubleUnaryOp
@@ -899,7 +935,7 @@ struct FRigVMFunction_MathDoubleSin : public FRigVMFunction_MathDoubleUnaryOp
 };
 
 /**
- * Returns the cosinus value of the given value (in radians)
+ * Returns the trigonometric cosine value of the given angle (in radians)
  */
 USTRUCT(meta=(DisplayName="Cos", TemplateName="Cos"))
 struct FRigVMFunction_MathDoubleCos : public FRigVMFunction_MathDoubleUnaryOp
@@ -911,7 +947,7 @@ struct FRigVMFunction_MathDoubleCos : public FRigVMFunction_MathDoubleUnaryOp
 };
 
 /**
- * Returns the tangens value of the given value (in radians)
+ * Returns the trigonometric tangent value of the given angle (in radians)
  */
 USTRUCT(meta=(DisplayName="Tan", TemplateName="Tan"))
 struct FRigVMFunction_MathDoubleTan : public FRigVMFunction_MathDoubleUnaryOp
@@ -923,7 +959,7 @@ struct FRigVMFunction_MathDoubleTan : public FRigVMFunction_MathDoubleUnaryOp
 };
 
 /**
- * Returns the inverse sinus value (in radians) of the given value
+ * Returns the inverse sine value (in radians between -pi/2 and pi/2) of the given value
  */
 USTRUCT(meta=(DisplayName="Asin", TemplateName="Asin", Keywords="Arcsin"))
 struct FRigVMFunction_MathDoubleAsin : public FRigVMFunction_MathDoubleUnaryOp
@@ -935,7 +971,7 @@ struct FRigVMFunction_MathDoubleAsin : public FRigVMFunction_MathDoubleUnaryOp
 };
 
 /**
- * Returns the inverse cosinus value (in radians) of the given value
+ * Returns the inverse cosine value (in radians between 0 and pi) of the given value
  */
 USTRUCT(meta=(DisplayName="Acos", TemplateName="Acos", Keywords="Arccos"))
 struct FRigVMFunction_MathDoubleAcos : public FRigVMFunction_MathDoubleUnaryOp
@@ -947,7 +983,7 @@ struct FRigVMFunction_MathDoubleAcos : public FRigVMFunction_MathDoubleUnaryOp
 };
 
 /**
- * Returns the inverse tangens value (in radians) of the given value
+ * Returns the inverse tangent value (in radians between -pi/2 and pi/2) of the given value
  */
 USTRUCT(meta=(DisplayName="Atan", TemplateName="Atan", Keywords="Arctan"))
 struct FRigVMFunction_MathDoubleAtan : public FRigVMFunction_MathDoubleUnaryOp
@@ -959,7 +995,7 @@ struct FRigVMFunction_MathDoubleAtan : public FRigVMFunction_MathDoubleUnaryOp
 };
 
 /**
- * Returns the arctangent of the specified A and B coordinates.
+ * Returns the inverse tangent value (in radians between -pi and pi) of A/B, using the signs of A and B to determine the correct quadrant.
  */
 USTRUCT(meta=(DisplayName="Atan2", TemplateName="Atan2", Keywords="Arctan"))
 struct FRigVMFunction_MathDoubleAtan2 : public FRigVMFunction_MathDoubleBinaryOp
diff --git a/Engine/Plugins/Runtime/RigVM/Source/RigVM/Public/RigVMFunctions/Math/RigVMFunction_MathFloat.h b/Engine/Plugins/Runtime/RigVM/Source/RigVM/Public/RigVMFunctions/Math/RigVMFunction_MathFloat.h
--- a/Engine/Plugins/Runtime/RigVM/Source/RigVM/Public/RigVMFunctions/Math/RigVMFunction_MathFloat.h
+++ b/Engine/Plugins/Runtime/RigVM/Source/RigVM/Public/RigVMFunctions/Math/RigVMFunction_MathFloat.h
@@ -541,6 +541,42 @@ struct FRigVMFunction_MathFloatClamp : public FRigVMFunction_MathFloatBase
 	float Result;
 };
 
+/**
+ * Wraps the given value to be within minimum and maximum, inclusive
+ * When the value can wrap to both Min and Max, it will wrap to Min if it lies below the range and wrap to Max if it is above the range.
+ */
+USTRUCT(meta = (DisplayName = "Wrap", TemplateName = "Wrap", Keywords = "Range,Remap"))
+struct FRigVMFunction_MathFloatWrap : public FRigVMFunction_MathFloatBase
+{
+	GENERATED_BODY()
+
+	FRigVMFunction_MathFloatWrap()
+	{
+		Value = Result = 0.0f;
+		Minimum = -1.0f;
+		Maximum = 1.0f;
+	}
+
+	RIGVM_METHOD()
+	UE_API virtual void Execute() override;
+
+	// The number to wrap
+	UPROPERTY(meta = (Input))
+	float Value;
+
+	// The Minimum for the resulting range
+	UPROPERTY(meta = (Input))
+	float Minimum;
+
+	// The Maximum for the resulting range
+	UPROPERTY(meta = (Input))
+	float Maximum;
+
+	// The resulting wrapped value
+	UPROPERTY(meta = (Output))
+	float Result;
+};
+
 /**
  * Linearly interpolates between A and B using the ratio T
  */
@@ -922,7 +958,7 @@ struct FRigVMFunction_MathFloatRad : public FRigVMFunction_MathFloatUnaryOp
 };
 
 /**
- * Returns the sinus value of the given value (in radians)
+ * Returns the trigonometric sine value of the given angle (in radians)
  */
 USTRUCT(meta=(DisplayName="Sin", TemplateName="Sin"))
 struct FRigVMFunction_MathFloatSin : public FRigVMFunction_MathFloatUnaryOp
@@ -934,7 +970,7 @@ struct FRigVMFunction_MathFloatSin : public FRigVMFunction_MathFloatUnaryOp
 };
 
 /**
- * Returns the cosinus value of the given value (in radians)
+ * Returns the trigonometric cosine value of the given angle (in radians)
  */
 USTRUCT(meta=(DisplayName="Cos", TemplateName="Cos"))
 struct FRigVMFunction_MathFloatCos : public FRigVMFunction_MathFloatUnaryOp
@@ -946,7 +982,7 @@ struct FRigVMFunction_MathFloatCos : public FRigVMFunction_MathFloatUnaryOp
 };
 
 /**
- * Returns the tangens value of the given value (in radians)
+ * Returns the trigonometric tangent value of the given angle (in radians)
  */
 USTRUCT(meta=(DisplayName="Tan", TemplateName="Tan"))
 struct FRigVMFunction_MathFloatTan : public FRigVMFunction_MathFloatUnaryOp
@@ -958,7 +994,7 @@ struct FRigVMFunction_MathFloatTan : public FRigVMFunction_MathFloatUnaryOp
 };
 
 /**
- * Returns the inverse sinus value (in radians) of the given value
+ * Returns the inverse sine value (in radians between -pi/2 and pi/2) of the given value
  */
 USTRUCT(meta=(DisplayName="Asin", TemplateName="Asin", Keywords="Arcsin"))
 struct FRigVMFunction_MathFloatAsin : public FRigVMFunction_MathFloatUnaryOp
@@ -970,7 +1006,7 @@ struct FRigVMFunction_MathFloatAsin : public FRigVMFunction_MathFloatUnaryOp
 };
 
 /**
- * Returns the inverse cosinus value (in radians) of the given value
+ * Returns the inverse cosine value (in radians between 0 and pi) of the given value
  */
 USTRUCT(meta=(DisplayName="Acos", TemplateName="Acos", Keywords="Arccos"))
 struct FRigVMFunction_MathFloatAcos : public FRigVMFunction_MathFloatUnaryOp
@@ -982,7 +1018,7 @@ struct FRigVMFunction_MathFloatAcos : public FRigVMFunction_MathFloatUnaryOp
 };
 
 /**
- * Returns the inverse tangens value (in radians) of the given value
+ * Returns the inverse tangent value (in radians between -pi/2 and pi/2) of the given value
  */
 USTRUCT(meta=(DisplayName="Atan", TemplateName="Atan", Keywords="Arctan"))
 struct FRigVMFunction_MathFloatAtan : public FRigVMFunction_MathFloatUnaryOp
@@ -994,7 +1030,7 @@ struct FRigVMFunction_MathFloatAtan : public FRigVMFunction_MathFloatUnaryOp
 };
 
 /**
- * Returns the arctangent of the specified A and B coordinates.
+ * Returns the inverse tangent value (in radians between -pi and pi) of A/B, using the signs of A and B to determine the correct quadrant.
  */
 USTRUCT(meta=(DisplayName="Atan2", TemplateName="Atan2", Keywords="Arctan"))
 struct FRigVMFunction_MathFloatAtan2 : public FRigVMFunction_MathFloatBinaryOp
