Control Rig: Add ability to create variants of multiple assets
#rnx

#Synced-CL 48139254

diff --git a/Engine/Plugins/Runtime/RigVM/Source/RigVMEditor/Private/RigVMEditorModule.cpp b/Engine/Plugins/Runtime/RigVM/Source/RigVMEditor/Private/RigVMEditorModule.cpp
--- a/Engine/Plugins/Runtime/RigVM/Source/RigVMEditor/Private/RigVMEditorModule.cpp
+++ b/Engine/Plugins/Runtime/RigVM/Source/RigVMEditor/Private/RigVMEditorModule.cpp
@@ -109,25 +109,23 @@ void FRigVMEditorModule::StartupModule()
 					UContentBrowserAssetContextMenuContext* Context = InSection.FindContext<UContentBrowserAssetContextMenuContext>();
 					if (Context)
 					{
-						TArray<FAssetData> SelectedAssets = Context->SelectedAssets;
-						if (SelectedAssets.Num() != 1)
-						{
-							// We only expect a single asset
-							return;
-						}
-
-						const FAssetData& SelectedAssetData = SelectedAssets[0];
-						if (!SelectedAssetData.GetClass()->ImplementsInterface(URigVMAssetInterface::StaticClass()))
-						{
-							// We aren't dealing with a RigVMBlueprint derived type
-							return;
-						}
-
 						if(CVarRigVMEnableVariants.GetValueOnAnyThread())
 						{
-							FSoftObjectPath SoftObjectPath = SelectedAssetData.GetSoftObjectPath();
-							InSection.AddMenuEntry("CreateVariant", LOCTEXT("CreateVariant", "Create variant"), LOCTEXT("CreateVariant_ToolTip", "Create a variant for this asset"), FSlateIcon(FRigVMEditorStyle::Get().GetStyleSetName(), "RigVM", "RigVM.Unit"), FExecuteAction::CreateLambda([SoftObjectPath]()
+							TArray<FAssetData> SelectedAssets = Context->SelectedAssets;
+							for (FAssetData SelectedAssetData : SelectedAssets)
+							{
+								if (!SelectedAssetData.GetClass()->ImplementsInterface(URigVMAssetInterface::StaticClass()))
 								{
+									// We aren't dealing with a RigVMBlueprint derived type
+									return;
+								}
+							}
+
+							InSection.AddMenuEntry("CreateVariant", LOCTEXT("CreateVariant", "Create variant"), LOCTEXT("CreateVariant_ToolTip", "Create a variant for this asset"), FSlateIcon(FRigVMEditorStyle::Get().GetStyleSetName(), "RigVM", "RigVM.Unit"), FExecuteAction::CreateLambda([SelectedAssets]()
+							{
+								for (FAssetData SelectedAssetData : SelectedAssets)
+								{
+									FSoftObjectPath SoftObjectPath = SelectedAssetData.GetSoftObjectPath();
 									// Perform the load from within our lambda since this can be expensive, and should not be done speculatively
 									UObject* SelectedObject = SoftObjectPath.TryLoad();
 									if (!SelectedObject)
@@ -149,7 +147,9 @@ void FRigVMEditorModule::StartupModule()
 									{
 										DuplicateBlueprint->GetAssetVariant() = Cast<IRigVMAssetInterface>(SelectedObject)->GetAssetVariant();
 									}
-								}));
+								}
+							}));
+							
 						}
 					}
 				}));
