Fix RigVM Wrap test

#Synced-CL 48362139

diff --git a/Engine/Plugins/Runtime/RigVM/Source/RigVM/Private/RigVMFunctions/Math/RigVMFunction_MathFloatTests.cpp b/Engine/Plugins/Runtime/RigVM/Source/RigVM/Private/RigVMFunctions/Math/RigVMFunction_MathFloatTests.cpp
--- a/Engine/Plugins/Runtime/RigVM/Source/RigVM/Private/RigVMFunctions/Math/RigVMFunction_MathFloatTests.cpp
+++ b/Engine/Plugins/Runtime/RigVM/Source/RigVM/Private/RigVMFunctions/Math/RigVMFunction_MathFloatTests.cpp
@@ -176,11 +176,20 @@ IMPLEMENT_RIGVMSTRUCT_AUTOMATION_TEST(FRigVMFunction_MathFloatWrap)
 
 	Unit.Value = -0.5f;
 	Execute();
-	AddErrorIfFalse(FMath::IsNearlyEqual(Unit.Result, 2.5f), TEXT("unexpected result"));
+	AddErrorIfFalse(FMath::IsNearlyEqual(Unit.Result, 1.5f), TEXT("unexpected result"));
 
 	Unit.Value = 3.5f;
 	Execute();
 	AddErrorIfFalse(FMath::IsNearlyEqual(Unit.Result, 1.5f), TEXT("unexpected result"));
+
+	Unit.Value = -1.0f;
+	Execute();
+	AddErrorIfFalse(FMath::IsNearlyEqual(Unit.Result, 1.0f), TEXT("unexpected result"));
+
+	Unit.Value = 5.0f;
+	Execute();
+	AddErrorIfFalse(FMath::IsNearlyEqual(Unit.Result, 3.0f), TEXT("unexpected result"));
+
 	return true;
 }
 
