Control Rig: Split IRigVMAssetInterface into IRigVMEditorAssetInterface and IRigVMRuntimeAssetInterface.

diff --git a/Engine/Plugins/Runtime/RigVM/Source/RigVM/Private/RigVMHost.cpp b/Engine/Plugins/Runtime/RigVM/Source/RigVM/Private/RigVMHost.cpp
--- a/Engine/Plugins/Runtime/RigVM/Source/RigVM/Private/RigVMHost.cpp
+++ b/Engine/Plugins/Runtime/RigVM/Source/RigVM/Private/RigVMHost.cpp
@@ -17,6 +17,8 @@
 #include "Kismet2/BlueprintEditorUtils.h"
 #endif// WITH_EDITOR
 
+#include "RigVMRuntimeAsset.h"
+
 #include UE_INLINE_GENERATED_CPP_BY_NAME(RigVMHost)
 
 #define LOCTEXT_NAMESPACE "RigVMHost"
@@ -265,8 +267,15 @@ void URigVMHost::Initialize(bool bRequestInit)
 		return;
 	}
 
-	InitializeFromCDO();
-	InstantiateVMFromCDO();
+	if (TScriptInterface<IRigVMRuntimeAssetInterface> RuntimeAsset = GeneratedBy)
+	{
+		RuntimeAsset->InitializeInstance(this);
+	}
+	else
+	{
+		InitializeFromCDO();
+		InstantiateVMFromCDO();
+	}
 
 	if (bRequestInit)
 	{
@@ -899,6 +908,11 @@ USceneComponent* URigVMHost::GetOwningSceneComponent()
 	return GetTypedOuter<USceneComponent>();
 }
 
+void URigVMHost::PostInitInstanceIfRequired()
+{
+	Variables.Refresh();
+}
+
 void URigVMHost::SwapVMToNativizedIfRequired(UClass* InNativizedClass)
 {
 	if (HasAnyFlags(RF_NeedPostLoad))
@@ -1199,77 +1213,89 @@ TArray<FRigVMExternalVariable> URigVMHost::GetExternalVariablesImpl(bool bFallba
 {
 	TArray<FRigVMExternalVariable> ExternalVariables;
 
-	for (TFieldIterator<FProperty> PropertyIt(GetClass()); PropertyIt; ++PropertyIt)
+	if (GeneratedBy)
 	{
-		FProperty* Property = *PropertyIt;
-		if(Property->IsNative())
+		FRigVMPropertyBag& MutableVariables = const_cast<URigVMHost*>(this)->Variables;
+		ExternalVariables = GeneratedBy->GetExternalVariables();
+		for (FRigVMExternalVariable& Variable : ExternalVariables)
 		{
-			continue;
+			Variable.Memory = MutableVariables.GetDataByName<uint8>(Variable.Name);
 		}
-
-		FRigVMExternalVariable ExternalVariable = FRigVMExternalVariable::Make(Property, (UObject*)this);
-		if(!ExternalVariable.IsValid())
-		{
-			UE_LOG(LogRigVM, Warning, TEXT("%s: Property '%s' of type '%s' is not supported."), *GetClass()->GetName(), *Property->GetName(), *Property->GetCPPType());
-			continue;
-		}
-
-		ExternalVariables.Add(ExternalVariable);
 	}
-
-#if WITH_EDITOR
-
-	if (bFallbackToBlueprint)
+	else
 	{
-		// if we have a difference in the blueprint variables compared to us - let's 
-		// use those instead. the assumption here is that the blueprint is dirty and
-		// hasn't been compiled yet.
-		if (UBlueprint* Blueprint = Cast<UBlueprint>(GetClass()->ClassGeneratedBy))
+		for (TFieldIterator<FProperty> PropertyIt(GetClass()); PropertyIt; ++PropertyIt)
 		{
-			TArray<FRigVMExternalVariable> BlueprintVariables;
-			for (const FBPVariableDescription& VariableDescription : Blueprint->NewVariables)
+			FProperty* Property = *PropertyIt;
+			if(Property->IsNative())
 			{
-				FRigVMExternalVariable ExternalVariable = RigVMTypeUtils::ExternalVariableFromBPVariableDescription(VariableDescription, (UObject*)this);
-				if (ExternalVariable.TypeName.IsNone())
-				{
-					continue;
-				}
-
-				ExternalVariable.Memory = nullptr;
-				BlueprintVariables.Add(ExternalVariable);
+				continue;
 			}
 
-			if (ExternalVariables.Num() != BlueprintVariables.Num())
+			FRigVMExternalVariable ExternalVariable = FRigVMExternalVariable::Make(Property, (UObject*)this);
+			if(!ExternalVariable.IsValid())
 			{
-				return BlueprintVariables;
+				UE_LOG(LogRigVM, Warning, TEXT("%s: Property '%s' of type '%s' is not supported."), *GetClass()->GetName(), *Property->GetName(), *Property->GetCPPType());
+				continue;
 			}
 
-			TMap<FName, int32> NameMap;
-			for (int32 Index = 0; Index < ExternalVariables.Num(); Index++)
-			{
-				NameMap.Add(ExternalVariables[Index].Name, Index);
-			}
+			ExternalVariables.Add(ExternalVariable);
+		}
+
+#if WITH_EDITOR
 
-			for (FRigVMExternalVariable BlueprintVariable : BlueprintVariables)
+		if (bFallbackToBlueprint)
+		{
+			// if we have a difference in the blueprint variables compared to us - let's 
+			// use those instead. the assumption here is that the blueprint is dirty and
+			// hasn't been compiled yet.
+			if (UBlueprint* Blueprint = Cast<UBlueprint>(GetClass()->ClassGeneratedBy))
 			{
-				const int32* Index = NameMap.Find(BlueprintVariable.Name);
-				if (Index == nullptr)
+				TArray<FRigVMExternalVariable> BlueprintVariables;
+				for (const FBPVariableDescription& VariableDescription : Blueprint->NewVariables)
 				{
-					return BlueprintVariables;
+					FRigVMExternalVariable ExternalVariable = RigVMTypeUtils::ExternalVariableFromBPVariableDescription(VariableDescription, (UObject*)this);
+					if (ExternalVariable.TypeName.IsNone())
+					{
+						continue;
+					}
+
+					ExternalVariable.Memory = nullptr;
+					BlueprintVariables.Add(ExternalVariable);
 				}
 
-				FRigVMExternalVariable ExternalVariable = ExternalVariables[*Index];
-				if (ExternalVariable.bIsArray != BlueprintVariable.bIsArray ||
-					ExternalVariable.bIsPublic != BlueprintVariable.bIsPublic ||
-					ExternalVariable.TypeName != BlueprintVariable.TypeName ||
-					ExternalVariable.TypeObject != BlueprintVariable.TypeObject)
+				if (ExternalVariables.Num() != BlueprintVariables.Num())
 				{
 					return BlueprintVariables;
 				}
+
+				TMap<FName, int32> NameMap;
+				for (int32 Index = 0; Index < ExternalVariables.Num(); Index++)
+				{
+					NameMap.Add(ExternalVariables[Index].Name, Index);
+				}
+
+				for (FRigVMExternalVariable BlueprintVariable : BlueprintVariables)
+				{
+					const int32* Index = NameMap.Find(BlueprintVariable.Name);
+					if (Index == nullptr)
+					{
+						return BlueprintVariables;
+					}
+
+					FRigVMExternalVariable ExternalVariable = ExternalVariables[*Index];
+					if (ExternalVariable.bIsArray != BlueprintVariable.bIsArray ||
+						ExternalVariable.bIsPublic != BlueprintVariable.bIsPublic ||
+						ExternalVariable.TypeName != BlueprintVariable.TypeName ||
+						ExternalVariable.TypeObject != BlueprintVariable.TypeObject)
+					{
+						return BlueprintVariables;
+					}
+				}
 			}
 		}
-	}
 #endif
+	}
 
 	return ExternalVariables;
 }
@@ -1283,13 +1309,26 @@ void URigVMHost::InstantiateVMFromCDO()
 		FRigVMExtendedExecuteContext& ExtendedExecuteContext = GetRigVMExtendedExecuteContext();
 
 		URigVMHost* CDO = GetClass()->GetDefaultObject<URigVMHost>();
-		if (VM && CDO && CDO->VM)
+		URigVM* SourceVM = nullptr;
+		if (CDO && CDO->VM)
+		{
+			SourceVM = CDO->VM;
+		}
+		else if (GeneratedBy && GeneratedBy->GetVM())
+		{
+			SourceVM = GeneratedBy->GetVM();
+		}
+		else
+		{
+			ensure(false);
+		}
+		if (VM && SourceVM)
 		{
 			if(!VM->IsNativized())
 			{
-				ExtendedExecuteContext.WorkMemoryStorage = CDO->VM->GetDefaultWorkMemory();
-				ExtendedExecuteContext.DebugMemoryStorage = CDO->VM->GetDefaultDebugMemory();
-				ExtendedExecuteContext.VMHash = CDO->VM->GetVMHash();
+				ExtendedExecuteContext.WorkMemoryStorage = SourceVM->GetDefaultWorkMemory();
+				ExtendedExecuteContext.DebugMemoryStorage = SourceVM->GetDefaultDebugMemory();
+				ExtendedExecuteContext.VMHash = SourceVM->GetVMHash();
 			}
 		}
 		else if (VM)
@@ -1313,6 +1352,11 @@ void URigVMHost::CopyExternalVariableDefaultValuesFromCDO()
 		URigVMHost* CDO = GetClass()->GetDefaultObject<URigVMHost>();
 		TArray<FRigVMExternalVariable> CurrentVariables = GetExternalVariablesImpl(false);
 		TArray<FRigVMExternalVariable> CDOVariables = CDO->GetExternalVariablesImpl(false);
+		if (GeneratedBy)
+		{
+			ensure(false);
+		}
+
 		if (ensure(CurrentVariables.Num() == CDOVariables.Num()))
 		{
 			for (int32 i=0; i<CurrentVariables.Num(); ++i)
diff --git a/Engine/Plugins/Runtime/RigVM/Source/RigVM/Private/RigVMTypeUtils.cpp b/Engine/Plugins/Runtime/RigVM/Source/RigVM/Private/RigVMTypeUtils.cpp
--- a/Engine/Plugins/Runtime/RigVM/Source/RigVM/Private/RigVMTypeUtils.cpp
+++ b/Engine/Plugins/Runtime/RigVM/Source/RigVM/Private/RigVMTypeUtils.cpp
@@ -9,6 +9,7 @@
 #include "AssetRegistry/AssetData.h"
 #include "AssetRegistry/AssetRegistryModule.h"
 #include "Engine/UserDefinedEnum.h"
+#include "RigVMCore/RigVMExternalVariable.h"
 #include "StructUtils/UserDefinedStruct.h"
 #include "UObject/LinkerLoad.h"
 #include "Interfaces/IPluginManager.h"
@@ -672,4 +673,64 @@ TArray<FString> RigVMTypeUtils::GetPluginModulePrefixes(const IPlugin& InPlugin)
 	ModulePrefixes.AddUnique(FString::Printf(PluginFormat, *InPlugin.GetName()));
 
 	return ModulePrefixes;
-}
\ No newline at end of file
+}
+
+bool RigVMTypeUtils::CPPTypeFromExternalVariable(const FRigVMExternalVariable& InExternalVariable, FString& OutCPPType, UObject** OutCPPTypeObject)
+{
+	FString Prefix = "";
+	FString Suffix = "";
+	if (InExternalVariable.bIsArray)
+	{
+		Prefix = TEXT("TArray<");
+		Suffix = TEXT(">");
+	}
+
+	*OutCPPTypeObject = nullptr;
+	if (InExternalVariable.TypeName == BoolTypeName)
+	{
+		OutCPPType = Prefix + BoolType + Suffix;
+	}
+	else if (InExternalVariable.TypeName == Int32TypeName)
+	{
+		OutCPPType = Prefix + Int32Type + Suffix;
+	}
+	else if (InExternalVariable.TypeName == FloatTypeName)
+	{
+		OutCPPType = Prefix + FloatType + Suffix;
+	}
+	else if (InExternalVariable.TypeName == DoubleTypeName)
+	{
+		OutCPPType = Prefix + DoubleType + Suffix;
+	}
+	else if (InExternalVariable.TypeName == FNameTypeName)
+	{
+		OutCPPType = Prefix + FNameType + Suffix;
+	}
+	else if (InExternalVariable.TypeName == FStringTypeName)
+	{
+		OutCPPType = Prefix + FStringType + Suffix;
+	}
+	else if (UScriptStruct* Struct = Cast<UScriptStruct>(InExternalVariable.TypeObject))
+	{
+		OutCPPType = Prefix + *GetUniqueStructTypeName(Struct) + Suffix;
+		*OutCPPTypeObject = Struct;	
+	}
+	else if (UEnum* Enum = Cast<UEnum>(InExternalVariable.TypeObject))
+	{
+		OutCPPType = Prefix + Enum->GetFName().ToString() + Suffix;
+		*OutCPPTypeObject = Enum;
+	}
+	else if (UClass* Class = Cast<UClass>(InExternalVariable.TypeObject))
+	{
+		Prefix += TEXT("TObjectPtr<U");
+		Suffix += TEXT(">");
+		OutCPPType = Prefix + Class->GetFName().ToString() + Suffix;
+		*OutCPPTypeObject = Class;
+	}
+	else
+	{
+		return false;
+	}
+		
+	return true;
+}
diff --git a/Engine/Plugins/Runtime/RigVM/Source/RigVM/Public/RigVMCore/RigVMMemoryStorage.h b/Engine/Plugins/Runtime/RigVM/Source/RigVM/Public/RigVMCore/RigVMMemoryStorage.h
--- a/Engine/Plugins/Runtime/RigVM/Source/RigVM/Public/RigVMCore/RigVMMemoryStorage.h
+++ b/Engine/Plugins/Runtime/RigVM/Source/RigVM/Public/RigVMCore/RigVMMemoryStorage.h
@@ -10,6 +10,7 @@
 #include "Misc/AssertionMacros.h"
 #include "RigVMMemoryCommon.h"
 #include "RigVMDefines.h"
+#include "RigVMPropertyBag.h"
 #if UE_RIGVM_DEBUG_EXECUTION
 #include "RigVMModule.h"
 #endif
@@ -999,74 +1000,7 @@ private:
 	friend struct TRigVMLazyValueBase;
 	friend class URigVMMemoryStorage;
 	friend struct FRigVMMemoryStorageStruct;
-};
-
-//////////////////////////////////////////////////////////////////////////////
-/// Property Management
-//////////////////////////////////////////////////////////////////////////////
-
-/**
- * The property description is used to provide all required information
- * to create a property for the memory storage class.
- */
-struct FRigVMPropertyDescription
-{
-public:
-	
-	// The name of the property to create
-	FName Name;
-
-	// The property to base a new property off of
-	const FProperty* Property;
-
-	// The complete CPP type to base a new property off of (for ex: 'TArray<TArray<FVector>>')
-	FString CPPType;
-
-	// The tail CPP Type object, for example the UScriptStruct for a struct 
-	UObject* CPPTypeObject;
-
-	// A list of containers to use for this property, for example [Array, Array]
-	TArray<EPinContainerType> Containers;
-
-	// The default value to use for this property (for example: '(((X=1.000000, Y=2.000000, Z=3.000000)))')
-	FString DefaultValue;
-
-	// Default constructor
-	FRigVMPropertyDescription()
-		: Name(NAME_None)
-		, Property(nullptr)
-		, CPPType()
-		, CPPTypeObject(nullptr)
-		, Containers()
-		, DefaultValue()
-	{}
-
-	// Constructor from an existing property
-	RIGVM_API FRigVMPropertyDescription(const FProperty* InProperty, const FString& InDefaultValue, const FName& InName = NAME_None);
-
-	// Constructor from complete data
-	RIGVM_API FRigVMPropertyDescription(const FName& InName, const FString& InCPPType, UObject* InCPPTypeObject, const FString& InDefaultValue);
-
-	// Returns a sanitized, valid name to use for a new property 
-	static RIGVM_API FName SanitizeName(const FName& InName);
-
-	// Sanitizes a name using a string ref, creating a valid name to use for a new property 
-	static RIGVM_API void SanitizeName(FString& InString);
-
-	// Sanitize the name of this description in line
-	RIGVM_API void SanitizeName();
-
-	// Returns true if this property description is valid
-	bool IsValid() const { return !Name.IsNone(); }
-
-	// Returns the CPP type of the tail property, for ex: '[2].Translation' it is 'FVector'
-	RIGVM_API FString GetTailCPPType() const;
-
-private:
-	
-	static const inline TCHAR* ArrayPrefix = TEXT("TArray<");
-	static const inline TCHAR* MapPrefix = TEXT("TMap<");
-	static const inline TCHAR* ContainerSuffix = TEXT(">");
+	friend struct FRigVMMemoryStorageStruct;
 };
 
 /**
diff --git a/Engine/Plugins/Runtime/RigVM/Source/RigVM/Public/RigVMCore/RigVMMemoryStorageStruct.h b/Engine/Plugins/Runtime/RigVM/Source/RigVM/Public/RigVMCore/RigVMMemoryStorageStruct.h
--- a/Engine/Plugins/Runtime/RigVM/Source/RigVM/Public/RigVMCore/RigVMMemoryStorageStruct.h
+++ b/Engine/Plugins/Runtime/RigVM/Source/RigVM/Public/RigVMCore/RigVMMemoryStorageStruct.h
@@ -2,8 +2,8 @@
 
 #pragma once
 
-#include "RigVMDefines.h"
 #include "RigVMCore/RigVMMemoryStorage.h"
+#include "RigVMCore/RigVMPropertyBag.h"
 #include "RigVMCore/RigVMPropertyPath.h"
 #include "StructUtils/PropertyBag.h"
 #include "RigVMMemoryStorageStruct.generated.h"
@@ -251,30 +251,30 @@ struct FRigVMMemoryStorageStruct : public FInstancedPropertyBag
 	//---------------------------------------------------------------------------
 
 	// Returns the exported text for a given property index
-	UE_API FString GetDataAsString(int32 InPropertyIndex, int32 PortFlags = PPF_None);
+	UE_API FString GetDataAsString(int32 InPropertyIndex, int32 PortFlags = PPF_None) const;
 
 	// Returns the exported text for given property name 
-	FString GetDataAsStringByName(const FName& InName, int32 PortFlags = PPF_None)
+	FString GetDataAsStringByName(const FName& InName, int32 PortFlags = PPF_None) const
 	{
 		const int32 PropertyIndex = GetPropertyIndexByName(InName);
 		return GetDataAsString(PropertyIndex, PortFlags);
 	}
 
 	// Returns the exported text for a given operand
-	UE_API FString GetDataAsString(const FRigVMOperand& InOperand, int32 PortFlags = PPF_None);
+	UE_API FString GetDataAsString(const FRigVMOperand& InOperand, int32 PortFlags = PPF_None) const;
 
 	// Returns the exported text for a given property index
-	UE_API FString GetDataAsStringSafe(int32 InPropertyIndex, int32 PortFlags = PPF_None);
+	UE_API FString GetDataAsStringSafe(int32 InPropertyIndex, int32 PortFlags = PPF_None) const;
 
 	// Returns the exported text for given property name 
-	FString GetDataAsStringByNameSafe(const FName& InName, int32 PortFlags = PPF_None)
+	FString GetDataAsStringByNameSafe(const FName& InName, int32 PortFlags = PPF_None) const
 	{
 		const int32 PropertyIndex = GetPropertyIndexByName(InName);
 		return GetDataAsStringSafe(PropertyIndex, PortFlags);
 	}
 
 	// Returns the exported text for a given operand
-	UE_API FString GetDataAsStringSafe(const FRigVMOperand& InOperand, int32 PortFlags = PPF_None);
+	UE_API FString GetDataAsStringSafe(const FRigVMOperand& InOperand, int32 PortFlags = PPF_None) const;
 
 	// Sets the content of a property by index given an exported string. Returns true if succeeded
 	UE_API bool SetDataFromString(int32 InPropertyIndex, const FString& InValue);
@@ -310,6 +310,7 @@ struct FRigVMMemoryStorageStruct : public FInstancedPropertyBag
 		PropertyPathDescriptions.Reset();
 	}
 
+	UE_API void Refresh();
 	UE_API void RefreshPropertyPaths();
 
 	bool IsValidPropertyPathDescriptionIndex(int32 Index) const
diff --git a/Engine/Plugins/Runtime/RigVM/Source/RigVM/Public/RigVMHost.h b/Engine/Plugins/Runtime/RigVM/Source/RigVM/Public/RigVMHost.h
--- a/Engine/Plugins/Runtime/RigVM/Source/RigVM/Public/RigVMHost.h
+++ b/Engine/Plugins/Runtime/RigVM/Source/RigVM/Public/RigVMHost.h
@@ -11,6 +11,7 @@
 #include "Engine/AssetUserData.h"
 #include "Interfaces/Interface_AssetUserData.h"
 #include "Misc/TransactionallySafeCriticalSection.h"
+#include "RigVMCore/RigVMMemoryStorageStruct.h"
 #if WITH_EDITOR
 #include "RigVMCore/RigVMDebugInfo.h"
 #endif
@@ -26,6 +27,8 @@
 #define UE_RIGVM_PROFILE_EXECUTE_UNITS_NUM 0
 #endif
 
+class IRigVMRuntimeAssetInterface;
+
 USTRUCT()
 struct FRigVMUserDefinedTypesInfo
 {
@@ -311,7 +314,7 @@ public:
 
 	UE_API virtual USceneComponent* GetOwningSceneComponent(); 
 
-	virtual void PostInitInstanceIfRequired() {};
+	UE_API virtual void PostInitInstanceIfRequired();
 	UE_API void SwapVMToNativizedIfRequired(UClass* InNativizedClass = nullptr);
 	static UE_API bool AreNativizedVMsDisabled() ;
 
@@ -361,6 +364,8 @@ public:
 
 	UE_API const TArray<FString>& GetRequiredPlugins() const;
 	
+	TScriptInterface<IRigVMRuntimeAssetInterface> GetGeneratedByAsset() const { return GeneratedBy; }
+
 protected:
 
 	UE_API virtual void PostInitInstance(URigVMHost* InCDO);
@@ -385,6 +390,12 @@ protected:
 	UPROPERTY(Transient)
 	TObjectPtr<URigVM> VM;
 
+	UPROPERTY()
+	FRigVMPropertyBag Variables;
+
+	UPROPERTY()
+	TScriptInterface<IRigVMRuntimeAssetInterface> GeneratedBy;
+
 #if WITH_EDITOR
 	FRigVMLog* RigVMLog;
 	bool bEnableLogging;
@@ -593,11 +604,14 @@ protected:
 #endif
 
 	friend class URigVMBlueprint;
-	friend class IRigVMAssetInterface;
+	friend class IRigVMEditorAssetInterface;
 	friend class URigVMBlueprintGeneratedClass;
 	friend class FRigVMCompileSettingsDetails;
 	friend class FRigVMTrace;
 	friend class URigVM;
+	friend class URigVMRuntimeAsset;
+	friend class URigVMEditorAsset;
+	friend class FRigVMExampleAssetActions;
 };
 
 class FRigVMBracketScope
diff --git a/Engine/Plugins/Runtime/RigVM/Source/RigVM/Public/RigVMRuntimeAsset.h b/Engine/Plugins/Runtime/RigVM/Source/RigVM/Public/RigVMRuntimeAsset.h
new file mode 100644
--- /dev/null
+++ b/Engine/Plugins/Runtime/RigVM/Source/RigVM/Public/RigVMRuntimeAsset.h
@@ -0,0 +1,277 @@
+// Copyright Epic Games, Inc. All Rights Reserved.
+
+#pragma once
+
+#include "CoreMinimal.h"
+#include "RigVMHost.h"
+#include "RigVMCore/RigVM.h"
+#include "RigVMCore/RigVMGraphFunctionHost.h"
+#include "RigVMCore/RigVMMemoryStorageStruct.h"
+#include "RigVMCore/RigVMVariableDescription.h"
+
+#include "RigVMRuntimeAsset.generated.h"
+
+#define UE_API RIGVM_API
+
+DECLARE_MULTICAST_DELEGATE_OneParam(FOnRigVMPreVariablesChanged,  const FName);
+DECLARE_MULTICAST_DELEGATE_OneParam(FOnRigVMPostVariablesChanged, const FName);
+
+
+/**
+ * Enumerates states a RigVMAsset can be in.
+ * Copied from EBlueprintStatus.
+ */
+UENUM()
+enum ERigVMAssetStatus : int
+{
+	/** Asset is in an unknown state. */
+	RVMA_Unknown,
+	/** Asset has been modified but not recompiled. */
+	RVMA_Dirty,
+	/** Asset tried but failed to be compiled. */
+	RVMA_Error,
+	/** Asset has been compiled since it was last modified. */
+	RVMA_UpToDate,
+	/** Asset is in the process of being created for the first time. */
+	RVMA_BeingCreated,
+	/** Asset has been compiled since it was last modified. There are warnings. */
+	RVMA_UpToDateWithWarnings,
+	RVMA_MAX,
+};
+
+UINTERFACE(MinimalAPI, BlueprintType)
+class URigVMRuntimeAssetInterface : public UInterface
+{
+	GENERATED_BODY()
+};
+
+class IRigVMRuntimeAssetInterface
+{
+	GENERATED_BODY()
+
+protected:
+	mutable TArray<TObjectPtr<UObject>> ArchetypeInstances;
+
+
+public:
+
+#if WITH_EDITORONLY_DATA
+	virtual TObjectPtr<UObject> GetEditorOnlyData() = 0;
+	virtual const TObjectPtr<UObject> GetEditorOnlyData() const = 0;
+#endif // WITH_EDITORONLY_DATA
+
+	virtual UClass* GetRigVMHostClass() const = 0;
+	virtual URigVM* GetVM() = 0;
+
+	virtual TArray<FRigVMExternalVariable> GetExternalVariables() = 0;
+	virtual TArray<FRigVMGraphVariableDescription> GetAssetVariables() = 0;
+	virtual FProperty* FindGeneratedPropertyByName(const FName& InName) = 0;
+	virtual const FProperty* FindGeneratedPropertyByName(const FName& InName) const { return const_cast<IRigVMRuntimeAssetInterface*>(this)->FindGeneratedPropertyByName(InName); }
+	virtual FRigVMGraphVariableDescription FindAssetVariable(const FName& InName) const = 0;
+#if WITH_EDITOR
+	virtual FOnRigVMPreVariablesChanged& OnPreVariablesChanged() = 0;
+	virtual FOnRigVMPostVariablesChanged& OnPostVariablesChanged() = 0;
+	virtual FName AddHostMemberVariableFromExternal(FRigVMExternalVariable InVariableToCreate, FString InDefaultValue = FString()) = 0;
+	virtual bool RemoveMemberVariable(const FName& InName) = 0;
+	virtual bool BulkRemoveMemberVariables(const TArray<FName>& InNames) = 0;
+	virtual bool RenameMemberVariable(const FName& InOldName, const FName& InNewName) = 0;
+	virtual bool ChangeMemberVariableType(const FName& InName, const FString& InCPPType, bool bIsPublic, bool bIsReadOnly, FString InDefaultValue) = 0;
+	virtual FText GetVariableTooltip(const FName& InName) const = 0;
+	virtual bool SetVariableTooltip(const FName& InName, const FText& InTooltip) = 0;
+	virtual FString GetVariableCategory(const FName& InName) const = 0;
+	virtual bool SetVariableCategory(const FName& InName, const FString& InCategory) = 0;
+	virtual bool SetVariableMetadataValue(const FName& InName, const FName& InKey, const FString& InValue) = 0;
+	virtual bool RemoveVariableMetadataValue(const FName& InName, const FName& InKey) = 0;
+	virtual FString GetVariableMetadataValue(const FName& InName, const FName& InKey) const = 0;
+	virtual bool SetVariableExposeOnSpawn(const FName& InName, const bool bInExposeOnSpawn) = 0;
+	virtual bool SetVariableExposeToCinematics(const FName& InName, const bool bInExposeToCinematics) = 0;
+	virtual bool SetVariablePrivate(const FName& InName, const bool bInPrivate) = 0;
+	virtual bool SetVariablePublic(const FName& InName, const bool bIsPublic) = 0;
+	virtual FString GetVariableDefaultValue(const FName& InName) const { return FindAssetVariable(InName).DefaultValue; }
+#endif
+	
+	virtual FRigVMRuntimeSettings& GetVMRuntimeSettings() = 0;
+	virtual FRigVMVariant& GetAssetVariant() = 0;
+	virtual const FRigVMVariant& GetAssetVariant() const { return const_cast<IRigVMRuntimeAssetInterface*>(this)->GetAssetVariant(); }
+	virtual TArray<FRigVMGraphFunctionHeader>& GetPublicGraphFunctions() = 0;
+	virtual TMap<FString, FSoftObjectPath>& GetUserDefinedStructGuidToPathName() = 0;
+	virtual TMap<FString, FSoftObjectPath>& GetUserDefinedEnumToPathName() = 0;
+	virtual TSet<TObjectPtr<UObject>>& GetUserDefinedTypesInUse() = 0;
+	virtual TArray<FName> GetSupportedEventNames() = 0;
+	virtual void UpdateSupportedEventNames() = 0;
+	virtual TArray<UObject*> GetArchetypeInstances(bool bIncludeDerivedClass) const = 0;
+	virtual ERigVMAssetStatus GetAssetStatus() const = 0;
+	virtual void SetAssetStatus(const ERigVMAssetStatus& InStatus) = 0;
+	virtual bool IsUpToDate() const = 0;
+	virtual FRigVMExtendedExecuteContext* GetRigVMExtendedExecuteContext() = 0;
+	virtual void MarkAssetAsModified(FPropertyChangedEvent PropertyChangedEvent) = 0;
+	virtual void MarkAssetAsStructurallyModified(bool bSkipDirtyAssetStatus) = 0;
+	virtual URigVMHost* InstantiateObject(UObject* InOuter) = 0;
+	virtual bool InitializeInstance(URigVMHost* InInstance) = 0;
+	virtual const TArray<FString>& GetRequiredPlugins(bool bRefresh = true) const = 0;
+	virtual FRigVMPropertyBag* GetVariablesPropertyBag() = 0;
+
+	friend class IRigVMEditorAssetInterface;
+
+};
+
+UCLASS(MinimalAPI, Abstract)
+class URigVMRuntimeAsset : public UObject, public IRigVMRuntimeAssetInterface, public IRigVMGraphFunctionHost
+{
+	GENERATED_BODY()
+
+protected:
+
+	virtual FRigVMGraphFunctionStore* GetRigVMGraphFunctionStore() override { return &GraphFunctionStore; }
+	virtual const FRigVMGraphFunctionStore* GetRigVMGraphFunctionStore() const override { return &GraphFunctionStore; }
+
+	UE_API virtual void PostLoad() override;
+	UE_API virtual void PreDuplicate(FObjectDuplicationParameters& DupParams) override;
+	UE_API virtual void PostDuplicate(bool bDuplicateForPIE) override;
+
+	virtual UClass* GetRigVMHostClass() const override { return URigVMHost::StaticClass(); }
+	virtual URigVM* GetVM() override { return VM; }
+
+	// Variables Begin
+	UE_API virtual TArray<FRigVMExternalVariable> GetExternalVariables() override;
+	UE_API virtual TArray<FRigVMGraphVariableDescription> GetAssetVariables() override;
+	UE_API virtual FProperty* FindGeneratedPropertyByName(const FName& InName) override;
+	UE_API virtual FRigVMGraphVariableDescription FindAssetVariable(const FName& InName) const override;
+#if WITH_EDITOR
+	FOnRigVMPreVariablesChanged OnPreVariablesChangedDelegate;
+	FOnRigVMPostVariablesChanged OnPostVariablesChangedDelegate;
+	virtual FOnRigVMPreVariablesChanged& OnPreVariablesChanged() override { return OnPreVariablesChangedDelegate; }
+	virtual FOnRigVMPreVariablesChanged& OnPostVariablesChanged() override { return OnPostVariablesChangedDelegate; }
+	bool bIsModifyingVariables = false;
+	UE_API virtual bool RemoveMemberVariable(const FName& InName) override;
+	UE_API virtual bool BulkRemoveMemberVariables(const TArray<FName>& InNames) override;
+	UE_API virtual bool RenameMemberVariable(const FName& InOldName, const FName& InNewName) override;
+	UE_API virtual bool ChangeMemberVariableType(const FName& InName, const FString& InCPPType, bool bIsPublic, bool bIsReadOnly, FString InDefaultValue) override;
+	UE_API virtual FText GetVariableTooltip(const FName& InName) const override;
+	UE_API virtual bool SetVariableTooltip(const FName& InName, const FText& InTooltip) override;
+	UE_API virtual FString GetVariableCategory(const FName& InName) const override;
+	UE_API virtual bool SetVariableCategory(const FName& InName, const FString& InCategory) override;
+	UE_API virtual FString GetVariableMetadataValue(const FName& InName, const FName& InKey) const override;
+	UE_API virtual bool SetVariableMetadataValue(const FName& InName, const FName& InKey, const FString& InValue) override;
+	UE_API virtual bool RemoveVariableMetadataValue(const FName& InName, const FName& InKey) override;
+	UE_API virtual bool SetVariableExposeOnSpawn(const FName& InName, const bool bInExposeOnSpawn) override;
+	UE_API virtual bool SetVariableExposeToCinematics(const FName& InName, const bool bInExposeToCinematics) override;
+	UE_API virtual bool SetVariablePrivate(const FName& InName, const bool bInPrivate) override;
+	UE_API virtual bool SetVariablePublic(const FName& InName, const bool bIsPublic) override;
+	UE_API virtual FName AddHostMemberVariableFromExternal(FRigVMExternalVariable InVariableToCreate, FString InDefaultValue = FString()) override;
+#endif
+	// Variables End
+
+	virtual FRigVMRuntimeSettings& GetVMRuntimeSettings() override { return RuntimeSettings; }
+	virtual FRigVMVariant& GetAssetVariant() override { return Variant; }
+	virtual TArray<FRigVMGraphFunctionHeader>& GetPublicGraphFunctions() override { return PublicGraphFunctions; }
+	virtual TMap<FString, FSoftObjectPath>& GetUserDefinedStructGuidToPathName() override { return UserDefinedStructGuidToPathName; }
+	virtual TMap<FString, FSoftObjectPath>& GetUserDefinedEnumToPathName() override { return UserDefinedEnumToPathName; }
+	virtual TSet<TObjectPtr<UObject>>& GetUserDefinedTypesInUse() override { return UserDefinedTypesInUse; }
+	virtual TArray<FName> GetSupportedEventNames() override { return SupportedEventNames; }
+	UE_API virtual void UpdateSupportedEventNames() override;
+	UE_API virtual TArray<UObject*> GetArchetypeInstances(bool bIncludeDerivedClass) const override;
+	virtual ERigVMAssetStatus GetAssetStatus() const override { return Status; }
+	virtual void SetAssetStatus(const ERigVMAssetStatus& InStatus) override { Status = InStatus; }
+	UE_API virtual bool IsUpToDate() const override;
+	virtual FRigVMExtendedExecuteContext* GetRigVMExtendedExecuteContext() override { return &ExecuteContext; }
+	UE_API virtual void MarkAssetAsModified(FPropertyChangedEvent PropertyChangedEvent = FPropertyChangedEvent(nullptr)) override;
+	UE_API virtual void MarkAssetAsStructurallyModified(bool bSkipDirtyAssetStatus) override;
+	UE_API virtual const TArray<FString>& GetRequiredPlugins(bool bRefresh = true) const override;
+	UE_API virtual void GenerateRequiredPluginsData(FRigVMExtendedExecuteContext& InContext);
+	virtual FRigVMPropertyBag* GetVariablesPropertyBag() override { return &Variables; }
+
+public:
+	UE_API virtual void Initialize(UClass* InEditorOnlyDataClass);
+	UE_API URigVMHost* InstantiateObject(UObject* InOuter);
+	UE_API virtual bool InitializeInstance(URigVMHost* InInstance) override;
+
+#if WITH_EDITORONLY_DATA
+	virtual TObjectPtr<UObject> GetEditorOnlyData() override { return EditorAsset;}
+	virtual const TObjectPtr<UObject> GetEditorOnlyData() const override { return EditorAsset;}
+#endif
+
+protected:
+
+	UPROPERTY()
+	TObjectPtr<URigVM> VM;
+
+	UPROPERTY()
+	FRigVMExtendedExecuteContext ExecuteContext;
+
+	UPROPERTY()
+	FRigVMPropertyBag Variables;
+
+	UPROPERTY()
+	FRigVMGraphFunctionStore GraphFunctionStore;
+
+	/** Asset searchable information about exposed public functions on this rig */
+	UPROPERTY(AssetRegistrySearchable)
+	TArray<FRigVMGraphFunctionHeader> PublicGraphFunctions;
+
+	UPROPERTY(EditAnywhere, Category = "VM")
+	FRigVMRuntimeSettings RuntimeSettings;
+
+	UPROPERTY()
+	FRigVMDrawContainer DrawContainer;
+
+	/** Variant information about this asset */
+	UPROPERTY(EditAnywhere, AssetRegistrySearchable, Category = "Variant")
+	FRigVMVariant Variant;
+
+	UPROPERTY()
+	TMap<FString, FSoftObjectPath> UserDefinedStructGuidToPathName;
+	UPROPERTY()
+	TMap<FString, FSoftObjectPath> UserDefinedEnumToPathName;
+	UPROPERTY(transient)
+	TSet<TObjectPtr<UObject>> UserDefinedTypesInUse;
+
+	/** The event names this rigvm blueprint contains */
+	UPROPERTY(AssetRegistrySearchable)
+	TArray<FName> SupportedEventNames;
+
+	/** The current status of this blueprint */
+	UPROPERTY(transient)
+	TEnumAsByte<ERigVMAssetStatus> Status;
+
+#if WITH_EDITORONLY_DATA
+	UPROPERTY()
+	TObjectPtr<UObject> EditorAsset; // The outer of this object is the runtime asset
+
+#endif
+
+	UPROPERTY(VisibleAnywhere, AssetRegistrySearchable, Category = "Dependencies")
+	TArray<FString> RequiredPlugins;
+
+	/** Array of user data stored with the asset */
+	UPROPERTY(EditAnywhere, AdvancedDisplay, Instanced, Category = "Default")
+	TArray<TObjectPtr<UAssetUserData>> AssetUserData;
+
+#if WITH_EDITORONLY_DATA
+	/** Array of user data stored with the asset */
+	UPROPERTY(EditAnywhere, AdvancedDisplay, Instanced, Category = "Default")
+	TArray<TObjectPtr<UAssetUserData>> AssetUserDataEditorOnly;
+#endif
+
+
+	friend class FRigVMExampleAssetActions;
+	friend class FRigVMAssetVariablesChangeScope;
+};
+
+#if WITH_EDITOR
+class FRigVMAssetVariablesChangeScope
+{
+public:
+   
+	UE_API FRigVMAssetVariablesChangeScope(URigVMRuntimeAsset* InAsset, const FName InVariableName);
+
+	UE_API ~FRigVMAssetVariablesChangeScope();
+
+private:
+
+	const FName VariableName;
+	URigVMRuntimeAsset* Asset;
+};
+#endif
+
+#undef UE_API
\ No newline at end of file
diff --git a/Engine/Plugins/Runtime/RigVM/Source/RigVM/Public/RigVMTypeUtils.h b/Engine/Plugins/Runtime/RigVM/Source/RigVM/Public/RigVMTypeUtils.h
--- a/Engine/Plugins/Runtime/RigVM/Source/RigVM/Public/RigVMTypeUtils.h
+++ b/Engine/Plugins/Runtime/RigVM/Source/RigVM/Public/RigVMTypeUtils.h
@@ -15,6 +15,7 @@
 
 #define UE_API RIGVM_API
 
+struct FRigVMExternalVariable;
 struct FRigVMTemplateArgumentType;
 class IPlugin;
 
@@ -279,6 +280,8 @@ namespace RigVMTypeUtils
 
 	RIGVM_API FString GetPluginName(const UObject* InObject, TMap<FString,FString>& InOutPluginPrefixes);
 	RIGVM_API TArray<FString> GetPluginModulePrefixes(const IPlugin& InPlugin);
+
+	RIGVM_API bool CPPTypeFromExternalVariable(const FRigVMExternalVariable& InExternalVariable, FString& OutCPPType, UObject** OutCPPTypeObject);
 }
 
 #undef UE_API
diff --git a/Engine/Plugins/Runtime/RigVM/Source/RigVMDeveloper/Private/EdGraph/RigVMEdGraph.cpp b/Engine/Plugins/Runtime/RigVM/Source/RigVMDeveloper/Private/EdGraph/RigVMEdGraph.cpp
--- a/Engine/Plugins/Runtime/RigVM/Source/RigVMDeveloper/Private/EdGraph/RigVMEdGraph.cpp
+++ b/Engine/Plugins/Runtime/RigVM/Source/RigVMDeveloper/Private/EdGraph/RigVMEdGraph.cpp
@@ -1,7 +1,7 @@
 // Copyright Epic Games, Inc. All Rights Reserved.
 
 #include "EdGraph/RigVMEdGraph.h"
-#include "RigVMAsset.h"
+#include "RigVMEditorAsset.h"
 #include "RigVMBlueprintGeneratedClass.h"
 #include "Editor/EditorEngine.h"
 #include "EdGraph/RigVMEdGraphSchema.h"
@@ -15,6 +15,7 @@
 #include "RigVMModel/Nodes/RigVMRerouteNode.h"
 #include "Stats/StatsHierarchical.h"
 #include "Blueprint/BlueprintExtension.h"
+#include "RigVMDeveloperTypeUtils.h"
 
 #include UE_INLINE_GENERATED_CPP_BY_NAME(RigVMEdGraph)
 
@@ -110,12 +111,12 @@ const URigVMBlueprint* URigVMEdGraph::GetBlueprintDefaultObject() const
 	return Cast<URigVMBlueprint>(GetAssetDefaultObject().GetObject());
 }
 
-const FRigVMAssetInterfacePtr URigVMEdGraph::GetAssetDefaultObject() const
+const FRigVMEditorAssetInterfacePtr URigVMEdGraph::GetAssetDefaultObject() const
 {
 	if(RigVMBlueprintClass)
 	{
 		UObject* Obj = RigVMBlueprintClass->GetDefaultObject(false);
-		if (Obj->Implements<URigVMAssetInterface>())
+		if (Obj->Implements<URigVMEditorAssetInterface>())
 		{
 			return Obj;
 		}
@@ -126,13 +127,13 @@ const FRigVMAssetInterfacePtr URigVMEdGraph::GetAssetDefaultObject() const
 void URigVMEdGraph::SetBlueprintClass(const UClass* InClass)
 {
 	RigVMBlueprintClass = InClass;
-	if(const FRigVMAssetInterfacePtr Blueprint = GetAssetDefaultObject())
+	if(const FRigVMEditorAssetInterfacePtr Blueprint = GetAssetDefaultObject())
 	{
 		Schema = Blueprint->GetRigVMClientHost()->GetRigVMEdGraphSchemaClass();
 	}
 }
 
-void URigVMEdGraph::InitializeFromAsset(FRigVMAssetInterfacePtr InBlueprint)
+void URigVMEdGraph::InitializeFromAsset(FRigVMEditorAssetInterfacePtr InBlueprint)
 {
 	DECLARE_SCOPE_HIERARCHICAL_COUNTER_FUNC()
 
@@ -1021,7 +1022,7 @@ void URigVMEdGraph::AddLocalVariableSearchMetaDataInfo(const FName InVariableNam
 			{
 				OutTaggedMetaData.Emplace(FRigVMSearchTags::FiB_Name, FText::FromName(InVariableName));
 				
-				const FEdGraphPinType Type = LocalVariable.ToPinType();
+				const FEdGraphPinType Type = RigVMTypeUtils::PinTypeFromRigVMVariableDescription(LocalVariable);
 				OutTaggedMetaData.Emplace(FRigVMSearchTags::FiB_PinCategory, FText::FromName(Type.PinCategory));
 				OutTaggedMetaData.Emplace(FRigVMSearchTags::FiB_PinSubCategory, FText::FromName(Type.PinSubCategory));
 				if (Type.PinSubCategoryObject.IsValid())
@@ -1034,7 +1035,7 @@ void URigVMEdGraph::AddLocalVariableSearchMetaDataInfo(const FName InVariableNam
 	}
 }
 
-FRigVMAssetInterfacePtr URigVMEdGraph::GetAsset() const
+FRigVMEditorAssetInterfacePtr URigVMEdGraph::GetAsset() const
 {
 	if (URigVMEdGraph* OuterGraph = Cast<URigVMEdGraph>(GetOuter()))
 	{
diff --git a/Engine/Plugins/Runtime/RigVM/Source/RigVMDeveloper/Private/EdGraph/RigVMEdGraphNode.cpp b/Engine/Plugins/Runtime/RigVM/Source/RigVMDeveloper/Private/EdGraph/RigVMEdGraphNode.cpp
--- a/Engine/Plugins/Runtime/RigVM/Source/RigVMDeveloper/Private/EdGraph/RigVMEdGraphNode.cpp
+++ b/Engine/Plugins/Runtime/RigVM/Source/RigVMDeveloper/Private/EdGraph/RigVMEdGraphNode.cpp
@@ -18,11 +18,12 @@
 #endif
 #include "RigVMEditorModule.h"
 #include "Textures/SlateIcon.h"
-#include "RigVMAsset.h"
+#include "RigVMEditorAsset.h"
 #include "PropertyPathHelpers.h"
 #include "RigVMBlueprintLegacy.h"
 #include "RigVMBlueprintUtils.h"
 #include "RigVMStringUtils.h"
+#include "RigVMDeveloperTypeUtils.h"
 #include "RigVMCore/RigVMExecuteContext.h"
 #include "RigVMFunctions/RigVMDispatch_CastEnum.h"
 #include "RigVMFunctions/RigVMDispatch_CastObject.h"
@@ -132,7 +133,7 @@ FText URigVMEdGraphNode::GetNodeTitle(ENodeTitleType::Type TitleType) const
 
 				if (SubTitle.Len() == 0)
 				{
-					if(FRigVMAssetInterfacePtr Blueprint = GetAsset())
+					if(FRigVMEditorAssetInterfacePtr Blueprint = GetAsset())
 					{
 						const FName VariableName = VariableNode->GetVariableName();
 						for(const FRigVMGraphVariableDescription& NewVariable : Blueprint->GetAssetVariables())
@@ -146,7 +147,7 @@ FText URigVMEdGraphNode::GetNodeTitle(ENodeTitleType::Type TitleType) const
 								}
 								else
 								{
-									const FEdGraphPinType PinType = NewVariable.ToPinType();
+									const FEdGraphPinType PinType = RigVMTypeUtils::PinTypeFromRigVMVariableDescription(NewVariable);
 									// Change the order of values in rotators so that they match the pin order
 									if (!PinType.IsContainer() && PinType.PinSubCategoryObject == TBaseStructure<FRotator>::Get())
 									{
@@ -248,7 +249,7 @@ void URigVMEdGraphNode::ReconstructNode_Internal(bool bForce)
 	// Maintain watches up to date
 	if (const URigVMNode* Node = GetModelNode())
 	{
-		FRigVMAssetInterfacePtr RigVMBlueprint = GetAsset();
+		FRigVMEditorAssetInterfacePtr RigVMBlueprint = GetAsset();
 		for (UEdGraphPin* NewPin : Pins)
 		{
 			const FString PinName = NewPin->GetName();
@@ -371,7 +372,7 @@ void URigVMEdGraphNode::DestroyPinList(TArray<UEdGraphPin*>& InPins)
 {
 	DECLARE_SCOPE_HIERARCHICAL_COUNTER_FUNC()
 
-	const FRigVMAssetInterfacePtr Blueprint = GetAsset();
+	const FRigVMEditorAssetInterfacePtr Blueprint = GetAsset();
 	bool bNotify = false;
 	if (Blueprint != nullptr)
 	{
@@ -1516,7 +1517,7 @@ float URigVMEdGraphNode::GetFadedOutState() const
 		bool bWasRunDuringLastEval = false;
 		if (CachedModelNode.IsValid())
 		{
-			if(IRigVMAssetInterface* Blueprint = GetImplementingOuter<IRigVMAssetInterface>())
+			if(IRigVMEditorAssetInterface* Blueprint = GetImplementingOuter<IRigVMEditorAssetInterface>())
 			{
 				if(URigVMHost* DebuggedHost = Cast<URigVMHost>(Blueprint->GetObjectBeingDebugged()))
 				{
@@ -1666,7 +1667,7 @@ void URigVMEdGraphNode::AllocateDefaultPins()
 
 UClass* URigVMEdGraphNode::GetRigVMGeneratedClass() const
 {
-	if(const IRigVMAssetInterface* Blueprint = GetImplementingOuter<IRigVMAssetInterface>())
+	if(const IRigVMEditorAssetInterface* Blueprint = GetImplementingOuter<IRigVMEditorAssetInterface>())
 	{
 		if (const UClass* GeneratedClass = Blueprint->GetRigVMGeneratedClass())
 		{
@@ -1834,7 +1835,7 @@ void URigVMEdGraphNode::GetNodeContextMenuActions(class UToolMenu* Menu, class U
 {
 #if WITH_EDITOR
 	const URigVMEdGraphSchema* Schema = Cast<URigVMEdGraphSchema>(GetSchema());
-	if(const FRigVMAssetInterfacePtr Blueprint = GetAsset())
+	if(const FRigVMEditorAssetInterfacePtr Blueprint = GetAsset())
 	{
 		return Blueprint->GetEditorModule()->GetContextMenuActions(Schema, Menu, Context);
 	}
@@ -1859,7 +1860,7 @@ void URigVMEdGraphNode::DestroyNode()
 	{
 		BreakAllNodeLinks();
 
-		if(IRigVMAssetInterface* RigVMBlueprint = Graph->GetImplementingOuter<IRigVMAssetInterface>())
+		if(IRigVMEditorAssetInterface* RigVMBlueprint = Graph->GetImplementingOuter<IRigVMEditorAssetInterface>())
 		{
 			if(PropertyName_DEPRECATED.IsValid())
 			{
@@ -1932,7 +1933,7 @@ URigVMBlueprint* URigVMEdGraphNode::GetBlueprint() const
 	return nullptr;
 }
 
-FRigVMAssetInterfacePtr URigVMEdGraphNode::GetAsset() const
+FRigVMEditorAssetInterfacePtr URigVMEdGraphNode::GetAsset() const
 {
 	if(const URigVMEdGraph* Graph = Cast<URigVMEdGraph>(GetOuter()))
 	{
@@ -2327,12 +2328,12 @@ void URigVMEdGraphNode::UpdateVisualSettings()
 	
 	if(const URigVMEdGraph* RigGraph = Cast<URigVMEdGraph>(GetOuter()))
 	{
-		if(const FRigVMAssetInterfacePtr RigBlueprint = RigGraph->GetAsset())
+		if(const FRigVMEditorAssetInterfacePtr RigBlueprint = RigGraph->GetAsset())
 		{
 			bEnableProfiling = RigBlueprint->GetVMRuntimeSettings().bEnableProfiling;
 			if (bEnableProfiling.Get(false))
 			{
-				if(IRigVMAssetInterface* Blueprint = GetImplementingOuter<IRigVMAssetInterface>())
+				if(IRigVMEditorAssetInterface* Blueprint = GetImplementingOuter<IRigVMEditorAssetInterface>())
 				{
 					if(URigVMHost* DebuggedHost = Cast<URigVMHost>(Blueprint->GetObjectBeingDebugged()))
 					{
diff --git a/Engine/Plugins/Runtime/RigVM/Source/RigVMDeveloper/Private/EdGraph/RigVMEdGraphSchema.cpp b/Engine/Plugins/Runtime/RigVM/Source/RigVMDeveloper/Private/EdGraph/RigVMEdGraphSchema.cpp
--- a/Engine/Plugins/Runtime/RigVM/Source/RigVMDeveloper/Private/EdGraph/RigVMEdGraphSchema.cpp
+++ b/Engine/Plugins/Runtime/RigVM/Source/RigVMDeveloper/Private/EdGraph/RigVMEdGraphSchema.cpp
@@ -21,14 +21,15 @@
 #include "RigVMModel/Nodes/RigVMCollapseNode.h"
 #include "RigVMModel/Nodes/RigVMFunctionEntryNode.h"
 #include "RigVMModel/Nodes/RigVMFunctionReturnNode.h"
-#include "RigVMModel/RigVMVariableDescription.h"
+#include "RigVMCore/RigVMVariableDescription.h"
 #include "Kismet2/Kismet2NameValidators.h"
 #include "Algo/Count.h"
 #include "RigVMFunctions/RigVMDispatch_Array.h"
 #include "RigVMFunctions/RigVMDispatch_Constant.h"
 #include "RigVMFunctions/RigVMDispatch_MakeStruct.h"
 #include "RigVMModel/Nodes/RigVMDispatchNode.h"
-#include "RigVMAsset.h"
+#include "RigVMEditorAsset.h"
+#include "RigVMDeveloperTypeUtils.h"
 #if WITH_RIGVMLEGACYEDITOR
 #include "BlueprintEditor.h"
 #else
@@ -36,6 +37,7 @@
 #endif
 
 #include "RigVMBlueprintUtils.h"
+#include "EdGraph/NodeSpawners/RigVMEdGraphNodeSpawner.h"
 
 #include UE_INLINE_GENERATED_CPP_BY_NAME(RigVMEdGraphSchema)
 
@@ -59,7 +61,7 @@ const FText FRigVMSearchTags::FiB_Glyph = LOCTEXT("Glyph", "Glyph");
 const FText FRigVMSearchTags::FiB_GlyphStyleSet = LOCTEXT("GlyphStyleSet", "GlyphStyleSet");
 const FText FRigVMSearchTags::FiB_GlyphColor = LOCTEXT("GlyphColor", "GlyphColor");
 
-FRigVMLocalVariableNameValidator::FRigVMLocalVariableNameValidator(const FRigVMAssetInterfacePtr InAsset, const URigVMGraph* Graph, FName InExistingName)
+FRigVMLocalVariableNameValidator::FRigVMLocalVariableNameValidator(const FRigVMEditorAssetInterfacePtr InAsset, const URigVMGraph* Graph, FName InExistingName)
 	: FStringSetNameValidator(InExistingName.ToString())
 {
 	if (const UBlueprint* Blueprint = Cast<UBlueprint>(InAsset.GetObject()))
@@ -79,7 +81,12 @@ FRigVMLocalVariableNameValidator::FRigVMLocalVariableNameValidator(const FRigVMA
 	}
 	else if (InAsset)
 	{
-		check(false);
+		FRigVMClient* Client = InAsset->GetRigVMClient();
+		TArray<URigVMGraph*> Graphs = Client->GetAllModels(true, true);
+		for (URigVMGraph* ItGraph : Graphs)
+		{
+			Names.Add(ItGraph->GetName());
+		}
 	}
 
 	if (Graph)
@@ -141,7 +148,7 @@ FRigVMNameValidator::FRigVMNameValidator(const UBlueprint* Blueprint, const UStr
 	}
 }
 
-FRigVMNameValidator::FRigVMNameValidator(const FRigVMAssetInterfacePtr Blueprint, const UStruct* ValidationScope, FName InExistingName)
+FRigVMNameValidator::FRigVMNameValidator(const FRigVMEditorAssetInterfacePtr Blueprint, const UStruct* ValidationScope, FName InExistingName)
 : FStringSetNameValidator(InExistingName.ToString())
 {
 	if (Blueprint)
@@ -202,7 +209,7 @@ FEdGraphPinType FRigVMEdGraphSchemaAction_LocalVar::GetPinType() const
 		{
 			if (Variable.Name == GetVariableName())
 			{
-				return Variable.ToPinType();
+				return RigVMTypeUtils::PinTypeFromRigVMVariableDescription(Variable);
 			}
 		}
 
@@ -210,7 +217,7 @@ FEdGraphPinType FRigVMEdGraphSchemaAction_LocalVar::GetPinType() const
 		{
 			if (Variable.Name == GetVariableName())
 			{
-				return Variable.ToPinType();
+				return RigVMTypeUtils::PinTypeFromRigVMVariableDescription(Variable);
 			}
 		}
 	}
@@ -316,7 +323,7 @@ UEdGraphNode* FRigVMEdGraphSchemaAction_PromoteToVariable::PerformAction(UEdGrap
 		return nullptr;
 	}
 
-	FRigVMAssetInterfacePtr Blueprint = RigGraph->GetAsset();
+	FRigVMEditorAssetInterfacePtr Blueprint = RigGraph->GetAsset();
 	URigVMGraph* Model = RigGraph->GetModel();
 	URigVMController* Controller = RigGraph->GetController();
 	if((Blueprint == nullptr) ||
@@ -337,14 +344,6 @@ UEdGraphNode* FRigVMEdGraphSchemaAction_PromoteToVariable::PerformAction(UEdGrap
 
 	if(bLocalVariable)
 	{
-#if WITH_EDITOR
-
-		if (GEditor)
-		{
-			GEditor->CancelTransaction(0);
-		}
-	
-#endif
 		const FRigVMGraphVariableDescription VariableDescription = Controller->AddLocalVariable(
 			*ModelPin->GetPinPath(),
 			ModelPin->GetCPPType(),
@@ -545,7 +544,7 @@ FRigVMEdGraphSchemaAction_Event::FRigVMEdGraphSchemaAction_Event(const FName& In
 
 FReply FRigVMEdGraphSchemaAction_Event::OnDoubleClick(UBlueprint* InBlueprint)
 {
-	if (FRigVMAssetInterfacePtr Blueprint = InBlueprint)
+	if (FRigVMEditorAssetInterfacePtr Blueprint = InBlueprint)
 	{
 		if(const URigVMNode* ModelNode = Blueprint->GetRigVMClient()->FindNode(NodePath))
 		{
@@ -623,7 +622,7 @@ FReply FRigVMFunctionDragDropAction::DroppedOnPanel(const TSharedRef< class SWid
 	// For functions
 	else if (URigVMEdGraph* TargetRigGraph = Cast<URigVMEdGraph>(&Graph))
 	{
-		if (FRigVMAssetInterfacePtr TargetRigBlueprint = FRigVMBlueprintUtils::FindAssetForGraph(TargetRigGraph))
+		if (FRigVMEditorAssetInterfacePtr TargetRigBlueprint = FRigVMBlueprintUtils::FindAssetForGraph(TargetRigGraph))
 		{
 			if (URigVMGraph* FunctionDefinitionGraph = SourceRigBlueprint->GetModel(SourceRigGraph))
 			{
@@ -633,7 +632,7 @@ FReply FRigVMFunctionDragDropAction::DroppedOnPanel(const TSharedRef< class SWid
 					{
 						if(URigVMFunctionLibrary* FunctionLibrary = Cast<URigVMFunctionLibrary>(FunctionDefinitionNode->GetOuter()))
 						{
-							if(FRigVMAssetInterfacePtr FunctionRigBlueprint = IRigVMAssetInterface::GetInterfaceOuter(FunctionLibrary))
+							if(FRigVMEditorAssetInterfacePtr FunctionRigBlueprint = IRigVMEditorAssetInterface::GetInterfaceOuter(FunctionLibrary))
 							{
 #if WITH_EDITOR
 								if(FunctionRigBlueprint != TargetRigBlueprint)
@@ -701,7 +700,7 @@ FRigVMFunctionDragDropAction::FRigVMFunctionDragDropAction()
 {
 }
 
-TSharedRef<FRigVMFunctionDragDropAction> FRigVMFunctionDragDropAction::New(TSharedPtr<FEdGraphSchemaAction> InAction, FRigVMAssetInterfacePtr InRigBlueprint, URigVMEdGraph* InRigGraph)
+TSharedRef<FRigVMFunctionDragDropAction> FRigVMFunctionDragDropAction::New(TSharedPtr<FEdGraphSchemaAction> InAction, FRigVMEditorAssetInterfacePtr InRigBlueprint, URigVMEdGraph* InRigGraph)
 {
 	TSharedRef<FRigVMFunctionDragDropAction> Action = MakeShareable(new FRigVMFunctionDragDropAction);
 	Action->SourceAction = InAction;
@@ -983,12 +982,22 @@ void URigVMEdGraphSchema::InsertAdditionalActions(TArray<UBlueprint*> InBlueprin
 	TArray<UEdGraphPin*> EdGraphPins, FGraphActionListBuilderBase& OutAllActions) const
 {
 	Super::InsertAdditionalActions(InBlueprints, EdGraphs, EdGraphPins, OutAllActions);
+	TArray<FRigVMEditorAssetInterfacePtr> Assets;
+	Algo::Transform(InBlueprints, Assets, [](UBlueprint* Blueprint)
+		{
+			return TScriptInterface<IRigVMEditorAssetInterface>(Blueprint);
+		});
+	InsertAdditionalActions(Assets, EdGraphs, EdGraphPins, OutAllActions);
+}
 
-	if(EdGraphPins.Num() > 0)
+void URigVMEdGraphSchema::InsertAdditionalActions(TArray<FRigVMEditorAssetInterfacePtr> InAssets, TArray<UEdGraph*> InGraphs,
+	TArray<UEdGraphPin*> InPins, FGraphActionListBuilderBase& OutAllActions) const
+{
+	if(InPins.Num() > 0)
 	{
-		if(const URigVMEdGraphNode* RigNode = Cast<URigVMEdGraphNode>(EdGraphPins[0]->GetOwningNode()))
+		if(const URigVMEdGraphNode* RigNode = Cast<URigVMEdGraphNode>(InPins[0]->GetOwningNode()))
 		{
-			if(const URigVMPin* ModelPin = RigNode->GetModelPinFromPinPath(EdGraphPins[0]->GetName()))
+			if(const URigVMPin* ModelPin = RigNode->GetModelPinFromPinPath(InPins[0]->GetName()))
 			{
 				const bool bIsRootGraph = ModelPin->GetGraph()->IsRootGraph();
 				if(!ModelPin->IsExecuteContext() && !ModelPin->IsWildCard())
@@ -996,13 +1005,13 @@ void URigVMEdGraphSchema::InsertAdditionalActions(TArray<UBlueprint*> InBlueprin
 					if(!ModelPin->GetNode()->IsA<URigVMVariableNode>())
 					{
 						OutAllActions.AddAction(TSharedPtr<FRigVMEdGraphSchemaAction_PromoteToVariable>(
-							new FRigVMEdGraphSchemaAction_PromoteToVariable(EdGraphPins[0], false)
+							new FRigVMEdGraphSchemaAction_PromoteToVariable(InPins[0], false)
 						));
 
 						if(!bIsRootGraph && !ModelPin->IsWildCard())
 						{
 							OutAllActions.AddAction(TSharedPtr<FRigVMEdGraphSchemaAction_PromoteToVariable>(
-								new FRigVMEdGraphSchemaAction_PromoteToVariable(EdGraphPins[0], true)
+								new FRigVMEdGraphSchemaAction_PromoteToVariable(InPins[0], true)
 							));
 						}
 					}
@@ -1013,7 +1022,7 @@ void URigVMEdGraphSchema::InsertAdditionalActions(TArray<UBlueprint*> InBlueprin
 					if (!ModelPin->GetGraph()->GetRootGraph()->IsA<URigVMFunctionLibrary>() || !ModelPin->IsWildCard())
 					{
 						OutAllActions.AddAction(TSharedPtr<FRigVMEdGraphSchemaAction_PromoteToExposedPin>(
-								   new FRigVMEdGraphSchemaAction_PromoteToExposedPin(EdGraphPins[0])
+								   new FRigVMEdGraphSchemaAction_PromoteToExposedPin(InPins[0])
 							   ));
 					}
 				}
@@ -1041,7 +1050,7 @@ TSharedPtr<INameValidatorInterface> URigVMEdGraphSchema::GetNameValidator(const
 	return MakeShareable(new FRigVMNameValidator(BlueprintObj, ValidationScope, OriginalName));		
 }
 
-TSharedPtr<INameValidatorInterface> URigVMEdGraphSchema::GetNameValidator(const FRigVMAssetInterfacePtr Asset, const FName& OriginalName, const UStruct* ValidationScope) const
+TSharedPtr<INameValidatorInterface> URigVMEdGraphSchema::GetNameValidator(const FRigVMEditorAssetInterfacePtr Asset, const FName& OriginalName, const UStruct* ValidationScope) const
 {
 	return MakeShareable(new FRigVMNameValidator(Asset, ValidationScope, OriginalName));
 }
@@ -1307,7 +1316,7 @@ FReply URigVMEdGraphSchema::BeginGraphDragAction(TSharedPtr<FEdGraphSchemaAction
 		FEdGraphSchemaAction_K2Graph* FuncAction = (FEdGraphSchemaAction_K2Graph*)InAction.Get();
 		if (URigVMEdGraph* RigGraph = Cast<URigVMEdGraph>(FuncAction->EdGraph))
 		{
-			if (FRigVMAssetInterfacePtr RigBlueprint = FRigVMBlueprintUtils::FindAssetForGraph(RigGraph))
+			if (FRigVMEditorAssetInterfacePtr RigBlueprint = FRigVMBlueprintUtils::FindAssetForGraph(RigGraph))
 			{
 				TSharedRef<FRigVMFunctionDragDropAction> Action = FRigVMFunctionDragDropAction::New(InAction, RigBlueprint, RigGraph);
 				Action->SetAltDrag(MouseEvent.IsAltDown());
@@ -1321,7 +1330,7 @@ FReply URigVMEdGraphSchema::BeginGraphDragAction(TSharedPtr<FEdGraphSchemaAction
 		FRigVMEdGraphSchemaAction_LocalVar* VarAction = (FRigVMEdGraphSchemaAction_LocalVar*)InAction.Get();
 		if (URigVMEdGraph* RigGraph = Cast<URigVMEdGraph>(VarAction->GetVariableScope()))
 		{
-			if (FRigVMAssetInterfacePtr RigBlueprint = FRigVMBlueprintUtils::FindAssetForGraph(RigGraph))
+			if (FRigVMEditorAssetInterfacePtr RigBlueprint = FRigVMBlueprintUtils::FindAssetForGraph(RigGraph))
 			{
 				TSharedRef<FRigVMFunctionDragDropAction> Action = FRigVMFunctionDragDropAction::New(InAction, RigBlueprint, RigGraph);
 				Action->SetAltDrag(MouseEvent.IsAltDown());
@@ -1336,7 +1345,7 @@ FReply URigVMEdGraphSchema::BeginGraphDragAction(TSharedPtr<FEdGraphSchemaAction
 FConnectionDrawingPolicy* URigVMEdGraphSchema::CreateConnectionDrawingPolicy(int32 InBackLayerID, int32 InFrontLayerID, float InZoomFactor, const FSlateRect& InClippingRect, class FSlateWindowElementList& InDrawElements, class UEdGraph* InGraphObj) const
 {
 #if WITH_EDITOR
-	if (const FRigVMAssetInterfacePtr Blueprint = FRigVMBlueprintUtils::FindAssetForGraph(InGraphObj))
+	if (const FRigVMEditorAssetInterfacePtr Blueprint = FRigVMBlueprintUtils::FindAssetForGraph(InGraphObj))
 	{
 		return Blueprint->GetEditorModule()->CreateConnectionDrawingPolicy(InBackLayerID, InFrontLayerID, InZoomFactor, InClippingRect, InDrawElements, InGraphObj);
 	}	
@@ -1544,7 +1553,7 @@ bool URigVMEdGraphSchema::GetLocalVariables(const UEdGraph* InGraph, TArray<FBPV
 				VariableDescription.VarName = LocalVariable.Name;
 				VariableDescription.FriendlyName = LocalVariable.Name.ToString();
 				VariableDescription.DefaultValue = LocalVariable.DefaultValue;
-				VariableDescription.VarType = LocalVariable.ToPinType();
+				VariableDescription.VarType = RigVMTypeUtils::PinTypeFromRigVMVariableDescription(LocalVariable);
 				VariableDescription.PropertyFlags |= CPF_BlueprintVisible;
 				OutLocalVariables.Add(VariableDescription);
 			}
@@ -1625,7 +1634,7 @@ FReply URigVMEdGraphSchema::TrySetGraphCategory(const UEdGraph* InGraph, const F
 {
 	if (const URigVMEdGraph* RigGraph = Cast<URigVMEdGraph>((UEdGraph*)InGraph))
 	{
-		if (const FRigVMAssetInterfacePtr RigBlueprint = FRigVMBlueprintUtils::FindAssetForGraph(RigGraph))
+		if (const FRigVMEditorAssetInterfacePtr RigBlueprint = FRigVMBlueprintUtils::FindAssetForGraph(RigGraph))
 		{
 			if (const URigVMGraph* Model = RigGraph->GetModel())
 			{
@@ -1649,7 +1658,7 @@ bool URigVMEdGraphSchema::TryDeleteGraph(UEdGraph* GraphToDelete) const
 {
 	if (const URigVMEdGraph* RigGraph = Cast<URigVMEdGraph>(GraphToDelete))
 	{
-		if (const FRigVMAssetInterfacePtr RigBlueprint = FRigVMBlueprintUtils::FindAssetForGraph(RigGraph))
+		if (const FRigVMEditorAssetInterfacePtr RigBlueprint = FRigVMBlueprintUtils::FindAssetForGraph(RigGraph))
 		{
 			if (const URigVMGraph* Model = RigBlueprint->GetModel(GraphToDelete))
 			{
@@ -1694,7 +1703,7 @@ bool URigVMEdGraphSchema::TryDeleteGraph(UEdGraph* GraphToDelete) const
 								{
 									if (Reference.IsValid())
 									{
-										const FRigVMAssetInterfacePtr OtherBlueprint = IRigVMAssetInterface::GetInterfaceOuter(Reference.Get());
+										const FRigVMEditorAssetInterfacePtr OtherBlueprint = IRigVMEditorAssetInterface::GetInterfaceOuter(Reference.Get());
 										if (OtherBlueprint != RigBlueprint)
 										{
 											if(RigBlueprint->OnRequestBulkEditDialog().IsBound())
@@ -1811,7 +1820,7 @@ bool URigVMEdGraphSchema::TryToGetChildEvents(const UEdGraph* Graph, const int32
 
 UEdGraphPin* URigVMEdGraphSchema::DropPinOnNode(UEdGraphNode* InTargetNode, const FName& InSourcePinName, const FEdGraphPinType& InSourcePinType, EEdGraphPinDirection InSourcePinDirection) const
 {
-	if (FRigVMAssetInterfacePtr RigBlueprint = FRigVMBlueprintUtils::FindAssetForNode(InTargetNode))
+	if (FRigVMEditorAssetInterfacePtr RigBlueprint = FRigVMBlueprintUtils::FindAssetForNode(InTargetNode))
 	{
 		if (URigVMEdGraphNode* RigNode = Cast<URigVMEdGraphNode>(InTargetNode))
 		{
@@ -2169,7 +2178,7 @@ bool URigVMEdGraphSchema::CanVariableBeDropped(UEdGraph* InGraph, FProperty* InV
 
 bool URigVMEdGraphSchema::RequestFunctionDropOnPanel(UEdGraph* InGraph, const FRigVMGraphFunctionIdentifier& InFunction, const FVector2D& InDropPosition, const FVector2D& InScreenPosition) const
 {
-	if (FRigVMAssetInterfacePtr TargetRigBlueprint = FRigVMBlueprintUtils::FindAssetForGraph(InGraph))
+	if (FRigVMEditorAssetInterfacePtr TargetRigBlueprint = FRigVMBlueprintUtils::FindAssetForGraph(InGraph))
 	{
 		if(URigVMController* TargetController = TargetRigBlueprint->GetController(InGraph))
 		{
@@ -2206,7 +2215,7 @@ bool URigVMEdGraphSchema::RequestVariableDropOnPanel(UEdGraph* InGraph, FPropert
 #if WITH_EDITOR
 	if (CanVariableBeDropped(InGraph, InVariableToDrop))
 	{
-		FRigVMAssetInterfacePtr RigBlueprint = FRigVMBlueprintUtils::FindAssetForGraph(InGraph);
+		FRigVMEditorAssetInterfacePtr RigBlueprint = FRigVMBlueprintUtils::FindAssetForGraph(InGraph);
 		if (RigBlueprint != nullptr)
 		{
 			RigBlueprint->OnVariableDropped().Broadcast(InGraph, InVariableToDrop, FDeprecateSlateVector2D(InDropPosition), FDeprecateSlateVector2D(InScreenPosition));
@@ -2573,7 +2582,7 @@ void URigVMEdGraphSchema::HandleModifiedEvent(ERigVMGraphNotifType InNotifType,
 
 TSubclassOf<URigVMEdGraphNode> URigVMEdGraphSchema::GetGraphNodeClass(const URigVMEdGraph* InGraph) const
 {
-	if (FRigVMAssetInterfacePtr RigBlueprint = FRigVMBlueprintUtils::FindAssetForGraph(InGraph))
+	if (FRigVMEditorAssetInterfacePtr RigBlueprint = FRigVMBlueprintUtils::FindAssetForGraph(InGraph))
 	{
 		return RigBlueprint->GetRigVMClientHost()->GetRigVMEdGraphNodeClass();
 	}
diff --git a/Engine/Plugins/Runtime/RigVM/Source/RigVMDeveloper/Private/RigVMBlueprintLegacy.cpp b/Engine/Plugins/Runtime/RigVM/Source/RigVMDeveloper/Private/RigVMBlueprintLegacy.cpp
--- a/Engine/Plugins/Runtime/RigVM/Source/RigVMDeveloper/Private/RigVMBlueprintLegacy.cpp
+++ b/Engine/Plugins/Runtime/RigVM/Source/RigVMDeveloper/Private/RigVMBlueprintLegacy.cpp
@@ -15,6 +15,7 @@
 #include "RigVMObjectVersion.h"
 #include "BlueprintCompilationManager.h"
 #include "IMessageLogListing.h"
+#include "RigVMDeveloperTypeUtils.h"
 #include "RigVMCompiler/RigVMCompiler.h"
 #include "RigVMCore/RigVMRegistry.h"
 #include "AssetRegistry/AssetRegistryModule.h"
@@ -86,7 +87,7 @@ bool FRigVMOldPublicFunctionData::IsMutable() const
 void URigVMBlueprint::BeginDestroy()
 {
 	Super::BeginDestroy();
-	Clear();
+	IRigVMEditorAssetInterface::BeginDestroy();
 }
 
 void URigVMBlueprint::UpdateSupportedEventNames()
@@ -149,12 +150,12 @@ FCompilerResultsLog URigVMBlueprint::CompileBlueprint()
 }
 
 URigVMBlueprint::URigVMBlueprint()
-	: IRigVMAssetInterface()
+	: IRigVMEditorAssetInterface()
 {
 }
 
 URigVMBlueprint::URigVMBlueprint(const FObjectInitializer& ObjectInitializer)
-	: IRigVMAssetInterface(ObjectInitializer)
+	: IRigVMEditorAssetInterface(ObjectInitializer)
 {
 
 #if WITH_EDITORONLY_DATA
@@ -163,7 +164,7 @@ URigVMBlueprint::URigVMBlueprint(const FObjectInitializer& ObjectInitializer)
 
 	bRecompileOnLoad = 0;
 	SupportedEventNames.Reset();
-	VMCompileSettings.ASTSettings.ReportDelegate.BindUObject(this, &IRigVMAssetInterface::HandleReportFromCompiler);
+	VMCompileSettings.ASTSettings.ReportDelegate.BindUObject(this, &IRigVMEditorAssetInterface::HandleReportFromCompiler);
 
 	bUpdatingExternalVariables = false;
 	
@@ -349,7 +350,7 @@ void URigVMBlueprint::Compile()
 
 void URigVMBlueprint::PatchVariableNodesOnLoad()
 {
-	IRigVMAssetInterface::PatchVariableNodesOnLoad();
+	IRigVMEditorAssetInterface::PatchVariableNodesOnLoad();
 #if WITH_EDITOR
 	LastNewVariables = NewVariables;
 #endif
@@ -378,6 +379,15 @@ bool URigVMBlueprint::IsPinBeingWatched(const UEdGraphPin* InPin) const
 	return FKismetDebugUtilities::IsPinBeingWatched(this, InPin);
 }
 
+bool URigVMBlueprint::NodeContainsWatchedPins(const UEdGraphNode* InNode) const
+{
+	UEdGraphPin* WatchedPin = FKismetDebugUtilities::FindPinWatchByPredicate(this, [InNode](const UEdGraphPin* InPin)
+		{
+			return InPin->GetOuter() == InNode;
+		});
+	return WatchedPin != nullptr;
+}
+
 void URigVMBlueprint::ForeachPinWatch(TFunctionRef<void(UEdGraphPin*)> Task)
 {
 	FKismetDebugUtilities::ForeachPinWatch(this, Task);
@@ -464,19 +474,19 @@ void URigVMBlueprint::AppendToClassSchema(FAppendToClassSchemaContext& Context)
 void URigVMBlueprint::PostEditChangeChainProperty(struct FPropertyChangedChainEvent& PropertyChangedEvent)
 {
 	UBlueprint::PostEditChangeChainProperty(PropertyChangedEvent);
-	IRigVMAssetInterface::PostEditChangeChainProperty(PropertyChangedEvent);
+	IRigVMEditorAssetInterface::PostEditChangeChainProperty(PropertyChangedEvent);
 }
 
 void URigVMBlueprint::PostRename(UObject* OldOuter, const FName OldName)
 {
 	UBlueprint::PostRename(OldOuter, OldName);
-	IRigVMAssetInterface::PostRename(OldOuter, OldName);
+	IRigVMEditorAssetInterface::PostRename(OldOuter, OldName);
 }
 
 void URigVMBlueprint::GetPreloadDependencies(TArray<UObject*>& OutDeps)
 {
 	UBlueprint::GetPreloadDependencies(OutDeps);
-	IRigVMAssetInterface::GetPreloadDependencies(OutDeps);
+	IRigVMEditorAssetInterface::GetPreloadDependencies(OutDeps);
 }
 
 void URigVMBlueprint::PostSaveRoot(FObjectPostSaveRootContext ObjectSaveContext)
@@ -494,7 +504,7 @@ void URigVMBlueprint::PostSaveRoot(FObjectPostSaveRootContext ObjectSaveContext)
 void URigVMBlueprint::PreSave(FObjectPreSaveContext ObjectSaveContext)
 {
 	UBlueprint::PreSave(ObjectSaveContext);
-	IRigVMAssetInterface::PreSave(ObjectSaveContext);
+	IRigVMEditorAssetInterface::PreSave(ObjectSaveContext);
 	if (URigVMBlueprintGeneratedClass* RigClass = GetRigVMBlueprintGeneratedClass())
 	{
 		if (const URigVMHost* CDO = Cast<URigVMHost>(RigClass->GetDefaultObject(true /* create if needed */)))
@@ -510,7 +520,7 @@ void URigVMBlueprint::PreSave(FObjectPreSaveContext ObjectSaveContext)
 void URigVMBlueprint::PostLoad()
 {
 	UBlueprint::PostLoad();
-	IRigVMAssetInterface::PostLoad();
+	IRigVMEditorAssetInterface::PostLoad();
 
 	FCoreUObjectDelegates::OnObjectModified.RemoveAll(this);
 	FCoreUObjectDelegates::OnObjectModified.AddUObject(this, &URigVMBlueprint::OnPreVariableChange);
@@ -524,19 +534,19 @@ void URigVMBlueprint::PostLoad()
 void URigVMBlueprint::DeclareConstructClasses(TArray<FTopLevelAssetPath>& OutConstructClasses, const UClass* SpecificSubclass)
 {
 	UBlueprint::DeclareConstructClasses(OutConstructClasses, SpecificSubclass);
-	IRigVMAssetInterface::DeclareConstructClasses(OutConstructClasses, SpecificSubclass);
+	IRigVMEditorAssetInterface::DeclareConstructClasses(OutConstructClasses, SpecificSubclass);
 }
 #endif
 
 
 void URigVMBlueprint::OnBlueprintChanged(UBlueprint* InBlueprint)
 {
-	IRigVMAssetInterface::OnChanged().Broadcast(InBlueprint);
+	IRigVMEditorAssetInterface::OnChanged().Broadcast(InBlueprint);
 }
 
 void URigVMBlueprint::OnSetObjectBeingDebuggedReceived(UObject* InObject)
 {
-	IRigVMAssetInterface::OnSetObjectBeingDebugged().Broadcast(InObject);
+	IRigVMEditorAssetInterface::OnSetObjectBeingDebugged().Broadcast(InObject);
 }
 
 void URigVMBlueprint::PreCompile()
@@ -554,7 +564,7 @@ void URigVMBlueprint::PreCompile()
 	}
 }
 
-FProperty* URigVMBlueprint::FindGeneratedPropertyByName(const FName& InName) const
+FProperty* URigVMBlueprint::FindGeneratedPropertyByName(const FName& InName)
 {
 	return SkeletonGeneratedClass->FindPropertyByName(InName);
 }
@@ -1272,7 +1282,7 @@ void URigVMBlueprint::GetTypeActions(FBlueprintActionDatabaseRegistrar& ActionRe
 	DECLARE_SCOPE_HIERARCHICAL_COUNTER_FUNC()
 
 #if WITH_EDITOR
-	GetEditorModule()->GetTypeActions(FRigVMAssetInterfacePtr(const_cast<URigVMBlueprint*>(this)), ActionRegistrar);
+	GetEditorModule()->GetTypeActions(FRigVMEditorAssetInterfacePtr(const_cast<URigVMBlueprint*>(this)), ActionRegistrar);
 #endif
 }
 
@@ -1281,7 +1291,7 @@ void URigVMBlueprint::GetInstanceActions(FBlueprintActionDatabaseRegistrar& Acti
 	DECLARE_SCOPE_HIERARCHICAL_COUNTER_FUNC()
 
 #if WITH_EDITOR
-	GetEditorModule()->GetInstanceActions(FRigVMAssetInterfacePtr(const_cast<URigVMBlueprint*>(this)), ActionRegistrar);
+	GetEditorModule()->GetInstanceActions(FRigVMEditorAssetInterfacePtr(const_cast<URigVMBlueprint*>(this)), ActionRegistrar);
 #endif
 }
 
diff --git a/Engine/Plugins/Runtime/RigVM/Source/RigVMDeveloper/Private/RigVMBlueprintUtils.cpp b/Engine/Plugins/Runtime/RigVM/Source/RigVMDeveloper/Private/RigVMBlueprintUtils.cpp
--- a/Engine/Plugins/Runtime/RigVM/Source/RigVMDeveloper/Private/RigVMBlueprintUtils.cpp
+++ b/Engine/Plugins/Runtime/RigVM/Source/RigVMDeveloper/Private/RigVMBlueprintUtils.cpp
@@ -7,7 +7,7 @@
 #include "Kismet2/BlueprintEditorUtils.h"
 #include "RigVMCore/RigVMStruct.h"
 #include "UObject/UObjectIterator.h"
-#include "RigVMAsset.h"
+#include "RigVMEditorAsset.h"
 #include "EdGraph/RigVMEdGraph.h"
 #include "EdGraph/RigVMEdGraphSchema.h"
 #include "Kismet2/Kismet2NameValidators.h"
@@ -15,7 +15,7 @@
 
 #define LOCTEXT_NAMESPACE "RigVMBlueprintUtils"
 
-FName FRigVMBlueprintUtils::ValidateName(FRigVMAssetInterfacePtr InBlueprint, const FString& InName)
+FName FRigVMBlueprintUtils::ValidateName(FRigVMEditorAssetInterfacePtr InBlueprint, const FString& InName)
 {
 	DECLARE_SCOPE_HIERARCHICAL_COUNTER_FUNC()
 
@@ -79,7 +79,7 @@ FName FRigVMBlueprintUtils::ValidateName(FRigVMAssetInterfacePtr InBlueprint, co
 	return *Name;
 }
 
-UEdGraphNode* FRigVMBlueprintUtils::GetNodeByGUID(const FRigVMAssetInterfacePtr InBlueprint, const FGuid& InNodeGuid)
+UEdGraphNode* FRigVMBlueprintUtils::GetNodeByGUID(const FRigVMEditorAssetInterfacePtr InBlueprint, const FGuid& InNodeGuid)
 {
 	TArray<UEdGraphNode*> Nodes;
 	TArray<UEdGraph*> AllGraphs;
@@ -102,11 +102,11 @@ UEdGraphNode* FRigVMBlueprintUtils::GetNodeByGUID(const FRigVMAssetInterfacePtr
 	return nullptr;
 }
 
-FRigVMAssetInterfacePtr FRigVMBlueprintUtils::FindAssetForGraph(const UEdGraph* Graph)
+FRigVMEditorAssetInterfacePtr FRigVMBlueprintUtils::FindAssetForGraph(const UEdGraph* Graph)
 {
 	for (UObject* TestOuter = Graph ? Graph->GetOuter() : nullptr; TestOuter; TestOuter = TestOuter->GetOuter())
 	{
-		if (TestOuter->Implements<URigVMAssetInterface>())
+		if (TestOuter->Implements<URigVMEditorAssetInterface>())
 		{
 			return TestOuter;
 		}
@@ -115,7 +115,7 @@ FRigVMAssetInterfacePtr FRigVMBlueprintUtils::FindAssetForGraph(const UEdGraph*
 	return nullptr;
 }
 
-FRigVMAssetInterfacePtr FRigVMBlueprintUtils::FindAssetForNode(const UEdGraphNode* Node)
+FRigVMEditorAssetInterfacePtr FRigVMBlueprintUtils::FindAssetForNode(const UEdGraphNode* Node)
 {
 	if (Node)
 	{
@@ -127,7 +127,7 @@ FRigVMAssetInterfacePtr FRigVMBlueprintUtils::FindAssetForNode(const UEdGraphNod
 	return nullptr;
 }
 
-FName FRigVMBlueprintUtils::FindUniqueVariableName(const IRigVMAssetInterface* InBlueprint, const FString& InBaseName)
+FName FRigVMBlueprintUtils::FindUniqueVariableName(const IRigVMEditorAssetInterface* InBlueprint, const FString& InBaseName)
 {
 	DECLARE_SCOPE_HIERARCHICAL_COUNTER_FUNC()
 
@@ -165,20 +165,20 @@ void FRigVMBlueprintUtils::ForAllRigVMStructs(TFunction<void(UScriptStruct*)> In
 
 void FRigVMBlueprintUtils::HandleReconstructAllBlueprintNodes(UBlueprint* InBlueprint)
 {
-	if (FRigVMAssetInterfacePtr Interface = InBlueprint)
+	if (FRigVMEditorAssetInterfacePtr Interface = InBlueprint)
 	{
 		HandleReconstructAllBlueprintNodes(Interface);
 	}
 }
 
-void FRigVMBlueprintUtils::HandleReconstructAllBlueprintNodes(FRigVMAssetInterfacePtr InBlueprint)
+void FRigVMBlueprintUtils::HandleReconstructAllBlueprintNodes(FRigVMEditorAssetInterfacePtr InBlueprint)
 {
 	DECLARE_SCOPE_HIERARCHICAL_COUNTER_FUNC()
 
 	return HandleRefreshAllNodes(InBlueprint);
 }
 
-void FRigVMBlueprintUtils::HandleRefreshAllNodes(FRigVMAssetInterfacePtr InBlueprint)
+void FRigVMBlueprintUtils::HandleRefreshAllNodes(FRigVMEditorAssetInterfacePtr InBlueprint)
 {
 	InBlueprint->RefreshAllNodes();
 }
@@ -187,7 +187,7 @@ void FRigVMBlueprintUtils::HandleRefreshAllBlueprintNodes(UBlueprint* InBlueprin
 {
 	DECLARE_SCOPE_HIERARCHICAL_COUNTER_FUNC()
 
-	if (FRigVMAssetInterfacePtr Blueprint = InBlueprint)
+	if (FRigVMEditorAssetInterfacePtr Blueprint = InBlueprint)
 	{
 		Blueprint->RefreshAllNodes();
 	}
@@ -195,7 +195,7 @@ void FRigVMBlueprintUtils::HandleRefreshAllBlueprintNodes(UBlueprint* InBlueprin
 
 void FRigVMBlueprintUtils::HandleAssetDeleted(const FAssetData& InAssetData)
 {
-	if (InAssetData.GetClass() && InAssetData.GetClass()->ImplementsInterface(URigVMAssetInterface::StaticClass()))
+	if (InAssetData.GetClass() && InAssetData.GetClass()->ImplementsInterface(URigVMEditorAssetInterface::StaticClass()))
 	{
 		// Make sure any RigVMBlueprint removes any TypeActions related to this asset (e.g. public functions)
 		if (FBlueprintActionDatabase* ActionDatabase = FBlueprintActionDatabase::TryGet())
@@ -205,7 +205,7 @@ void FRigVMBlueprintUtils::HandleAssetDeleted(const FAssetData& InAssetData)
 	}
 }
 
-void FRigVMBlueprintUtils::RemoveMemberVariableIfNotUsed(FRigVMAssetInterfacePtr Blueprint, const FName VarName)
+void FRigVMBlueprintUtils::RemoveMemberVariableIfNotUsed(FRigVMEditorAssetInterfacePtr Blueprint, const FName VarName)
 {
 	DECLARE_SCOPE_HIERARCHICAL_COUNTER_FUNC()
 
diff --git a/Engine/Plugins/Runtime/RigVM/Source/RigVMDeveloper/Private/RigVMCompiler/RigVMCompiler.cpp b/Engine/Plugins/Runtime/RigVM/Source/RigVMDeveloper/Private/RigVMCompiler/RigVMCompiler.cpp
--- a/Engine/Plugins/Runtime/RigVM/Source/RigVMDeveloper/Private/RigVMCompiler/RigVMCompiler.cpp
+++ b/Engine/Plugins/Runtime/RigVM/Source/RigVMDeveloper/Private/RigVMCompiler/RigVMCompiler.cpp
@@ -19,6 +19,7 @@
 #include "Algo/Count.h"
 #include "Algo/Sort.h"
 #include "String/Join.h"
+#include "RigVMDeveloperTypeUtils.h"
 
 #if WITH_EDITOR
 #include "Widgets/Notifications/SNotificationList.h"
@@ -3322,7 +3323,7 @@ bool URigVMCompiler::TraverseNoOp(const FRigVMNoOpExprAST* InExpr, FRigVMCompile
 					{
 						RegisterNameString = FString::Printf(TEXT("VariableDefault::%s|%lu::Const"), *VariableName.ToString(), (unsigned long)DefaultValueHash);
 					}
-					const FName RegisterName = FRigVMPropertyDescription::SanitizeName(*RegisterNameString);
+					const FName RegisterName = FRigVMPropertyDescription::SanitizeName(*RegisterNameString, false);
 
 					// set up a default value literal based on the value pin configuration
 					const FRigVMOperand RootDefaultValueOperand = WorkData.FindProperty(ERigVMMemoryType::Literal, RegisterName);
@@ -3332,12 +3333,12 @@ bool URigVMCompiler::TraverseNoOp(const FRigVMNoOpExprAST* InExpr, FRigVMCompile
 					if (VariableNode->IsLocalVariable())
 					{
 						VariableRegisterNameString = FString::Printf(TEXT("LocalVariable::%s|%s"), *VariableNode->GetGraph()->GetGraphName(), *VariableName.ToString());
-						const FName VariableRegisterName = FRigVMPropertyDescription::SanitizeName(*VariableRegisterNameString);
+						const FName VariableRegisterName = FRigVMPropertyDescription::SanitizeName(*VariableRegisterNameString, false);
 						VariableOperand = WorkData.FindProperty(ERigVMMemoryType::Work, VariableRegisterName);
 					}
 					else
 					{
-						VariableRegisterNameString = FString::Printf(TEXT("Variable::%s"), *VariableName.ToString());					
+						VariableRegisterNameString = FString::Printf(TEXT("Variable::%s"), *VariableName.ToString());
 						if (const FRigVMOperand* VariableOperandPtr = WorkData.PinPathToOperand->Find(VariableRegisterNameString))
 						{
 							VariableOperand = *VariableOperandPtr;
diff --git a/Engine/Plugins/Runtime/RigVM/Source/RigVMDeveloper/Private/RigVMDeveloperTypeUtils.cpp b/Engine/Plugins/Runtime/RigVM/Source/RigVMDeveloper/Private/RigVMDeveloperTypeUtils.cpp
--- a/Engine/Plugins/Runtime/RigVM/Source/RigVMDeveloper/Private/RigVMDeveloperTypeUtils.cpp
+++ b/Engine/Plugins/Runtime/RigVM/Source/RigVMDeveloper/Private/RigVMDeveloperTypeUtils.cpp
@@ -6,7 +6,7 @@
 #include "RigVMFunctions/RigVMDispatch_CastEnum.h"
 #include "RigVMFunctions/RigVMDispatch_CastObject.h"
 #include "RigVMModel/RigVMController.h"
-#include "RigVMModel/RigVMVariableDescription.h"
+#include "RigVMCore/RigVMVariableDescription.h"
 
 #define LOCTEXT_NAMESPACE "RigVMTypeUtils"
 
@@ -14,35 +14,6 @@
 
 #include "EdGraphSchema_K2.h"
 
-FRigVMExternalVariable RigVMTypeUtils::ExternalVariableFromRigVMVariableDescription(const FRigVMGraphVariableDescription& InVariableDescription)
-{	
-	FRigVMExternalVariable ExternalVariable;
-	ExternalVariable.Name = InVariableDescription.Name;
-
-	if (IsArrayType(InVariableDescription.CPPType))
-	{
-		ExternalVariable.bIsArray = true;
-		ExternalVariable.TypeName = *InVariableDescription.CPPType.Mid(7, InVariableDescription.CPPType.Len() - 8);
-		ExternalVariable.TypeObject = InVariableDescription.CPPTypeObject;
-	}
-	else
-	{
-		ExternalVariable.bIsArray = false;
-		ExternalVariable.TypeName = *InVariableDescription.CPPType;
-		ExternalVariable.TypeObject = InVariableDescription.CPPTypeObject;
-	}
-
-	if (UEnum* Enum = Cast<UEnum>(ExternalVariable.TypeObject))
-	{
-		ExternalVariable.TypeName = Enum->GetFName();
-	}	
-
-	ExternalVariable.bIsPublic = false;
-	ExternalVariable.bIsReadOnly = false;
-	ExternalVariable.Memory = nullptr;
-	return ExternalVariable;
-
-}
 
 FEdGraphPinType RigVMTypeUtils::PinTypeFromTypeIndex(const TRigVMTypeIndex& InTypeIndex)
 {
@@ -468,66 +439,6 @@ bool RigVMTypeUtils::CPPTypeFromPinType(const FEdGraphPinType& InPinType, FStrin
 	return false;
 }
 
-bool RigVMTypeUtils::CPPTypeFromExternalVariable(const FRigVMExternalVariable& InExternalVariable, FString& OutCPPType, UObject** OutCPPTypeObject)
-{
-	FString Prefix = "";
-	FString Suffix = "";
-	if (InExternalVariable.bIsArray)
-	{
-		Prefix = TEXT("TArray<");
-		Suffix = TEXT(">");
-	}
-
-	*OutCPPTypeObject = nullptr;
-	if (InExternalVariable.TypeName == BoolTypeName)
-	{
-		OutCPPType = Prefix + BoolType + Suffix;
-	}
-	else if (InExternalVariable.TypeName == Int32TypeName)
-	{
-		OutCPPType = Prefix + Int32Type + Suffix;
-	}
-	else if (InExternalVariable.TypeName == FloatTypeName)
-	{
-		OutCPPType = Prefix + FloatType + Suffix;
-	}
-	else if (InExternalVariable.TypeName == DoubleTypeName)
-	{
-		OutCPPType = Prefix + DoubleType + Suffix;
-	}
-	else if (InExternalVariable.TypeName == FNameTypeName)
-	{
-		OutCPPType = Prefix + FNameType + Suffix;
-	}
-	else if (InExternalVariable.TypeName == FStringTypeName)
-	{
-		OutCPPType = Prefix + FStringType + Suffix;
-	}
-	else if (UScriptStruct* Struct = Cast<UScriptStruct>(InExternalVariable.TypeObject))
-	{
-		OutCPPType = Prefix + *GetUniqueStructTypeName(Struct) + Suffix;
-		*OutCPPTypeObject = Struct;	
-	}
-	else if (UEnum* Enum = Cast<UEnum>(InExternalVariable.TypeObject))
-	{
-		OutCPPType = Prefix + Enum->GetFName().ToString() + Suffix;
-		*OutCPPTypeObject = Enum;
-	}
-	else if (UClass* Class = Cast<UClass>(InExternalVariable.TypeObject))
-	{
-		Prefix += TEXT("TObjectPtr<U");
-		Suffix += TEXT(">");
-		OutCPPType = Prefix + Class->GetFName().ToString() + Suffix;
-		*OutCPPTypeObject = Class;
-	}
-	else
-	{
-		return false;
-	}
-		
-	return true;
-}
-
 TRigVMTypeIndex RigVMTypeUtils::TypeIndexFromPinType(const FEdGraphPinType& InPinType)
 {
 	FString CPPType;
diff --git a/Engine/Plugins/Runtime/RigVM/Source/RigVMDeveloper/Private/RigVMJsonUtils.cpp b/Engine/Plugins/Runtime/RigVM/Source/RigVMDeveloper/Private/RigVMJsonUtils.cpp
--- a/Engine/Plugins/Runtime/RigVM/Source/RigVMDeveloper/Private/RigVMJsonUtils.cpp
+++ b/Engine/Plugins/Runtime/RigVM/Source/RigVMDeveloper/Private/RigVMJsonUtils.cpp
@@ -3,8 +3,9 @@
 #include "RigVMJsonUtils.h"
 #include "HAL/IConsoleManager.h"
 #include "RigVMDeveloperModule.h"
+#include "RigVMDeveloperTypeUtils.h"
 #include "RigVMCore/RigVMStruct.h"
-#include "RigVMAsset.h"
+#include "RigVMEditorAsset.h"
 #include "EdGraph/RigVMEdGraphSchema.h"
 #include "RigVMModel/RigVMGraph.h"
 #include "RigVMModel/RigVMSchema.h"
@@ -14,7 +15,7 @@
 #include "RigVMModel/Nodes/RigVMFunctionReferenceNode.h"
 #include "RigVMModel/Nodes/RigVMCollapseNode.h"
 #include "RigVMModel/Nodes/RigVMCommentNode.h"
-#include "RigVMModel/RigVMVariableDescription.h"
+#include "RigVMCore/RigVMVariableDescription.h"
 #include "Serialization/JsonSerializer.h"
 #include "UObject/UObjectIterator.h"
 
@@ -230,7 +231,7 @@ FAutoConsoleCommand FCmdRigVMJsonPrintAsset
 			return;
 		}
 
-		const IRigVMAssetInterface* Asset = Cast<IRigVMAssetInterface>(AssetObject);
+		const IRigVMEditorAssetInterface* Asset = Cast<IRigVMEditorAssetInterface>(AssetObject);
 		if (!Asset)
 		{
 			UE_LOG(LogRigVMDeveloper, Error, TEXT("Object '%s' exists, but it's not a RigVMAsset (when running RigVM.JSON.PrintAsset)"), *AssetPathName);
@@ -962,7 +963,7 @@ bool RigVMJsonUtils::GenericStructAttributesToJson(FIntrospectionSettings& InSet
 	return true;
 }
 
-TSharedPtr<FJsonObject> RigVMJsonUtils::ToJson(FIntrospectionSettings& InSettings, const IRigVMAssetInterface* InAsset)
+TSharedPtr<FJsonObject> RigVMJsonUtils::ToJson(FIntrospectionSettings& InSettings, const IRigVMEditorAssetInterface* InAsset)
 {
 	check(InAsset);
 	
@@ -1120,7 +1121,7 @@ TSharedPtr<FJsonObject> RigVMJsonUtils::ToJson(FIntrospectionSettings& InSetting
 	else if (const URigVMFunctionReferenceNode* FunctionReferenceNode = Cast<URigVMFunctionReferenceNode>(InNode))
 	{
 		FString LibraryNodePath = FunctionReferenceNode->GetFunctionIdentifier().GetLibraryNodePath();
-		if (const IRigVMAssetInterface* Asset = FunctionReferenceNode->GetImplementingOuter<IRigVMAssetInterface>())
+		if (const IRigVMEditorAssetInterface* Asset = FunctionReferenceNode->GetImplementingOuter<IRigVMEditorAssetInterface>())
 		{
 			if (const FRigVMClient* Client = Asset->GetRigVMClient())
 			{
diff --git a/Engine/Plugins/Runtime/RigVM/Source/RigVMDeveloper/Private/RigVMModel/Nodes/RigVMUnitNode.cpp b/Engine/Plugins/Runtime/RigVM/Source/RigVMDeveloper/Private/RigVMModel/Nodes/RigVMUnitNode.cpp
--- a/Engine/Plugins/Runtime/RigVM/Source/RigVMDeveloper/Private/RigVMModel/Nodes/RigVMUnitNode.cpp
+++ b/Engine/Plugins/Runtime/RigVM/Source/RigVMDeveloper/Private/RigVMModel/Nodes/RigVMUnitNode.cpp
@@ -7,6 +7,7 @@
 #include "RigVMUserWorkflowRegistry.h"
 #include "RigVMHost.h"
 #include "RigVMBlueprintLegacy.h"
+#include "RigVMDeveloperModule.h"
 
 #include UE_INLINE_GENERATED_CPP_BY_NAME(RigVMUnitNode)
 
@@ -484,7 +485,7 @@ void URigVMUnitNode::EnumeratePropertiesOnHostAndStructInstance(URigVMHost* InHo
 	check(InHost);
 	check(InInstance.IsValid() && InInstance->IsValid());
 
-	const IRigVMAssetInterface* Blueprint = GetImplementingOuter<IRigVMAssetInterface>();
+	const IRigVMEditorAssetInterface* Blueprint = GetImplementingOuter<IRigVMEditorAssetInterface>();
 	if(Blueprint == nullptr)
 	{
 		return;
@@ -768,7 +769,7 @@ void URigVMUnitNode::ComputePinValueDifferences(TSharedPtr<FStructOnScope> InDes
 
 bool URigVMUnitNode::IsPartOfRuntime() const
 {
-	if(const IRigVMAssetInterface* Blueprint = GetImplementingOuter<IRigVMAssetInterface>())
+	if(const IRigVMEditorAssetInterface* Blueprint = GetImplementingOuter<IRigVMEditorAssetInterface>())
 	{
 		if(URigVMHost* DebuggedHost = Cast<URigVMHost>(Blueprint->GetObjectBeingDebugged()))
 		{
diff --git a/Engine/Plugins/Runtime/RigVM/Source/RigVMDeveloper/Private/RigVMModel/RigVMBuildData.cpp b/Engine/Plugins/Runtime/RigVM/Source/RigVMDeveloper/Private/RigVMModel/RigVMBuildData.cpp
--- a/Engine/Plugins/Runtime/RigVM/Source/RigVMDeveloper/Private/RigVMModel/RigVMBuildData.cpp
+++ b/Engine/Plugins/Runtime/RigVM/Source/RigVMDeveloper/Private/RigVMModel/RigVMBuildData.cpp
@@ -1,7 +1,7 @@
 // Copyright Epic Games, Inc. All Rights Reserved.
 
 #include "RigVMModel/RigVMBuildData.h"
-#include "RigVMAsset.h"
+#include "RigVMEditorAsset.h"
 #include "RigVMBlueprintGeneratedClass.h"
 #include "AssetRegistry/AssetRegistryModule.h"
 #include "RigVMModel/RigVMFunctionLibrary.h"
@@ -78,7 +78,7 @@ TArray<UClass*> URigVMBuildData::FindAllRigVMAssetClasses()
 			ImplementedClasses.AddUnique(*ClassIterator);
 
 			UObject* CDO = ClassIterator->GetDefaultObject();
-			if(const FRigVMAssetInterfacePtr BlueprintCDO = CDO)
+			if(const FRigVMEditorAssetInterfacePtr BlueprintCDO = CDO)
 			{
 				if(UClass* GeneratedClass = BlueprintCDO->GetRigVMGeneratedClass())
 				{
@@ -179,7 +179,7 @@ void URigVMBuildData::TearDownRigVMGraphFunctionPointers()
 TArray<FRigVMGraphFunctionHeader> URigVMBuildData::GetFunctionHeadersForAsset(const FAssetData& InAssetData)
 {
 	TArray<FRigVMGraphFunctionHeader> Result;
-	if (!InAssetData.GetClass()->ImplementsInterface(URigVMAssetInterface::StaticClass()) &&
+	if (!InAssetData.GetClass()->ImplementsInterface(URigVMEditorAssetInterface::StaticClass()) &&
 		!InAssetData.IsInstanceOf(URigVMBlueprintGeneratedClass::StaticClass()))
 	{
 		return Result;
@@ -193,7 +193,7 @@ TArray<FRigVMGraphFunctionHeader> URigVMBuildData::GetFunctionHeadersForAsset(co
 		TScriptInterface<IRigVMGraphFunctionHost> GeneratedClass = AssetObject;
 		if (!GeneratedClass)
 		{
-			if (FRigVMAssetInterfacePtr RigVMBlueprint = AssetObject)
+			if (FRigVMEditorAssetInterfacePtr RigVMBlueprint = AssetObject)
 			{
 				GeneratedClass = RigVMBlueprint->GetRigVMGeneratedClass();
 			}
@@ -530,7 +530,7 @@ TArray<FRigVMVariantRef> URigVMBuildData::GatherAllFunctionVariantRefs() const
 	const FAssetRegistryModule& AssetRegistryModule = FModuleManager::LoadModuleChecked<FAssetRegistryModule>("AssetRegistry");
 	TArray<FAssetData> Assets;
 	FARFilter AssetFilter;
-	AssetFilter.ClassPaths.Add(URigVMAssetInterface::StaticClass()->GetClassPathName());
+	AssetFilter.ClassPaths.Add(URigVMEditorAssetInterface::StaticClass()->GetClassPathName());
 	AssetFilter.ClassPaths.Add(URigVMBlueprintGeneratedClass::StaticClass()->GetClassPathName());
 	AssetFilter.bRecursiveClasses = true;
 	AssetRegistryModule.Get().GetAssets(AssetFilter, Assets);
@@ -620,7 +620,7 @@ FRigVMVariantRef URigVMBuildData::CreateFunctionVariant(const FRigVMGraphFunctio
 	{
 		if(URigVMFunctionLibrary* FunctionLibrary = Cast<URigVMFunctionLibrary>(LibraryNode->GetGraph()))
 		{
-			if(IRigVMAssetInterface* RigVMBlueprint = FunctionLibrary->GetImplementingOuter<IRigVMAssetInterface>())
+			if(IRigVMEditorAssetInterface* RigVMBlueprint = FunctionLibrary->GetImplementingOuter<IRigVMEditorAssetInterface>())
 			{
 				if(URigVMController* Controller = RigVMBlueprint->GetOrCreateController(FunctionLibrary))
 				{
@@ -641,7 +641,7 @@ TArray<FRigVMVariantRef> URigVMBuildData::GatherAllAssetVariantRefs() const
 	const FAssetRegistryModule& AssetRegistryModule = FModuleManager::LoadModuleChecked<FAssetRegistryModule>("AssetRegistry");
 	TArray<FAssetData> Assets;
 	FARFilter AssetFilter;
-	AssetFilter.ClassPaths.Add(URigVMAssetInterface::StaticClass()->GetClassPathName());
+	AssetFilter.ClassPaths.Add(URigVMEditorAssetInterface::StaticClass()->GetClassPathName());
 	AssetFilter.bRecursiveClasses = true;
 	AssetRegistryModule.Get().GetAssets(AssetFilter, Assets);
 
@@ -698,7 +698,7 @@ FRigVMVariantRef URigVMBuildData::CreateAssetVariant(const FAssetData& InAssetDa
 	
 	if(UObject* DuplicatedAsset = EditorAssetSubsystem->DuplicateAsset(SourcePackageLongName, DestinationPackageLongName))
 	{
-		if(FRigVMAssetInterfacePtr RigVMBlueprint = DuplicatedAsset)
+		if(FRigVMEditorAssetInterfacePtr RigVMBlueprint = DuplicatedAsset)
 		{
 			// since we duplicated the asset the variant guid will be the same too,
 			// but to be sure we'll run the join code path as well.
@@ -725,13 +725,17 @@ FRigVMVariantRef URigVMBuildData::GetVariantRefForAsset(const FAssetData& InAsse
 	FRigVMVariant AssetVariant;
 	if(InAssetData.IsAssetLoaded())
 	{
-		if(const IRigVMAssetInterface* Blueprint = Cast<IRigVMAssetInterface>(InAssetData.GetAsset()))
+		if(const IRigVMRuntimeAssetInterface* RuntimeAsset = Cast<IRigVMRuntimeAssetInterface>(InAssetData.GetAsset()))
 		{
-			AssetVariant = Blueprint->GetAssetVariant();
+			AssetVariant = RuntimeAsset->GetAssetVariant();
+		}
+		else if(const IRigVMEditorAssetInterface* EditorAsset = Cast<IRigVMEditorAssetInterface>(InAssetData.GetAsset()))
+		{
+			AssetVariant = EditorAsset->GetAssetVariant();
 		}
 	}
-		
-	if (!AssetVariant.Guid.IsValid())
+
+	if (!AssetVariant.IsValid())
 	{
 		static const FName AssetVariantPropertyName = TEXT("AssetVariant");
 		const FProperty* AssetVariantProperty = CastField<FProperty>(InAssetData.GetClass()->FindPropertyByName(AssetVariantPropertyName));
@@ -761,9 +765,9 @@ FRigVMVariantRef URigVMBuildData::SplitVariantFromSet(const FRigVMVariantRef& In
 	if(Result.IsValid())
 	{
 		FAssetData AssetData = GetAssetDataForVariant(Result);
-		if (AssetData.GetClass()->ImplementsInterface(URigVMAssetInterface::StaticClass()))
+		if (AssetData.GetClass()->ImplementsInterface(URigVMEditorAssetInterface::StaticClass()))
 		{
-			FRigVMAssetInterfacePtr RigVMBlueprint(AssetData.GetAsset());
+			FRigVMEditorAssetInterfacePtr RigVMBlueprint(AssetData.GetAsset());
 			if(UObject* Subject = InVariantRef.ObjectPath.TryLoad())
 			{
 				if(Subject == RigVMBlueprint->GetObject())
@@ -798,9 +802,9 @@ FRigVMVariantRef URigVMBuildData::JoinVariantSet(const FRigVMVariantRef& InVaria
 	if(Result.IsValid() && InGuid.IsValid())
 	{
 		FAssetData AssetData = GetAssetDataForVariant(Result);
-		if (AssetData.GetClass()->ImplementsInterface(URigVMAssetInterface::StaticClass()))
+		if (AssetData.GetClass()->ImplementsInterface(URigVMEditorAssetInterface::StaticClass()))
 		{
-			FRigVMAssetInterfacePtr RigVMBlueprint = AssetData.GetAsset();
+			FRigVMEditorAssetInterfacePtr RigVMBlueprint = AssetData.GetAsset();
 			if(UObject* Subject = InVariantRef.ObjectPath.TryLoad())
 			{
 				if(Subject == RigVMBlueprint->GetObject())
diff --git a/Engine/Plugins/Runtime/RigVM/Source/RigVMDeveloper/Private/RigVMModel/RigVMClient.cpp b/Engine/Plugins/Runtime/RigVM/Source/RigVMDeveloper/Private/RigVMModel/RigVMClient.cpp
--- a/Engine/Plugins/Runtime/RigVM/Source/RigVMDeveloper/Private/RigVMModel/RigVMClient.cpp
+++ b/Engine/Plugins/Runtime/RigVM/Source/RigVMDeveloper/Private/RigVMModel/RigVMClient.cpp
@@ -2,7 +2,7 @@
 
 #include "RigVMModel/RigVMClient.h"
 
-#include "RigVMAsset.h"
+#include "RigVMEditorAsset.h"
 #include "Misc/TransactionObjectEvent.h"
 #include "Misc/ScopedSlowTask.h"
 #include "Exporters/Exporter.h"
@@ -751,14 +751,14 @@ URigVMFunctionLibrary* FRigVMClient::GetOrCreateFunctionLibrary(bool bSetupUndoR
 	return GetOrCreateFunctionLibrary(DefaultSchemaClass, bSetupUndoRedo, ObjectInitializer, bCreateController);
 }
 
-static void SetGetFunctionHostObjectPathDelegate(FRigVMClient* RigVMClient, URigVMFunctionLibrary* FunctionLibrary)
+void FRigVMClient::SetGetFunctionHostObjectPathDelegate(URigVMFunctionLibrary* InFunctionLibrary)
 {
-	if (RigVMClient->GetOuter()->Implements<URigVMClientHost>())
+	if (GetOuter()->Implements<URigVMClientHost>())
 	{
-		if (IRigVMClientHost* ClientHost = Cast<IRigVMClientHost>(RigVMClient->GetOuter()))
+		if (IRigVMClientHost* ClientHost = Cast<IRigVMClientHost>(GetOuter()))
 		{
 			TWeakObjectPtr<UObject> WeakClientHost = Cast<UObject>(ClientHost);
-			FunctionLibrary->GetFunctionHostObjectPathDelegate.BindLambda([WeakClientHost]() -> const FSoftObjectPath
+			InFunctionLibrary->GetFunctionHostObjectPathDelegate.BindLambda([WeakClientHost]() -> const FSoftObjectPath
 				{
 					if (WeakClientHost.IsValid())
 					{
@@ -779,7 +779,7 @@ URigVMFunctionLibrary* FRigVMClient::GetOrCreateFunctionLibrary(TSubclassOf<URig
 	{
 		if (!FunctionLibrary->GetFunctionHostObjectPathDelegate.IsBound())
 		{
-			SetGetFunctionHostObjectPathDelegate(this, FunctionLibrary);
+			SetGetFunctionHostObjectPathDelegate(FunctionLibrary);
 		}
 		return FunctionLibrary;
 	}
@@ -809,7 +809,7 @@ URigVMFunctionLibrary* FRigVMClient::GetOrCreateFunctionLibrary(TSubclassOf<URig
 
 	if (!NewFunctionLibrary->GetFunctionHostObjectPathDelegate.IsBound())
 	{
-		SetGetFunctionHostObjectPathDelegate(this, NewFunctionLibrary);
+		SetGetFunctionHostObjectPathDelegate(NewFunctionLibrary);
 	}
 
 	AddModel(NewFunctionLibrary, bCreateController);
diff --git a/Engine/Plugins/Runtime/RigVM/Source/RigVMDeveloper/Private/RigVMModel/RigVMController.cpp b/Engine/Plugins/Runtime/RigVM/Source/RigVMDeveloper/Private/RigVMModel/RigVMController.cpp
--- a/Engine/Plugins/Runtime/RigVM/Source/RigVMDeveloper/Private/RigVMModel/RigVMController.cpp
+++ b/Engine/Plugins/Runtime/RigVM/Source/RigVMDeveloper/Private/RigVMModel/RigVMController.cpp
@@ -29,6 +29,7 @@
 #include "RigVMPythonUtils.h"
 #include "RigVMSettings.h"
 #include "RigVMTypeUtils.h"
+#include "RigVMDeveloperTypeUtils.h"
 #include "RigVMStringUtils.h"
 #include "StructUtils/UserDefinedStruct.h"
 #include "RigVMFunctions/RigVMDispatch_If.h"
@@ -16027,7 +16028,7 @@ bool URigVMController::RenameLocalVariable(const FName& InVariableName, const FN
 		(void)Graph->MarkPackageDirty();
 		if (URigVMLibraryNode* LibraryNode = Cast<URigVMLibraryNode>(Graph->GetOuter()))
 		{
-			Notify(ERigVMGraphNotifType::LocalVariableDefaultValueChanged, LibraryNode);
+			Notify(ERigVMGraphNotifType::LocalVariableRenamed, LibraryNode);
 		}
 		for (URigVMNode* RenamedNode : RenamedNodes)
 		{
diff --git a/Engine/Plugins/Runtime/RigVM/Source/RigVMDeveloper/Private/RigVMModel/RigVMGraph.cpp b/Engine/Plugins/Runtime/RigVM/Source/RigVMDeveloper/Private/RigVMModel/RigVMGraph.cpp
--- a/Engine/Plugins/Runtime/RigVM/Source/RigVMDeveloper/Private/RigVMModel/RigVMGraph.cpp
+++ b/Engine/Plugins/Runtime/RigVM/Source/RigVMDeveloper/Private/RigVMModel/RigVMGraph.cpp
@@ -10,6 +10,7 @@
 #include "UObject/ObjectSaveContext.h"
 #include "RigVMTypeUtils.h"
 #include "RigVMModel/RigVMClient.h"
+#include "RigVMDeveloperModule.h"
 
 #if WITH_EDITOR
 #include "RigVMSettings.h"
diff --git a/Engine/Plugins/Runtime/RigVM/Source/RigVMDeveloper/Private/RigVMModel/RigVMNode.cpp b/Engine/Plugins/Runtime/RigVM/Source/RigVMDeveloper/Private/RigVMModel/RigVMNode.cpp
--- a/Engine/Plugins/Runtime/RigVM/Source/RigVMDeveloper/Private/RigVMModel/RigVMNode.cpp
+++ b/Engine/Plugins/Runtime/RigVM/Source/RigVMDeveloper/Private/RigVMModel/RigVMNode.cpp
@@ -12,6 +12,7 @@
 #include "RigVMUserWorkflowRegistry.h"
 #include "Logging/LogScopedVerbosityOverride.h"
 #include "Algo/Count.h"
+#include "RigVMModel/Nodes/RigVMDispatchNode.h"
 #include "Styling/AppStyle.h"
 #include "Textures/SlateIcon.h"
 
diff --git a/Engine/Plugins/Runtime/RigVM/Source/RigVMDeveloper/Public/EdGraph/RigVMEdGraph.h b/Engine/Plugins/Runtime/RigVM/Source/RigVMDeveloper/Public/EdGraph/RigVMEdGraph.h
--- a/Engine/Plugins/Runtime/RigVM/Source/RigVMDeveloper/Public/EdGraph/RigVMEdGraph.h
+++ b/Engine/Plugins/Runtime/RigVM/Source/RigVMDeveloper/Public/EdGraph/RigVMEdGraph.h
@@ -2,7 +2,7 @@
 
 #pragma once
 
-#include "RigVMAsset.h"
+#include "RigVMEditorAsset.h"
 #include "EdGraph/EdGraph.h"
 #include "EdGraph/RigVMEdGraphNode.h"
 #include "RigVMModel/RigVMGraph.h"
@@ -13,7 +13,7 @@
 
 #define UE_API RIGVMDEVELOPER_API
 
-class IRigVMAssetInterface;
+class IRigVMEditorAssetInterface;
 class URigVMEdGraphSchema;
 class URigVMController;
 class URigVMBlueprint;
@@ -146,13 +146,13 @@ public:
 	UE_API virtual void HandleRigVMGraphRenamed(const FString& InOldNodePath, const FString& InNewNodePath) override;
 
 	/** Set up this graph */
-	UE_DEPRECATED(5.7, "Please use const FRigVMAssetInterfacePtr GetAssetDefaultObject() const")
+	UE_DEPRECATED(5.7, "Please use const FRigVMEditorAssetInterfacePtr GetAssetDefaultObject() const")
 	UE_API const URigVMBlueprint* GetBlueprintDefaultObject() const;
-	UE_API const FRigVMAssetInterfacePtr GetAssetDefaultObject() const;
+	UE_API const FRigVMEditorAssetInterfacePtr GetAssetDefaultObject() const;
 	UE_API void SetBlueprintClass(const UClass* InClass);
-	UE_DEPRECATED(5.7, "Please use void InitializeFromAsset(FRigVMAssetInterfacePtr InAsset);")
-	UE_API virtual void InitializeFromBlueprint(URigVMBlueprint* InBlueprint) {}
-	UE_API virtual void InitializeFromAsset(FRigVMAssetInterfacePtr InAsset);
+	UE_DEPRECATED(5.7, "Please use void InitializeFromAsset(FRigVMEditorAssetInterfacePtr InAsset);")
+	virtual void InitializeFromBlueprint(URigVMBlueprint* InBlueprint) {}
+	UE_API virtual void InitializeFromAsset(FRigVMEditorAssetInterfacePtr InAsset);
 	UE_API bool IsPreviewGraph() const;
 
 	/** Get the ed graph schema */
@@ -171,9 +171,9 @@ public:
 
 	UE_API UEdGraphNode* FindNodeForModelNodeName(const FName& InModelNodeName, const bool bCacheIfRequired = true);
 
-	UE_DEPRECATED(5.7, "Please use FRigVMAssetInterfacePtr GetAsset() const")
+	UE_DEPRECATED(5.7, "Please use FRigVMEditorAssetInterfacePtr GetAsset() const")
 	UE_API URigVMBlueprint* GetBlueprint() const;
-	UE_API FRigVMAssetInterfacePtr GetAsset() const;
+	UE_API FRigVMEditorAssetInterfacePtr GetAsset() const;
 	UE_API URigVMGraph* GetModel() const;
 	UE_API URigVMController* GetController() const;
 	bool IsRootGraph() const { return GetRootGraph() == this; }
@@ -220,7 +220,7 @@ private:
 
 #endif
 	friend class URigVMEdGraphNode;
-	friend class IRigVMAssetInterface;
+	friend class IRigVMEditorAssetInterface;
 	friend class FRigVMEditorBase;
 	friend class SRigVMGraphNode;
 
diff --git a/Engine/Plugins/Runtime/RigVM/Source/RigVMDeveloper/Public/EdGraph/RigVMEdGraphNode.h b/Engine/Plugins/Runtime/RigVM/Source/RigVMDeveloper/Public/EdGraph/RigVMEdGraphNode.h
--- a/Engine/Plugins/Runtime/RigVM/Source/RigVMDeveloper/Public/EdGraph/RigVMEdGraphNode.h
+++ b/Engine/Plugins/Runtime/RigVM/Source/RigVMDeveloper/Public/EdGraph/RigVMEdGraphNode.h
@@ -4,16 +4,15 @@
 #pragma once
 
 #include "CoreMinimal.h"
+#include "RigVMEditorAsset.h"
 #include "Kismet2/Kismet2NameValidators.h"
-#include "RigVMAsset.h"
 #include "UObject/WeakObjectPtrTemplates.h"
-#include "RigVMModel/RigVMGraph.h"
 #include "Blueprint/BlueprintExtension.h"
 #include "RigVMEdGraphNode.generated.h"
 
 #define UE_API RIGVMDEVELOPER_API
 
-class IRigVMAssetInterface;
+class IRigVMEditorAssetInterface;
 class URigVMBlueprint;
 class FBlueprintActionDatabaseRegistrar;
 class UEdGraph;
@@ -159,9 +158,9 @@ public:
 	UE_API void CopyPinDefaultsToModel(UEdGraphPin* Pin, bool bUndo = false, bool bPrintPythonCommand = false);
 
 	/** Get the blueprint that this node is contained within */
-	UE_DEPRECATED(5.7, "Please use FRigVMAssetInterfacePtr GetAsset() const")
+	UE_DEPRECATED(5.7, "Please use FRigVMEditorAssetInterfacePtr GetAsset() const")
 	UE_API URigVMBlueprint* GetBlueprint() const;
-	UE_API FRigVMAssetInterfacePtr GetAsset() const;
+	UE_API FRigVMEditorAssetInterfacePtr GetAsset() const;
 
 	/** Get the VM model this node lives within */
 	UE_API URigVMGraph* GetModel() const;
diff --git a/Engine/Plugins/Runtime/RigVM/Source/RigVMDeveloper/Public/EdGraph/RigVMEdGraphSchema.h b/Engine/Plugins/Runtime/RigVM/Source/RigVMDeveloper/Public/EdGraph/RigVMEdGraphSchema.h
--- a/Engine/Plugins/Runtime/RigVM/Source/RigVMDeveloper/Public/EdGraph/RigVMEdGraphSchema.h
+++ b/Engine/Plugins/Runtime/RigVM/Source/RigVMDeveloper/Public/EdGraph/RigVMEdGraphSchema.h
@@ -2,7 +2,7 @@
 
 #pragma once
 
-#include "RigVMAsset.h"
+#include "RigVMEditorAsset.h"
 #include "EdGraph/RigVMEdGraphNode.h"
 #include "EdGraph/EdGraphSchema.h"
 #include "GraphEditorDragDropAction.h"
@@ -87,7 +87,7 @@ class FRigVMLocalVariableNameValidator : public FStringSetNameValidator
 {
 
 public:
-	UE_API FRigVMLocalVariableNameValidator(const FRigVMAssetInterfacePtr Blueprint, const URigVMGraph* Graph, FName InExistingName = NAME_None);
+	UE_API FRigVMLocalVariableNameValidator(const FRigVMEditorAssetInterfacePtr Blueprint, const URigVMGraph* Graph, FName InExistingName = NAME_None);
 
 	// Begin FNameValidatorInterface
 	UE_API virtual EValidatorResult IsValid(const FString& Name, bool bOriginal) override;
@@ -102,7 +102,7 @@ class FRigVMNameValidator : public FStringSetNameValidator
 
 public:
 	UE_API FRigVMNameValidator(const class UBlueprint* Blueprint, const UStruct* ValidationScope, FName InExistingName = NAME_None);
-	UE_API FRigVMNameValidator(const FRigVMAssetInterfacePtr Blueprint, const UStruct* ValidationScope, FName InExistingName = NAME_None);
+	UE_API FRigVMNameValidator(const FRigVMEditorAssetInterfacePtr Blueprint, const UStruct* ValidationScope, FName InExistingName = NAME_None);
 
 	// Begin FNameValidatorInterface
 	UE_API virtual EValidatorResult IsValid(const FString& Name, bool bOriginal) override;
@@ -283,11 +283,11 @@ protected:
 	/** Constructor */
 	UE_API FRigVMFunctionDragDropAction();
 
-	static UE_API TSharedRef<FRigVMFunctionDragDropAction> New(TSharedPtr<FEdGraphSchemaAction> InAction, FRigVMAssetInterfacePtr InRigBlueprint, URigVMEdGraph* InRigGraph);
+	static UE_API TSharedRef<FRigVMFunctionDragDropAction> New(TSharedPtr<FEdGraphSchemaAction> InAction, FRigVMEditorAssetInterfacePtr InRigBlueprint, URigVMEdGraph* InRigGraph);
 
 protected:
 
-	TScriptInterface<IRigVMAssetInterface> SourceRigBlueprint;
+	TScriptInterface<IRigVMEditorAssetInterface> SourceRigBlueprint;
 	URigVMEdGraph* SourceRigGraph;
 	bool bControlDrag;
 	bool bAltDrag;
@@ -360,8 +360,9 @@ public:
 	UE_API virtual bool SupportsDropPinOnNode(UEdGraphNode* InTargetNode, const FEdGraphPinType& InSourcePinType, EEdGraphPinDirection InSourcePinDirection, FText& OutErrorMessage) const override;
 	virtual void SetPinBeingDroppedOnNode(UEdGraphPin* InSourcePin) const override { PinBeingDropped = InSourcePin; }
 	UE_API virtual void InsertAdditionalActions(TArray<UBlueprint*> InBlueprints, TArray<UEdGraph*> InGraphs, TArray<UEdGraphPin*> InPins, FGraphActionListBuilderBase& OutAllActions) const override;
+	UE_API virtual void InsertAdditionalActions(TArray<FRigVMEditorAssetInterfacePtr> InAssets, TArray<UEdGraph*> InGraphs, TArray<UEdGraphPin*> InPins, FGraphActionListBuilderBase& OutAllActions) const;
 	UE_API virtual TSharedPtr<INameValidatorInterface> GetNameValidator(const UBlueprint* BlueprintObj, const FName& OriginalName, const UStruct* ValidationScope, const FName& ActionTypeId) const override;
-	UE_API virtual TSharedPtr<INameValidatorInterface> GetNameValidator(const FRigVMAssetInterfacePtr Asset, const FName& OriginalName, const UStruct* ValidationScope) const;
+	UE_API virtual TSharedPtr<INameValidatorInterface> GetNameValidator(const FRigVMEditorAssetInterfacePtr Asset, const FName& OriginalName, const UStruct* ValidationScope) const;
 	UE_API bool CanPromotePinToVariable(const UEdGraphPin& Pin, bool bCond) const;
 	UE_API virtual bool SupportsPreviewFromHereOnNode(const UEdGraphNode* InNode) const;
 
