Control Rig: Split IRigVMAssetInterface into IRigVMEditorAssetInterface and IRigVMRuntimeAssetInterface.

diff --git a/Engine/Plugins/Runtime/RigVM/Source/RigVMDeveloper/Public/RigVMBlueprintLegacy.h b/Engine/Plugins/Runtime/RigVM/Source/RigVMDeveloper/Public/RigVMBlueprintLegacy.h
--- a/Engine/Plugins/Runtime/RigVM/Source/RigVMDeveloper/Public/RigVMBlueprintLegacy.h
+++ b/Engine/Plugins/Runtime/RigVM/Source/RigVMDeveloper/Public/RigVMBlueprintLegacy.h
@@ -16,7 +16,8 @@
 #include "HAL/CriticalSection.h"
 #endif
 
-#include "RigVMAsset.h"
+#include "RigVMEditorAsset.h"
+#include "RigVMModel/RigVMExternalDependency.h"
 #include "UObject/AssetRegistryTagsContext.h"
 #include "UObject/ObjectSaveContext.h"
 
@@ -28,7 +29,7 @@
 class URigVMBlueprintGeneratedClass;
 
 USTRUCT(meta = (Deprecated = "5.2"))
-struct RIGVMDEVELOPER_API FRigVMOldPublicFunctionArg
+struct FRigVMOldPublicFunctionArg
 {
 	GENERATED_BODY();
 	
@@ -55,11 +56,11 @@ struct RIGVMDEVELOPER_API FRigVMOldPublicFunctionArg
 	UPROPERTY()
 	ERigVMPinDirection Direction;
 
-	FEdGraphPinType GetPinType() const;
+	UE_API FEdGraphPinType GetPinType() const;
 };
 
 USTRUCT(meta = (Deprecated = "5.2"))
-struct RIGVMDEVELOPER_API FRigVMOldPublicFunctionData
+struct FRigVMOldPublicFunctionData
 {
 	GENERATED_BODY();
 
@@ -67,7 +68,7 @@ struct RIGVMDEVELOPER_API FRigVMOldPublicFunctionData
 		:Name(NAME_None)
 	{}
 
-	virtual ~FRigVMOldPublicFunctionData();
+	UE_API virtual ~FRigVMOldPublicFunctionData();
 
 	UPROPERTY()
 	FName Name;
@@ -87,16 +88,16 @@ struct RIGVMDEVELOPER_API FRigVMOldPublicFunctionData
 	UPROPERTY()
 	TArray<FRigVMOldPublicFunctionArg> Arguments;
 
-	bool IsMutable() const;
+	UE_API bool IsMutable() const;
 };
 
-UCLASS(BlueprintType, meta=(IgnoreClassThumbnail))
-class UE_API URigVMBlueprint : public UBlueprint, public IRigVMAssetInterface, public IRigVMClientHost, public IRigVMExternalDependencyManager
+UCLASS(MinimalAPI, BlueprintType, meta=(IgnoreClassThumbnail))
+class URigVMBlueprint : public UBlueprint, public IRigVMEditorAssetInterface, public IRigVMClientHost, public IRigVMExternalDependencyManager
 {
 	GENERATED_UCLASS_BODY()
 
 public:
-	virtual void BeginDestroy() override;
+	UE_API virtual void BeginDestroy() override;
 	virtual UObject* GetObject() override { return this; }
 	virtual TArray<FRigVMGraphFunctionHeader>& GetPublicGraphFunctions() override { return PublicGraphFunctions; }
 	virtual FRigVMClient* GetRigVMClient() override { return &RigVMClient; }
@@ -106,103 +107,104 @@ public:
 	virtual FRigVMCompileSettings& GetVMCompileSettings() override { return VMCompileSettings; }
 	virtual FRigVMRuntimeSettings& GetVMRuntimeSettings() override { return VMRuntimeSettings; }
 	virtual bool IsRegeneratingOnLoad() const override { return bIsRegeneratingOnLoad;}
-	virtual TArray<FRigVMGraphVariableDescription> GetAssetVariables() const override;
+	UE_API virtual TArray<FRigVMGraphVariableDescription> GetAssetVariables() const override;
 	virtual TArray<struct FEditedDocumentInfo>& GetLastEditedDocuments() override { return LastEditedDocuments;}
-	virtual const TArray<FString>& GetRequiredPlugins(bool bRefresh = true) const override;
+	UE_API virtual const TArray<FString>& GetRequiredPlugins(bool bRefresh = true) const override;
+	virtual FRigVMEdGraphDisplaySettings& GetRigGraphDisplaySettings() override { return RigGraphDisplaySettings; }
 
 #if WITH_EDITORONLY_DATA
-	static void AppendToClassSchema(FAppendToClassSchemaContext& Context);
+	UE_API static void AppendToClassSchema(FAppendToClassSchemaContext& Context);
 #endif
 
 protected:
 	virtual UClass* GetRigVMGeneratedClass() override { return GeneratedClass; }
-	virtual TArray<FRigVMExternalVariable> GetExternalVariables(bool bFallbackToBlueprint) const override;
-	virtual FString GetVariableDefaultValue(const FName& InName, bool bFromDebuggedObject) const override;
-	virtual URigVM* GetVM(bool bCreateIfNeeded = true) const override;
-	virtual FRigVMExtendedExecuteContext* GetRigVMExtendedExecuteContext() override;
+	UE_API virtual TArray<FRigVMExternalVariable> GetExternalVariables(bool bFallbackToBlueprint) const override;
+	UE_API virtual FString GetVariableDefaultValue(const FName& InName, bool bFromDebuggedObject) const override;
+	UE_API virtual URigVM* GetVM(bool bCreateIfNeeded = true) const override;
+	UE_API virtual FRigVMExtendedExecuteContext* GetRigVMExtendedExecuteContext() override;
 	virtual TScriptInterface<IRigVMGraphFunctionHost> GetRigVMGraphFunctionHost() override { return GetRigVMBlueprintGeneratedClass(); }
 	virtual TScriptInterface<const IRigVMGraphFunctionHost> GetRigVMGraphFunctionHost() const override { return GetRigVMBlueprintGeneratedClass(); }
 	virtual FCompilerResultsLog* GetCurrentMessageLog() const override { return CurrentMessageLog; }
-	virtual void SetAssetStatus(const ERigVMAssetStatus& InStatus) override;
-	virtual ERigVMAssetStatus GetAssetStatus() const override;
+	UE_API virtual void SetAssetStatus(const ERigVMAssetStatus& InStatus) override;
+	UE_API virtual ERigVMAssetStatus GetAssetStatus() const override;
 	virtual bool IsUpToDate() const override { return UBlueprint::IsUpToDate(); }
-	virtual TArray<UObject*> GetArchetypeInstances(bool bIncludeCDO, bool bIncludeDerivedClasses) const override;
+	UE_API virtual TArray<UObject*> GetArchetypeInstances(bool bIncludeCDO, bool bIncludeDerivedClasses) const override;
 	virtual const FString& GetObjectBeingDebuggedPath() const override { return UBlueprint::GetObjectPathToDebug(); }
 	virtual UWorld* GetWorldBeingDebugged() const override { return UBlueprint::GetWorldBeingDebugged(); }
 	virtual void SetWorldBeingDebugged(UWorld* NewWorld) override { return UBlueprint::SetWorldBeingDebugged(NewWorld); }
-	virtual FRigVMDebugInfo& GetDebugInfo() override;
+	UE_API virtual FRigVMDebugInfo& GetDebugInfo() override;
 	virtual TObjectPtr<URigVMEdGraph>& GetFunctionLibraryEdGraph() override { return FunctionLibraryEdGraph; }
-	virtual URigVMHost* CreateRigVMHostSuper(UObject* InOuter) override;
-	virtual void MarkAssetAsStructurallyModified(bool bSkipDirtyAssetStatus = false) override;
-	virtual void MarkAssetAsModified(FPropertyChangedEvent PropertyChangedEvent = FPropertyChangedEvent(nullptr)) override;
-	virtual void AddUbergraphPage(URigVMEdGraph* RigVMEdGraph) override;
-	virtual void AddLastEditedDocument(URigVMEdGraph* RigVMEdGraph) override;
-	virtual void Compile() override;
-	virtual FCompilerResultsLog CompileBlueprint() override;
-	virtual void PatchVariableNodesOnLoad() override;
-	virtual void AddPinWatch(UEdGraphPin* InPin) override;
-	virtual void RemovePinWatch(UEdGraphPin* InPin) override;
-	virtual void ClearPinWatches() override;
-	virtual bool IsPinBeingWatched(const UEdGraphPin* InPin) const override;
-	virtual void ForeachPinWatch(TFunctionRef<void(UEdGraphPin*)> Task) override;
-	virtual TMap<FString, FSoftObjectPath>& GetUserDefinedStructGuidToPathName(bool bFromCDO) override;
-	virtual TMap<FString, FSoftObjectPath>& GetUserDefinedEnumToPathName(bool bFromCDO) override;
-	virtual TSet<TObjectPtr<UObject>>& GetUserDefinedTypesInUse(bool bFromCDO) override;
-	virtual FRigVMEdGraphDisplaySettings& GetRigGraphDisplaySettings() override { return RigGraphDisplaySettings; }
+	UE_API virtual URigVMHost* CreateRigVMHostSuper(UObject* InOuter) override;
+	UE_API virtual void MarkAssetAsStructurallyModified(bool bSkipDirtyAssetStatus = false) override;
+	UE_API virtual void MarkAssetAsModified(FPropertyChangedEvent PropertyChangedEvent = FPropertyChangedEvent(nullptr)) override;
+	UE_API virtual void AddUbergraphPage(URigVMEdGraph* RigVMEdGraph) override;
+	UE_API virtual void AddLastEditedDocument(URigVMEdGraph* RigVMEdGraph) override;
+	UE_API virtual void Compile() override;
+	UE_API virtual FCompilerResultsLog CompileBlueprint() override;
+	UE_API virtual void PatchVariableNodesOnLoad() override;
+	UE_API virtual void AddPinWatch(UEdGraphPin* InPin) override;
+	UE_API virtual void RemovePinWatch(UEdGraphPin* InPin) override;
+	UE_API virtual void ClearPinWatches() override;
+	UE_API virtual bool IsPinBeingWatched(const UEdGraphPin* InPin) const override;
+	UE_API virtual bool NodeContainsWatchedPins(const UEdGraphNode* InNode) const override;
+	UE_API virtual void ForeachPinWatch(TFunctionRef<void(UEdGraphPin*)> Task) override;
+	UE_API virtual TMap<FString, FSoftObjectPath>& GetUserDefinedStructGuidToPathName(bool bFromCDO) override;
+	UE_API virtual TMap<FString, FSoftObjectPath>& GetUserDefinedEnumToPathName(bool bFromCDO) override;
+	UE_API virtual TSet<TObjectPtr<UObject>>& GetUserDefinedTypesInUse(bool bFromCDO) override;
 	virtual TArray<TObjectPtr<UEdGraph>>& GetFunctionGraphs() override { return FunctionGraphs; }
 	virtual bool& IsReferencedObjectPathsStored() override { return ReferencedObjectPathsStored; }
 	virtual TArray<FSoftObjectPath>& GetReferencedObjectPaths() override { return ReferencedObjectPaths; }
 	virtual TArray<FName> GetSupportedEventNames() override { return SupportedEventNames; }
-	virtual void UpdateSupportedEventNames() override;
+	UE_API virtual void UpdateSupportedEventNames() override;
 	virtual TArray<FRigVMReferenceNodeData>& GetFunctionReferenceNodeData() override { return FunctionReferenceNodeData; }
 	virtual void NotifyGraphRenamedSuper(class UEdGraph* Graph, FName OldName, FName NewName) override { UBlueprint::NotifyGraphRenamed(Graph, OldName, NewName); }
-	virtual UObject* GetDefaultsObject() override;
-	virtual void PostEditChangeBlueprintActors() override;
+	UE_API virtual UObject* GetDefaultsObject() override;
+	UE_API virtual void PostEditChangeBlueprintActors() override;
 
 public:
-	URigVMBlueprint();
+	UE_API URigVMBlueprint();
 
 	/** Get the (full) generated class for this rigvm blueprint */
-	URigVMBlueprintGeneratedClass* GetRigVMBlueprintGeneratedClass() const;
+	UE_API URigVMBlueprintGeneratedClass* GetRigVMBlueprintGeneratedClass() const;
 	/** Returns the class used as the super class for all generated classes */
 	virtual UClass* GetRigVMGeneratedClassPrototype() const { return URigVMBlueprintGeneratedClass::StaticClass(); }
 
-	virtual void Serialize(FArchive& Ar) override { IRigVMAssetInterface ::SerializeImpl(Ar); }
+	virtual void Serialize(FArchive& Ar) override { IRigVMEditorAssetInterface ::SerializeImpl(Ar); }
 	virtual void SerializeSuper(FArchive& Ar) override { UBlueprint ::Serialize(Ar); }
-	virtual void PostSerialize(FArchive& Ar) override;
+	UE_API virtual void PostSerialize(FArchive& Ar) override;
 
 #if WITH_EDITOR
 
 	// UBlueprint interface
-	virtual UClass* GetBlueprintClass() const override;
-	virtual UClass* RegenerateClass(UClass* ClassToRegenerate, UObject* PreviousCDO) override;
+	UE_API virtual UClass* GetBlueprintClass() const override;
+	UE_API virtual UClass* RegenerateClass(UClass* ClassToRegenerate, UObject* PreviousCDO) override;
 	virtual bool SupportedByDefaultBlueprintFactory() const override { return false; }
 	virtual bool IsValidForBytecodeOnlyRecompile() const override { return false; }
 	virtual void LoadModulesRequiredForCompilation() override {}
-	virtual void GetTypeActions(FBlueprintActionDatabaseRegistrar& ActionRegistrar) const override;
-	virtual void GetInstanceActions(FBlueprintActionDatabaseRegistrar& ActionRegistrar) const override;
-	virtual void SetObjectBeingDebugged(UObject* NewObject) override { IRigVMAssetInterface::SetObjectBeingDebugged(NewObject); }
+	UE_API virtual void GetTypeActions(FBlueprintActionDatabaseRegistrar& ActionRegistrar) const override;
+	UE_API virtual void GetInstanceActions(FBlueprintActionDatabaseRegistrar& ActionRegistrar) const override;
+	virtual void SetObjectBeingDebugged(UObject* NewObject) override { IRigVMEditorAssetInterface::SetObjectBeingDebugged(NewObject); }
 	virtual void SetObjectBeingDebuggedSuper(UObject* NewObject) override { UBlueprint::SetObjectBeingDebugged(NewObject); }
 	virtual UObject* GetObjectBeingDebugged() override { return UBlueprint::GetObjectBeingDebugged(); }
 	virtual UObject* GetObjectBeingDebugged() const override { return UBlueprint::GetObjectBeingDebugged(); }
-	virtual void PreSave(FObjectPreSaveContext ObjectSaveContext) override;  
-	virtual void PostLoad() override;
+	UE_API virtual void PreSave(FObjectPreSaveContext ObjectSaveContext) override;  
+	UE_API virtual void PostLoad() override;
 #if WITH_EDITORONLY_DATA
-	static void DeclareConstructClasses(TArray<FTopLevelAssetPath>& OutConstructClasses, const UClass* SpecificSubclass);
+	UE_API static void DeclareConstructClasses(TArray<FTopLevelAssetPath>& OutConstructClasses, const UClass* SpecificSubclass);
 #endif
-	virtual bool IsPostLoadThreadSafe() const override { return IRigVMAssetInterface::IsPostLoadThreadSafe(); }
-	virtual void PostTransacted(const FTransactionObjectEvent& TransactionEvent) override { IRigVMAssetInterface::PostTransacted(TransactionEvent); }
+	virtual bool IsPostLoadThreadSafe() const override { return IRigVMEditorAssetInterface::IsPostLoadThreadSafe(); }
+	virtual void PostTransacted(const FTransactionObjectEvent& TransactionEvent) override { IRigVMEditorAssetInterface::PostTransacted(TransactionEvent); }
 	virtual void PostTransactedSuper(const FTransactionObjectEvent& TransactionEvent) override { UBlueprint::PostTransacted(TransactionEvent); };
-	virtual void ReplaceDeprecatedNodes() override { IRigVMAssetInterface::ReplaceDeprecatedNodes(); }
+	virtual void ReplaceDeprecatedNodes() override { IRigVMEditorAssetInterface::ReplaceDeprecatedNodes(); }
 	virtual void ReplaceDeprecatedNodesSuper() override { UBlueprint::ReplaceDeprecatedNodes(); };
-	virtual void PreDuplicate(FObjectDuplicationParameters& DupParams) override { IRigVMAssetInterface::PreDuplicate(DupParams); }
+	virtual void PreDuplicate(FObjectDuplicationParameters& DupParams) override { IRigVMEditorAssetInterface::PreDuplicate(DupParams); }
 	virtual void PreDuplicateSuper(FObjectDuplicationParameters& DupParams) override { UBlueprint::PreDuplicate(DupParams); }
-	virtual void PostDuplicate(bool bDuplicateForPIE) override { IRigVMAssetInterface::PostDuplicate(bDuplicateForPIE); }
+	virtual void PostDuplicate(bool bDuplicateForPIE) override { IRigVMEditorAssetInterface::PostDuplicate(bDuplicateForPIE); }
 	virtual void PostDuplicateSuper(bool bDuplicateForPIE) override { UBlueprint::PostDuplicate(bDuplicateForPIE); };
-	virtual void GetAssetRegistryTags(FAssetRegistryTagsContext Context) const override { IRigVMAssetInterface::GetAssetRegistryTags(Context); }
+	virtual void GetAssetRegistryTags(FAssetRegistryTagsContext Context) const override { IRigVMEditorAssetInterface::GetAssetRegistryTags(Context); }
 	virtual void GetAssetRegistryTagsSuper(FAssetRegistryTagsContext Context) const override { UBlueprint::GetAssetRegistryTags(Context); }
 	UE_DEPRECATED(5.4, "Implement the version that takes FAssetRegistryTagsContext instead.")
-	virtual void GetAssetRegistryTags(TArray<FAssetRegistryTag>& OutTags) const override;
+	UE_API virtual void GetAssetRegistryTags(TArray<FAssetRegistryTag>& OutTags) const override;
 
 	virtual bool SupportsGlobalVariables() const override { return true; }
 	virtual bool SupportsLocalVariables() const override { return true; }
@@ -211,25 +213,23 @@ public:
 	virtual bool SupportsDelegates() const override { return false; }
 	virtual bool SupportsEventGraphs() const override { return true; }
 	virtual bool SupportsAnimLayers() const override { return false; }
-	virtual bool ExportGraphToText(UEdGraph* InEdGraph, FString& OutText) override { IRigVMAssetInterface::ExportGraphToText(InEdGraph, OutText); return true; }
-	virtual bool TryImportGraphFromText(const FString& InClipboardText, UEdGraph** OutGraphPtr = nullptr) override { return IRigVMAssetInterface::TryImportGraphFromText(InClipboardText, OutGraphPtr); }
-	virtual bool CanImportGraphFromText(const FString& InClipboardText) override { return IRigVMAssetInterface::CanImportGraphFromText(InClipboardText); }
-	virtual bool RequiresForceLoadMembers(UObject* InObject) const override { return IRigVMAssetInterface::RequiresForceLoadMembers(InObject); }
+	virtual bool ExportGraphToText(UEdGraph* InEdGraph, FString& OutText) override { IRigVMEditorAssetInterface::ExportGraphToText(InEdGraph, OutText); return true; }
+	virtual bool TryImportGraphFromText(const FString& InClipboardText, UEdGraph** OutGraphPtr = nullptr) override { return IRigVMEditorAssetInterface::TryImportGraphFromText(InClipboardText, OutGraphPtr); }
+	virtual bool CanImportGraphFromText(const FString& InClipboardText) override { return IRigVMEditorAssetInterface::CanImportGraphFromText(InClipboardText); }
+	virtual bool RequiresForceLoadMembers(UObject* InObject) const override { return IRigVMEditorAssetInterface::RequiresForceLoadMembers(InObject); }
 	virtual bool RequiresForceLoadMembersSuper(UObject* InObject) const override { return UBlueprint::RequiresForceLoadMembers(InObject); }
 	
 
 	// UObject interface
-	virtual void PostEditChangeChainProperty(struct FPropertyChangedChainEvent& PropertyChangedEvent) override;
-	virtual void PostRename(UObject* OldOuter, const FName OldName) override;
+	UE_API virtual void PostEditChangeChainProperty(struct FPropertyChangedChainEvent& PropertyChangedEvent) override;
+	UE_API virtual void PostRename(UObject* OldOuter, const FName OldName) override;
 	/** Called during cooking. Must return all objects that will be Preload()ed when this is serialized at load time. */
-	virtual void GetPreloadDependencies(TArray<UObject*>& OutDeps) override;
+	UE_API virtual void GetPreloadDependencies(TArray<UObject*>& OutDeps) override;
 
-	virtual bool MarkPackageDirty() override { return UBlueprint::MarkPackageDirty(); }
-	virtual UPackage* GetPackage() override { return UBlueprint::GetPackage(); }
 	virtual TArray<TObjectPtr<UEdGraph>>& GetUberGraphs() override { return UbergraphPages; }
 	virtual void GetAllEdGraphs(TArray<UEdGraph*>& Graphs) const override { UBlueprint::GetAllGraphs(Graphs); }
 	virtual void SetUberGraphs(const TArray<TObjectPtr<UEdGraph>>& InGraphs) override { UbergraphPages = InGraphs; }
-	virtual void PostSaveRoot(FObjectPostSaveRootContext ObjectSaveContext) override;
+	UE_API virtual void PostSaveRoot(FObjectPostSaveRootContext ObjectSaveContext) override;
 
 	//  --- IRigVMClientHost interface Start---
 	virtual FString GetAssetName() const override { return GetName(); }
@@ -239,84 +239,85 @@ public:
 	virtual UClass* GetRigVMEdGraphNodeClass() const override { return URigVMEdGraphNode::StaticClass(); }
 	virtual UClass* GetRigVMEdGraphSchemaClass() const override { return URigVMEdGraphSchema::StaticClass(); }
 	virtual UClass* GetRigVMEditorSettingsClass() const override { return URigVMEditorSettings::StaticClass(); }
-	virtual UObject* GetEditorObjectForRigVMGraph(const URigVMGraph* InVMGraph) const override { return IRigVMAssetInterface::GetEditorObjectForRigVMGraph(InVMGraph); }
-	virtual URigVMGraph* GetRigVMGraphForEditorObject(UObject* InObject) const override { return IRigVMAssetInterface::GetRigVMGraphForEditorObject(InObject); }
-	virtual void HandleRigVMGraphAdded(const FRigVMClient* InClient, const FString& InNodePath) override { return IRigVMAssetInterface::HandleRigVMGraphAdded(InClient, InNodePath); }
-	virtual void HandleRigVMGraphRemoved(const FRigVMClient* InClient, const FString& InNodePath) override { return IRigVMAssetInterface::HandleRigVMGraphRemoved(InClient, InNodePath); }
-	virtual void HandleRigVMGraphRenamed(const FRigVMClient* InClient, const FString& InOldNodePath, const FString& InNewNodePath) override { return IRigVMAssetInterface::HandleRigVMGraphRenamed(InClient, InOldNodePath, InNewNodePath); }
-	virtual void HandleConfigureRigVMController(const FRigVMClient* InClient, URigVMController* InControllerToConfigure) override { return IRigVMAssetInterface::HandleConfigureRigVMController(InClient, InControllerToConfigure); }
-	virtual void IncrementVMRecompileBracket() override { return IRigVMAssetInterface::IncrementVMRecompileBracket(); }
-	virtual void DecrementVMRecompileBracket() override { return IRigVMAssetInterface::DecrementVMRecompileBracket(); }
-	virtual void RefreshAllModels(ERigVMLoadType InLoadType = ERigVMLoadType::PostLoad) override { return IRigVMAssetInterface::RefreshAllModels(InLoadType); }
-	virtual void OnRigVMRegistryChanged() override { return IRigVMAssetInterface::OnRigVMRegistryChanged(); }
-	virtual URigVMGraph* GetModel(const FString& InNodePath) const override { return IRigVMAssetInterface::GetModel(InNodePath); }
-	virtual FRigVMGetFocusedGraph& OnGetFocusedGraph() override { return IRigVMAssetInterface::OnGetFocusedGraph(); }
-	virtual const FRigVMGetFocusedGraph& OnGetFocusedGraph() const override { return IRigVMAssetInterface::OnGetFocusedGraph(); }
-	virtual void SetupPinRedirectorsForBackwardsCompatibility() override { return IRigVMAssetInterface::SetupPinRedirectorsForBackwardsCompatibility(); }
-	virtual FRigVMGraphModifiedEvent& OnModified() override { return IRigVMAssetInterface::OnModified(); }
-	virtual bool IsFunctionPublic(const FName& InFunctionName) const override { return IRigVMAssetInterface::IsFunctionPublic(InFunctionName); }
-	virtual void MarkFunctionPublic(const FName& InFunctionName, bool bIsPublic = true) override { return IRigVMAssetInterface::MarkFunctionPublic(InFunctionName, bIsPublic); }
-	virtual void RenameGraph(const FString& InNodePath, const FName& InNewName) override { return IRigVMAssetInterface::RenameGraph(InNodePath, InNewName); }
+	virtual UObject* GetEditorObjectForRigVMGraph(const URigVMGraph* InVMGraph) const override { return IRigVMEditorAssetInterface::GetEditorObjectForRigVMGraph(InVMGraph); }
+	virtual URigVMGraph* GetRigVMGraphForEditorObject(UObject* InObject) const override { return IRigVMEditorAssetInterface::GetRigVMGraphForEditorObject(InObject); }
+	virtual void HandleRigVMGraphAdded(const FRigVMClient* InClient, const FString& InNodePath) override { return IRigVMEditorAssetInterface::HandleRigVMGraphAdded(InClient, InNodePath); }
+	virtual void HandleRigVMGraphRemoved(const FRigVMClient* InClient, const FString& InNodePath) override { return IRigVMEditorAssetInterface::HandleRigVMGraphRemoved(InClient, InNodePath); }
+	virtual void HandleRigVMGraphRenamed(const FRigVMClient* InClient, const FString& InOldNodePath, const FString& InNewNodePath) override { return IRigVMEditorAssetInterface::HandleRigVMGraphRenamed(InClient, InOldNodePath, InNewNodePath); }
+	virtual void HandleConfigureRigVMController(const FRigVMClient* InClient, URigVMController* InControllerToConfigure) override { return IRigVMEditorAssetInterface::HandleConfigureRigVMController(InClient, InControllerToConfigure); }
+	virtual void IncrementVMRecompileBracket() override { return IRigVMEditorAssetInterface::IncrementVMRecompileBracket(); }
+	virtual void DecrementVMRecompileBracket() override { return IRigVMEditorAssetInterface::DecrementVMRecompileBracket(); }
+	virtual void RefreshAllModels(ERigVMLoadType InLoadType = ERigVMLoadType::PostLoad) override { return IRigVMEditorAssetInterface::RefreshAllModels(InLoadType); }
+	virtual void OnRigVMRegistryChanged() override { return IRigVMEditorAssetInterface::OnRigVMRegistryChanged(); }
+	virtual URigVMGraph* GetModel(const FString& InNodePath) const override { return IRigVMEditorAssetInterface::GetModel(InNodePath); }
+	virtual FRigVMGetFocusedGraph& OnGetFocusedGraph() override { return IRigVMEditorAssetInterface::OnGetFocusedGraph(); }
+	virtual const FRigVMGetFocusedGraph& OnGetFocusedGraph() const override { return IRigVMEditorAssetInterface::OnGetFocusedGraph(); }
+	virtual void SetupPinRedirectorsForBackwardsCompatibility() override { return IRigVMEditorAssetInterface::SetupPinRedirectorsForBackwardsCompatibility(); }
+	virtual FRigVMGraphModifiedEvent& OnModified() override { return IRigVMEditorAssetInterface::OnModified(); }
+	virtual bool IsFunctionPublic(const FName& InFunctionName) const override { return IRigVMEditorAssetInterface::IsFunctionPublic(InFunctionName); }
+	virtual void MarkFunctionPublic(const FName& InFunctionName, bool bIsPublic = true) override { return IRigVMEditorAssetInterface::MarkFunctionPublic(InFunctionName, bIsPublic); }
+	virtual void RenameGraph(const FString& InNodePath, const FName& InNewName) override { return IRigVMEditorAssetInterface::RenameGraph(InNodePath, InNewName); }
 	//  --- IRigVMClientHost interface End---
 	
 	UFUNCTION(BlueprintCallable, Category = "RigVM Blueprint")
-	virtual void RecompileVM() override { IRigVMAssetInterface::RecompileVM(); }
+	virtual void RecompileVM() override { IRigVMEditorAssetInterface::RecompileVM(); }
 	UFUNCTION(BlueprintCallable, Category = "RigVM Blueprint")
-	virtual void RecompileVMIfRequired() override { return IRigVMAssetInterface::RecompileVMIfRequired(); }
+	virtual void RecompileVMIfRequired() override { return IRigVMEditorAssetInterface::RecompileVMIfRequired(); }
 	UFUNCTION(BlueprintCallable, Category = "RigVM Blueprint")
-	virtual void RequestAutoVMRecompilation() override { return IRigVMAssetInterface::RequestAutoVMRecompilation(); }
+	virtual void RequestAutoVMRecompilation() override { return IRigVMEditorAssetInterface::RequestAutoVMRecompilation(); }
 	UFUNCTION(BlueprintCallable, Category = "RigVM Blueprint")
-	virtual void SetAutoVMRecompile(bool bAutoRecompile) override { return IRigVMAssetInterface::SetAutoVMRecompile(bAutoRecompile); }
+	virtual void SetAutoVMRecompile(bool bAutoRecompile) override { return IRigVMEditorAssetInterface::SetAutoVMRecompile(bAutoRecompile); }
 	UFUNCTION(BlueprintPure, Category = "RigVM Blueprint")
-	virtual bool GetAutoVMRecompile() const override { return IRigVMAssetInterface::GetAutoVMRecompile(); }
+	virtual bool GetAutoVMRecompile() const override { return IRigVMEditorAssetInterface::GetAutoVMRecompile(); }
 
 	UFUNCTION(BlueprintCallable, Category = "RigVM Blueprint")
-	virtual void RequestRigVMInit() override { IRigVMAssetInterface::RequestRigVMInit(); }
+	virtual void RequestRigVMInit() override { IRigVMEditorAssetInterface::RequestRigVMInit(); }
 
 	UFUNCTION(BlueprintCallable, Category = "RigVM Blueprint")
-	virtual URigVMGraph* GetModel(const UEdGraph* InEdGraph = nullptr) const override { return IRigVMAssetInterface::GetModel(InEdGraph); }
+	virtual URigVMGraph* GetModel(const UEdGraph* InEdGraph = nullptr) const override { return IRigVMEditorAssetInterface::GetModel(InEdGraph); }
 
 	UFUNCTION(BlueprintCallable, Category = "RigVM Blueprint")
-	virtual URigVMGraph* GetDefaultModel() const override { return IRigVMAssetInterface::GetDefaultModel(); }
+	virtual URigVMGraph* GetDefaultModel() const override { return IRigVMEditorAssetInterface::GetDefaultModel(); }
 
 	UFUNCTION(BlueprintCallable, Category = "RigVM Blueprint")
-	virtual TArray<URigVMGraph*> GetAllModels() const override { return IRigVMAssetInterface::GetAllModels(); }
+	virtual TArray<URigVMGraph*> GetAllModels() const override { return IRigVMEditorAssetInterface::GetAllModels(); }
 
 	UFUNCTION(BlueprintCallable, Category = "RigVM Blueprint")
-	virtual URigVMFunctionLibrary* GetLocalFunctionLibrary() const override { return IRigVMAssetInterface::GetLocalFunctionLibrary(); }
+	virtual URigVMFunctionLibrary* GetLocalFunctionLibrary() const override { return IRigVMEditorAssetInterface::GetLocalFunctionLibrary(); }
 
 	UFUNCTION(BlueprintCallable, Category = "RigVM Blueprint")
-	virtual URigVMFunctionLibrary* GetOrCreateLocalFunctionLibrary(bool bSetupUndoRedo = true) override { return IRigVMAssetInterface::GetOrCreateLocalFunctionLibrary(bSetupUndoRedo); }
+	virtual URigVMFunctionLibrary* GetOrCreateLocalFunctionLibrary(bool bSetupUndoRedo = true) override { return IRigVMEditorAssetInterface::GetOrCreateLocalFunctionLibrary(bSetupUndoRedo); }
 
 	UFUNCTION(BlueprintCallable, Category = "RigVM Blueprint")
-	virtual URigVMGraph* AddModel(FString InName = TEXT("Rig Graph"), bool bSetupUndoRedo = true, bool bPrintPythonCommand = true) override { return IRigVMAssetInterface::AddModel(InName, bSetupUndoRedo, bPrintPythonCommand); }
+	virtual URigVMGraph* AddModel(FString InName = TEXT("Rig Graph"), bool bSetupUndoRedo = true, bool bPrintPythonCommand = true) override { return IRigVMEditorAssetInterface::AddModel(InName, bSetupUndoRedo, bPrintPythonCommand); }
 
 	UFUNCTION(BlueprintCallable, Category = "RigVM Blueprint")
-	virtual bool RemoveModel(FString InName = TEXT("Rig Graph"), bool bSetupUndoRedo = true, bool bPrintPythonCommand = true) override { return IRigVMAssetInterface::RemoveModel(InName, bSetupUndoRedo, bPrintPythonCommand); }
+	virtual bool RemoveModel(FString InName = TEXT("Rig Graph"), bool bSetupUndoRedo = true, bool bPrintPythonCommand = true) override { return IRigVMEditorAssetInterface::RemoveModel(InName, bSetupUndoRedo, bPrintPythonCommand); }
 
 
 	UFUNCTION(BlueprintCallable, Category = "RigVM Blueprint")
-	virtual URigVMGraph* GetFocusedModel() const override { return IRigVMAssetInterface::GetFocusedModel(); }
+	virtual URigVMGraph* GetFocusedModel() const override { return IRigVMEditorAssetInterface::GetFocusedModel(); }
 
 	UFUNCTION(BlueprintCallable, Category = "RigVM Blueprint")
-	virtual URigVMController* GetController(const URigVMGraph* InGraph = nullptr) const override { return IRigVMAssetInterface::GetController(InGraph); }
-	virtual URigVMController* GetController(const UEdGraph* InGraph) const override { return IRigVMAssetInterface::GetController(InGraph); }
+	virtual URigVMController* GetController(const URigVMGraph* InGraph = nullptr) const override { return IRigVMEditorAssetInterface::GetController(InGraph); }
+	virtual URigVMController* GetController(const UEdGraph* InGraph) const override { return IRigVMEditorAssetInterface::GetController(InGraph); }
 
 	UFUNCTION(BlueprintCallable, Category = "RigVM Blueprint")
-	virtual URigVMController* GetControllerByName(const FString InGraphName = TEXT("")) const override { return IRigVMAssetInterface::GetControllerByName(InGraphName); }
+	virtual URigVMController* GetControllerByName(const FString InGraphName = TEXT("")) const override { return IRigVMEditorAssetInterface::GetControllerByName(InGraphName); }
 
 	UFUNCTION(BlueprintCallable, Category = "RigVM Blueprint")
-	virtual URigVMController* GetOrCreateController(URigVMGraph* InGraph = nullptr) override { return IRigVMAssetInterface::GetOrCreateController(InGraph); }
-	virtual URigVMController* GetOrCreateController(const UEdGraph* InGraph) override { return IRigVMAssetInterface::GetOrCreateController(InGraph); }
+	virtual URigVMController* GetOrCreateController(URigVMGraph* InGraph = nullptr) override { return IRigVMEditorAssetInterface::GetOrCreateController(InGraph); }
+	virtual URigVMController* GetOrCreateController(const UEdGraph* InGraph) override { return IRigVMEditorAssetInterface::GetOrCreateController(InGraph); }
 
 	UFUNCTION(BlueprintCallable, Category = "RigVM Blueprint")
-	virtual TArray<FString> GeneratePythonCommands(const FString InNewBlueprintName) override { return IRigVMAssetInterface::GeneratePythonCommands(InNewBlueprintName); }
+	virtual TArray<FString> GeneratePythonCommands(const FString InNewBlueprintName) override { return IRigVMEditorAssetInterface::GeneratePythonCommands(InNewBlueprintName); }
 #endif
 
-	virtual TArray<FRigVMExternalDependency> GetExternalDependenciesForCategory(const FName& InCategory) const override { return IRigVMAssetInterface::GetExternalDependenciesForCategory(InCategory); }
+	virtual TArray<FRigVMExternalDependency> GetExternalDependenciesForCategory(const FName& InCategory) const override { return IRigVMEditorAssetInterface::GetExternalDependenciesForCategory(InCategory); }
 
 
 	virtual bool ShouldBeMarkedDirtyUponTransaction() const override { return false; }
+	virtual FRigVMPropertyBag* GetVariablesPropertyBag() override { return nullptr; }
 
 	
 #if WITH_EDITORONLY_DATA
@@ -394,50 +395,50 @@ public:
 	virtual UClass* GetRigVMHostClass() const { return GeneratedClass; }
 
 	UFUNCTION(BlueprintCallable, Category = "VM")
-	virtual URigVMHost* CreateRigVMHost() { return IRigVMAssetInterface::CreateRigVMHostImpl(); }
+	virtual URigVMHost* CreateRigVMHost() { return IRigVMEditorAssetInterface::CreateRigVMHostImpl(); }
 
 	UFUNCTION(BlueprintCallable, Category = "VM")
-	virtual URigVMHost* GetDebuggedRigVMHost() { return IRigVMAssetInterface::GetDebuggedRigVMHost(); }
+	virtual URigVMHost* GetDebuggedRigVMHost() { return IRigVMEditorAssetInterface::GetDebuggedRigVMHost(); }
 
 	UFUNCTION(BlueprintCallable, Category = "VM")
-	virtual TArray<UStruct*> GetAvailableRigVMStructs() const override { return IRigVMAssetInterface::GetAvailableRigVMStructs(); }
+	virtual TArray<UStruct*> GetAvailableRigVMStructs() const override { return IRigVMEditorAssetInterface::GetAvailableRigVMStructs(); }
 
 #if WITH_EDITOR
 	UFUNCTION(BlueprintCallable, Category = "Variables")
 	virtual TArray<FRigVMGraphVariableDescription> GetMemberVariables() const { return GetAssetVariables(); }
 	
 	UFUNCTION(BlueprintCallable, Category = "Variables")
-	virtual FName AddMemberVariable(const FName& InName, const FString& InCPPType, bool bIsPublic = false, bool bIsReadOnly = false, FString InDefaultValue = TEXT("")) override;
+	UE_API virtual FName AddMemberVariable(const FName& InName, const FString& InCPPType, bool bIsPublic = false, bool bIsReadOnly = false, FString InDefaultValue = TEXT("")) override;
 
 	UFUNCTION(BlueprintCallable, Category = "Variables")
-	virtual bool RemoveMemberVariable(const FName& InName) override;
+	UE_API virtual bool RemoveMemberVariable(const FName& InName) override;
 
 	UFUNCTION(BlueprintCallable, Category = "Variables")
-	virtual bool BulkRemoveMemberVariables(const TArray<FName>& InNames) override;
+	UE_API virtual bool BulkRemoveMemberVariables(const TArray<FName>& InNames) override;
 
 	UFUNCTION(BlueprintCallable, Category = "Variables")
-	virtual bool RenameMemberVariable(const FName& InOldName, const FName& InNewName) override;
+	UE_API virtual bool RenameMemberVariable(const FName& InOldName, const FName& InNewName) override;
 
 	UFUNCTION(BlueprintCallable, Category = "Variables")
-	virtual bool ChangeMemberVariableType(const FName& InName, const FString& InCPPType, bool bIsPublic = false, bool bIsReadOnly = false, FString InDefaultValue = TEXT("")) override;
-	virtual bool ChangeMemberVariableType(const FName& InName, const FEdGraphPinType& InType) override;
+	UE_API virtual bool ChangeMemberVariableType(const FName& InName, const FString& InCPPType, bool bIsPublic = false, bool bIsReadOnly = false, FString InDefaultValue = TEXT("")) override;
+	UE_API virtual bool ChangeMemberVariableType(const FName& InName, const FEdGraphPinType& InType) override;
 
 	UFUNCTION(BlueprintPure, Category = "Variants", meta = (DisplayName = "GetAssetVariant", ScriptName = "GetAssetVariant"))
 	virtual FRigVMVariant GetAssetVariantBP() const { return AssetVariant; }
 
 	UFUNCTION(BlueprintPure, Category = "Variants")
-	virtual FRigVMVariantRef GetAssetVariantRef() const override { return IRigVMAssetInterface::GetAssetVariantRefImpl(); }
+	virtual FRigVMVariantRef GetAssetVariantRef() const override { return IRigVMEditorAssetInterface::GetAssetVariantRefImpl(); }
 
 	/** Resets the asset's guid to a new one and splits it from the former variant set */
 	UFUNCTION(BlueprintCallable, Category = "Variants")
-	virtual bool SplitAssetVariant() override { return IRigVMAssetInterface::SplitAssetVariantImpl(); }
+	virtual bool SplitAssetVariant() override { return IRigVMEditorAssetInterface::SplitAssetVariantImpl(); }
 
 	/** Merges the asset's guid with a provided one to join the variant set */
 	UFUNCTION(BlueprintCallable, Category = "Variants")
-	virtual bool JoinAssetVariant(const FGuid& InGuid) override { return IRigVMAssetInterface::JoinAssetVariantImpl(InGuid); }
+	virtual bool JoinAssetVariant(const FGuid& InGuid) override { return IRigVMEditorAssetInterface::JoinAssetVariantImpl(InGuid); }
 
 	UFUNCTION(BlueprintPure, Category = "Variants")
-	virtual TArray<FRigVMVariantRef> GetMatchingVariants() const override { return IRigVMAssetInterface::GetMatchingVariantsImpl(); }
+	virtual TArray<FRigVMVariantRef> GetMatchingVariants() const override { return IRigVMEditorAssetInterface::GetMatchingVariantsImpl(); }
 #endif	// #if WITH_EDITOR
 
 	
@@ -449,49 +450,49 @@ private:
 
 #if WITH_EDITOR
 	UFUNCTION(BlueprintCallable, Category = "RigVM Blueprint")
-	virtual void SuspendNotifications(bool bSuspendNotifs) { return IRigVMAssetInterface::SuspendNotifications(bSuspendNotifs); }
+	virtual void SuspendNotifications(bool bSuspendNotifs) { return IRigVMEditorAssetInterface::SuspendNotifications(bSuspendNotifs); }
 #endif
 
 protected:
 
 #if WITH_EDITOR
 	/** Our currently running rig vm instance */
-	// Declaring these transient properties here instead of the IRigVMAssetInterface class in order to have them as UPROPERTY, to avoid having stale pointers when GC happens
+	// Declaring these transient properties here instead of the IRigVMEditorAssetInterface class in order to have them as UPROPERTY, to avoid having stale pointers when GC happens
 	UPROPERTY(transient)
 	TObjectPtr<URigVMHost> EditorHost = nullptr;
 	virtual TObjectPtr<URigVMHost>& GetEditorHost() override { return EditorHost; }
 
 	static FName FindHostMemberVariableUniqueName(TSharedPtr<FKismetNameValidator> InNameValidator, const FString& InBaseName);
 	static int32 AddHostMemberVariable(URigVMBlueprint* InBlueprint, const FName& InVarName, FEdGraphPinType InVarType, bool bIsPublic, bool bIsReadOnly, FString InDefaultValue);
-	virtual FName AddHostMemberVariableFromExternal(FRigVMExternalVariable InVariableToCreate, FString InDefaultValue = FString()) override;
+	UE_API virtual FName AddHostMemberVariableFromExternal(FRigVMExternalVariable InVariableToCreate, FString InDefaultValue = FString()) override;
 public:
 
-	virtual void OnPreVariableChange(UObject* InObject);
-	virtual void OnPostVariableChange(UBlueprint* InBlueprint);
+	UE_API virtual void OnPreVariableChange(UObject* InObject);
+	UE_API virtual void OnPostVariableChange(UBlueprint* InBlueprint);
 	bool bUpdatingExternalVariables;
 
-	void OnBlueprintChanged(UBlueprint* InBlueprint);
-	void OnSetObjectBeingDebuggedReceived(UObject* InObject);
+	UE_API void OnBlueprintChanged(UBlueprint* InBlueprint);
+	UE_API void OnSetObjectBeingDebuggedReceived(UObject* InObject);
 #endif
 
 protected:
 
-	virtual void SetupDefaultObjectDuringCompilation(URigVMHost* InCDO) override { IRigVMAssetInterface::SetupDefaultObjectDuringCompilation(InCDO); }
-	virtual void PreCompile() override;
-
-	virtual FProperty* FindGeneratedPropertyByName(const FName& InName) const override;
-	virtual bool SetVariableTooltip(const FName& InName, const FText& InTooltip) override;
-	virtual FText GetVariableTooltip(const FName& InName) const override;
-	virtual bool SetVariableCategory(const FName& InName, const FString& InCategory) override;
-	virtual FString GetVariableCategory(const FName& InName) override;
-	virtual FString GetVariableMetadataValue(const FName& InName, const FName& InKey) override;
-	virtual bool SetVariableExposeOnSpawn(const FName& InName, const bool bInExposeOnSpawn) override;
-	virtual bool SetVariableExposeToCinematics(const FName& InName, const bool bInExposeToCinematics) override;
-	virtual bool SetVariablePrivate(const FName& InName, const bool bInPrivate) override;
-	virtual bool SetVariablePublic(const FName& InName, const bool bIsPublic) override;
+	virtual void SetupDefaultObjectDuringCompilation(URigVMHost* InCDO) override { IRigVMEditorAssetInterface::SetupDefaultObjectDuringCompilation(InCDO); }
+	UE_API virtual void PreCompile() override;
+
+	UE_API virtual FProperty* FindGeneratedPropertyByName(const FName& InName) override;
+	UE_API virtual bool SetVariableTooltip(const FName& InName, const FText& InTooltip) override;
+	UE_API virtual FText GetVariableTooltip(const FName& InName) const override;
+	UE_API virtual bool SetVariableCategory(const FName& InName, const FString& InCategory) override;
+	UE_API virtual FString GetVariableCategory(const FName& InName) override;
+	UE_API virtual FString GetVariableMetadataValue(const FName& InName, const FName& InKey) override;
+	UE_API virtual bool SetVariableExposeOnSpawn(const FName& InName, const bool bInExposeOnSpawn) override;
+	UE_API virtual bool SetVariableExposeToCinematics(const FName& InName, const bool bInExposeToCinematics) override;
+	UE_API virtual bool SetVariablePrivate(const FName& InName, const bool bInPrivate) override;
+	UE_API virtual bool SetVariablePublic(const FName& InName, const bool bIsPublic) override;
 	//virtual bool ChangeAssetVariableType(const FName& InName, const FEdGraphPinType& InType) override;
-	virtual FString OnCopyVariable(const FName& InName) const override;
-	virtual bool OnPasteVariable(const FString& InText) override;
+	UE_API virtual FString OnCopyVariable(const FName& InName) const override;
+	UE_API virtual bool OnPasteVariable(const FString& InText) override;
 	
 #if WITH_EDITOR
 
diff --git a/Engine/Plugins/Runtime/RigVM/Source/RigVMDeveloper/Public/RigVMBlueprintUtils.h b/Engine/Plugins/Runtime/RigVM/Source/RigVMDeveloper/Public/RigVMBlueprintUtils.h
--- a/Engine/Plugins/Runtime/RigVM/Source/RigVMDeveloper/Public/RigVMBlueprintUtils.h
+++ b/Engine/Plugins/Runtime/RigVM/Source/RigVMDeveloper/Public/RigVMBlueprintUtils.h
@@ -3,13 +3,12 @@
 #pragma once
 
 #include "CoreMinimal.h"
-#include "RigVMAsset.h"
+#include "RigVMEditorAsset.h"
 #include "Templates/SubclassOf.h"
 #include "EdGraph/EdGraphPin.h"
 
 #define UE_API RIGVMDEVELOPER_API
 
-class IRigVMAssetInterface;
 class UStruct;
 class UBlueprint;
 class URigVMEdGraphNode;
@@ -23,34 +22,34 @@ static UE_API void ForAllRigVMStructs(TFunction<void(UScriptStruct*)> InFunction
 
 /** Handle blueprint node reconstruction */
 static UE_API void HandleReconstructAllBlueprintNodes(UBlueprint* InBlueprint);
-static UE_API void HandleReconstructAllBlueprintNodes(FRigVMAssetInterfacePtr InBlueprint);
+static UE_API void HandleReconstructAllBlueprintNodes(FRigVMEditorAssetInterfacePtr InBlueprint);
 
 /** Handle blueprint node refresh */
-static UE_API void HandleRefreshAllNodes(FRigVMAssetInterfacePtr InBlueprint);
+static UE_API void HandleRefreshAllNodes(FRigVMEditorAssetInterfacePtr InBlueprint);
 static UE_API void HandleRefreshAllBlueprintNodes(UBlueprint* InBlueprint);
 
 /** Handle blueprint deleted */
 static UE_API void HandleAssetDeleted(const FAssetData& InAssetData);
 
 /** remove the variable if not used by anybody else but ToBeDeleted*/
-static UE_API void RemoveMemberVariableIfNotUsed(FRigVMAssetInterfacePtr Blueprint, const FName VarName);
+static UE_API void RemoveMemberVariableIfNotUsed(FRigVMEditorAssetInterfacePtr Blueprint, const FName VarName);
 
 /** Create a new EdGraph */
 static UE_API URigVMEdGraph* CreateNewGraph(UObject* ParentScope, const FName& GraphName, TSubclassOf<class URigVMEdGraph> GraphClass, TSubclassOf<class URigVMEdGraphSchema> SchemaClass);
 
-static UE_API FName ValidateName(FRigVMAssetInterfacePtr InBlueprint, const FString& InName);
+static UE_API FName ValidateName(FRigVMEditorAssetInterfacePtr InBlueprint, const FString& InName);
 
 /** Searches all nodes in a Blueprint and checks for a matching Guid */
-static UE_API UEdGraphNode* GetNodeByGUID(const FRigVMAssetInterfacePtr InBlueprint, const FGuid& InNodeGuid);
+static UE_API UEdGraphNode* GetNodeByGUID(const FRigVMEditorAssetInterfacePtr InBlueprint, const FGuid& InNodeGuid);
 
 // Helper function to get the asset that ultimately owns a graph.
-static UE_API FRigVMAssetInterfacePtr FindAssetForGraph(const UEdGraph* Graph);
+static UE_API FRigVMEditorAssetInterfacePtr FindAssetForGraph(const UEdGraph* Graph);
 
 // Helper function to get the asset that ultimately owns a graph.
-static UE_API FRigVMAssetInterfacePtr FindAssetForNode(const UEdGraphNode* Node);
+static UE_API FRigVMEditorAssetInterfacePtr FindAssetForNode(const UEdGraphNode* Node);
 
 // Helper function to find a unique name for an asset variable
-static UE_API FName FindUniqueVariableName(const IRigVMAssetInterface* InBlueprint, const FString& InBaseName);
+static UE_API FName FindUniqueVariableName(const IRigVMEditorAssetInterface* InBlueprint, const FString& InBaseName);
 };
 
 #undef UE_API
diff --git a/Engine/Plugins/Runtime/RigVM/Source/RigVMDeveloper/Public/RigVMCompiler/RigVMAST.h b/Engine/Plugins/Runtime/RigVM/Source/RigVMDeveloper/Public/RigVMCompiler/RigVMAST.h
--- a/Engine/Plugins/Runtime/RigVM/Source/RigVMDeveloper/Public/RigVMCompiler/RigVMAST.h
+++ b/Engine/Plugins/Runtime/RigVM/Source/RigVMDeveloper/Public/RigVMCompiler/RigVMAST.h
@@ -843,7 +843,7 @@ public:
 
 	// returns the pin in the model this variable is representing
 	// @return the pin in the model this variable is representing
-	URigVMPin* GetPin() const { return GetProxy().GetSubjectChecked<URigVMPin>(); }
+	URigVMPin* GetPin() const { if (ensure(GetProxy().GetSubject()))  { return GetProxy().GetSubjectChecked<URigVMPin>(); } return nullptr; }
 
 	// returns the C++ data type of this variable
 	// @return the C++ data type of this variable
diff --git a/Engine/Plugins/Runtime/RigVM/Source/RigVMDeveloper/Public/RigVMDeveloperTypeUtils.h b/Engine/Plugins/Runtime/RigVM/Source/RigVMDeveloper/Public/RigVMDeveloperTypeUtils.h
--- a/Engine/Plugins/Runtime/RigVM/Source/RigVMDeveloper/Public/RigVMDeveloperTypeUtils.h
+++ b/Engine/Plugins/Runtime/RigVM/Source/RigVMDeveloper/Public/RigVMDeveloperTypeUtils.h
@@ -14,8 +14,6 @@ class URigVMPin;
 namespace RigVMTypeUtils
 {
 #if WITH_EDITOR
-	RIGVMDEVELOPER_API FRigVMExternalVariable ExternalVariableFromRigVMVariableDescription(const FRigVMGraphVariableDescription& InVariableDescription);
-	
 	RIGVMDEVELOPER_API FEdGraphPinType PinTypeFromTypeIndex(const TRigVMTypeIndex& InTypeIndex);
 
 	RIGVMDEVELOPER_API FEdGraphPinType PinTypeFromCPPType(const FName& InCPPType, UObject* InCPPTypeObject);
@@ -38,8 +36,6 @@ namespace RigVMTypeUtils
 
 	RIGVMDEVELOPER_API bool CPPTypeFromPinType(const FEdGraphPinType& InPinType, FString& OutCPPType, FString& OutCPPTypeObjectPath);
 
-	RIGVMDEVELOPER_API bool CPPTypeFromExternalVariable(const FRigVMExternalVariable& InExternalVariable, FString& OutCPPType, UObject** OutCPPTypeObject);
-
 	RIGVMDEVELOPER_API TRigVMTypeIndex TypeIndexFromPinType(const FEdGraphPinType& InPinType);
 
 	RIGVMDEVELOPER_API bool AreCompatible(const FName& InCPPTypeA, UObject* InCPPTypeObjectA, const FName& InCPPTypeB, UObject* InCPPTypeObjectB);
diff --git a/Engine/Plugins/Runtime/RigVM/Source/RigVMDeveloper/Public/RigVMJsonUtils.h b/Engine/Plugins/Runtime/RigVM/Source/RigVMDeveloper/Public/RigVMJsonUtils.h
--- a/Engine/Plugins/Runtime/RigVM/Source/RigVMDeveloper/Public/RigVMJsonUtils.h
+++ b/Engine/Plugins/Runtime/RigVM/Source/RigVMDeveloper/Public/RigVMJsonUtils.h
@@ -10,7 +10,7 @@
 #include "RigVMCore/RigVMFunction.h"
 #include "RigVMCore/RigVMDispatchFactory.h"
 
-class IRigVMAssetInterface;
+class IRigVMEditorAssetInterface;
 class URigVMGraph;
 class URigVMNode;
 class URigVMPin;
@@ -57,7 +57,7 @@ namespace RigVMJsonUtils
 	RIGVMDEVELOPER_API bool GenericStructAttributesToJson(FIntrospectionSettings& InSettings, const UScriptStruct* InStruct, FJsonObject* OutJsonObject);
 
 	// introspection for existing graphs
-	RIGVMDEVELOPER_API TSharedPtr<FJsonObject> ToJson(FIntrospectionSettings& InSettings, const IRigVMAssetInterface* InAsset);
+	RIGVMDEVELOPER_API TSharedPtr<FJsonObject> ToJson(FIntrospectionSettings& InSettings, const IRigVMEditorAssetInterface* InAsset);
 	RIGVMDEVELOPER_API TSharedPtr<FJsonObject> ToJson(FIntrospectionSettings& InSettings, const URigVMGraph* InGraph, bool bOnlySelection = false);
 	RIGVMDEVELOPER_API TSharedPtr<FJsonObject> ToJson(FIntrospectionSettings& InSettings, const URigVMNode* InNode);
 	RIGVMDEVELOPER_API TSharedPtr<FJsonObject> ToJson(FIntrospectionSettings& InSettings, const URigVMPin* InPin);
diff --git a/Engine/Plugins/Runtime/RigVM/Source/RigVMDeveloper/Public/RigVMModel/Nodes/RigVMFunctionReferenceNode.h b/Engine/Plugins/Runtime/RigVM/Source/RigVMDeveloper/Public/RigVMModel/Nodes/RigVMFunctionReferenceNode.h
--- a/Engine/Plugins/Runtime/RigVM/Source/RigVMDeveloper/Public/RigVMModel/Nodes/RigVMFunctionReferenceNode.h
+++ b/Engine/Plugins/Runtime/RigVM/Source/RigVMDeveloper/Public/RigVMModel/Nodes/RigVMFunctionReferenceNode.h
@@ -92,9 +92,10 @@ private:
 	friend class URigVMController;
 	friend class FRigVMParserAST;
 	friend class URigVMBlueprint;
-	friend class IRigVMAssetInterface;
+	friend class IRigVMEditorAssetInterface;
 	friend struct FRigVMClient;
 	friend struct EngineTestRigVMFramework;
+	friend class FRigVMExampleAssetActions;
 };
 
 #undef UE_API
diff --git a/Engine/Plugins/Runtime/RigVM/Source/RigVMDeveloper/Public/RigVMModel/Nodes/RigVMVariableNode.h b/Engine/Plugins/Runtime/RigVM/Source/RigVMDeveloper/Public/RigVMModel/Nodes/RigVMVariableNode.h
--- a/Engine/Plugins/Runtime/RigVM/Source/RigVMDeveloper/Public/RigVMModel/Nodes/RigVMVariableNode.h
+++ b/Engine/Plugins/Runtime/RigVM/Source/RigVMDeveloper/Public/RigVMModel/Nodes/RigVMVariableNode.h
@@ -4,7 +4,7 @@
 
 #include "RigVMCore/RigVMExternalVariable.h"
 #include "RigVMModel/RigVMNode.h"
-#include "RigVMModel/RigVMVariableDescription.h"
+#include "RigVMCore/RigVMVariableDescription.h"
 #include "RigVMVariableNode.generated.h"
 
 #define UE_API RIGVMDEVELOPER_API
diff --git a/Engine/Plugins/Runtime/RigVM/Source/RigVMDeveloper/Public/RigVMModel/RigVMClient.h b/Engine/Plugins/Runtime/RigVM/Source/RigVMDeveloper/Public/RigVMModel/RigVMClient.h
--- a/Engine/Plugins/Runtime/RigVM/Source/RigVMDeveloper/Public/RigVMModel/RigVMClient.h
+++ b/Engine/Plugins/Runtime/RigVM/Source/RigVMDeveloper/Public/RigVMModel/RigVMClient.h
@@ -267,6 +267,7 @@ public:
 	UE_API URigVMController* GetControllerByName(const FString InGraphName) const;
 	UE_API bool RemoveController(const URigVMGraph* InModel);
 	URigVMFunctionLibrary* GetFunctionLibrary() const { return FunctionLibrary; }
+	UE_API void SetGetFunctionHostObjectPathDelegate(URigVMFunctionLibrary* FunctionLibrary);
 	UE_API URigVMFunctionLibrary* GetOrCreateFunctionLibrary(bool bSetupUndoRedo, const FObjectInitializer* ObjectInitializer = nullptr, bool bCreateController = true);
 	UE_API URigVMFunctionLibrary* GetOrCreateFunctionLibrary(TSubclassOf<URigVMSchema> InSchemaClass, bool bSetupUndoRedo, const FObjectInitializer* ObjectInitializer = nullptr, bool bCreateController = true);
 	UE_API TArray<FName> GetEntryNames(UScriptStruct* InUnitScriptStructFilter = nullptr) const;
@@ -412,7 +413,8 @@ private:
 
 	friend class UEngineTestClientHost;
 	friend class URigVMBlueprint;
-	friend class IRigVMAssetInterface;
+	friend class IRigVMEditorAssetInterface;
+	friend class FRigVMExampleAssetActions;
 };
 
 template<>
diff --git a/Engine/Plugins/Runtime/RigVM/Source/RigVMDeveloper/Public/RigVMModel/RigVMController.h b/Engine/Plugins/Runtime/RigVM/Source/RigVMDeveloper/Public/RigVMModel/RigVMController.h
--- a/Engine/Plugins/Runtime/RigVM/Source/RigVMDeveloper/Public/RigVMModel/RigVMController.h
+++ b/Engine/Plugins/Runtime/RigVM/Source/RigVMDeveloper/Public/RigVMModel/RigVMController.h
@@ -1704,7 +1704,7 @@ private:
 	friend class URigVMActionStack;
 	friend struct FRigVMBaseAction;
 	friend class URigVMCompiler;
-	friend class IRigVMAssetInterface;
+	friend class IRigVMEditorAssetInterface;
 	friend struct FRigVMControllerObjectFactory;
 	friend struct FRigVMSetPinDefaultValueAction;
 	friend struct FRigVMAddRerouteNodeAction;
diff --git a/Engine/Plugins/Runtime/RigVM/Source/RigVMDeveloper/Public/RigVMModel/RigVMExternalDependency.h b/Engine/Plugins/Runtime/RigVM/Source/RigVMDeveloper/Public/RigVMModel/RigVMExternalDependency.h
--- a/Engine/Plugins/Runtime/RigVM/Source/RigVMDeveloper/Public/RigVMModel/RigVMExternalDependency.h
+++ b/Engine/Plugins/Runtime/RigVM/Source/RigVMDeveloper/Public/RigVMModel/RigVMExternalDependency.h
@@ -90,7 +90,7 @@ protected:
 
 	static UE_API TArray<FName> DependencyCategories;
 
-	friend class IRigVMAssetInterface;
+	friend class IRigVMEditorAssetInterface;
 };
 
 #undef UE_API
diff --git a/Engine/Plugins/Runtime/RigVM/Source/RigVMDeveloper/Public/RigVMModel/RigVMFunctionLibrary.h b/Engine/Plugins/Runtime/RigVM/Source/RigVMDeveloper/Public/RigVMModel/RigVMFunctionLibrary.h
--- a/Engine/Plugins/Runtime/RigVM/Source/RigVMDeveloper/Public/RigVMModel/RigVMFunctionLibrary.h
+++ b/Engine/Plugins/Runtime/RigVM/Source/RigVMDeveloper/Public/RigVMModel/RigVMFunctionLibrary.h
@@ -101,8 +101,8 @@ private:
 
 	friend class URigVMController;
 	friend class URigVMCompiler;
-	friend class IRigVMAssetInterface;
-	friend class IControlRigAssetInterface;
+	friend class IRigVMEditorAssetInterface;
+	friend class IControlRigEditorAssetInterface;
 	friend struct FRigVMClient;
 	friend class UControlRigBlueprint;
 };
diff --git a/Engine/Plugins/Runtime/RigVM/Source/RigVMDeveloper/Public/RigVMModel/RigVMGraph.h b/Engine/Plugins/Runtime/RigVM/Source/RigVMDeveloper/Public/RigVMModel/RigVMGraph.h
--- a/Engine/Plugins/Runtime/RigVM/Source/RigVMDeveloper/Public/RigVMModel/RigVMGraph.h
+++ b/Engine/Plugins/Runtime/RigVM/Source/RigVMDeveloper/Public/RigVMModel/RigVMGraph.h
@@ -10,7 +10,6 @@
 #include "Nodes/RigVMFunctionEntryNode.h"
 #include "Nodes/RigVMFunctionReturnNode.h"
 #include "RigVMCompiler/RigVMAST.h"
-#include "UObject/Interface.h"
 #include "ScribbleGraph.h"
 #include "RigVMGraph.generated.h"
 
@@ -252,7 +251,7 @@ private:
 
 	friend class URigVMController;
 	friend class URigVMBlueprint;
-	friend class IRigVMAssetInterface;
+	friend class IRigVMEditorAssetInterface;
 	friend class FRigVMControllerCompileBracketScope;
 	friend class URigVMCompiler;
 	friend class URigVMSchema;
diff --git a/Engine/Plugins/Runtime/RigVM/Source/RigVMEditor/Private/EdGraph/NodeSpawners/RigVMEdGraphEnumNodeSpawner.cpp b/Engine/Plugins/Runtime/RigVM/Source/RigVMEditor/Private/EdGraph/NodeSpawners/RigVMEdGraphEnumNodeSpawner.cpp
--- a/Engine/Plugins/Runtime/RigVM/Source/RigVMEditor/Private/EdGraph/NodeSpawners/RigVMEdGraphEnumNodeSpawner.cpp
+++ b/Engine/Plugins/Runtime/RigVM/Source/RigVMEditor/Private/EdGraph/NodeSpawners/RigVMEdGraphEnumNodeSpawner.cpp
@@ -51,7 +51,7 @@ URigVMEdGraphNode* URigVMEdGraphEnumNodeSpawner::Invoke(URigVMEdGraph* ParentGra
 	// First create a backing member for our node
 	URigVMEdGraph* RigGraph = Cast<URigVMEdGraph>(ParentGraph);
 	if(RigGraph == nullptr) return nullptr;
-	FRigVMAssetInterfacePtr RigBlueprint = FRigVMBlueprintUtils::FindAssetForGraph(ParentGraph);
+	FRigVMEditorAssetInterfacePtr RigBlueprint = FRigVMBlueprintUtils::FindAssetForGraph(ParentGraph);
 	check(RigBlueprint);
 
 #if WITH_EDITOR
diff --git a/Engine/Plugins/Runtime/RigVM/Source/RigVMEditor/Private/EdGraph/NodeSpawners/RigVMEdGraphFunctionRefNodeSpawner.cpp b/Engine/Plugins/Runtime/RigVM/Source/RigVMEditor/Private/EdGraph/NodeSpawners/RigVMEdGraphFunctionRefNodeSpawner.cpp
--- a/Engine/Plugins/Runtime/RigVM/Source/RigVMEditor/Private/EdGraph/NodeSpawners/RigVMEdGraphFunctionRefNodeSpawner.cpp
+++ b/Engine/Plugins/Runtime/RigVM/Source/RigVMEditor/Private/EdGraph/NodeSpawners/RigVMEdGraphFunctionRefNodeSpawner.cpp
@@ -152,7 +152,7 @@ URigVMEdGraphFunctionRefNodeSpawner::URigVMEdGraphFunctionRefNodeSpawner(const F
 	MenuCategory = FText::FromString(Category);
 	MenuKeywords = FText::FromString(InPublicFunction.Keywords);
 
-	if(const FRigVMAssetInterfacePtr ReferencedBlueprint = InAssetData.FastGetAsset(false))
+	if(const FRigVMEditorAssetInterfacePtr ReferencedBlueprint = InAssetData.FastGetAsset(false))
 	{
 		if(const URigVMFunctionLibrary* FunctionLibrary = ReferencedBlueprint->GetLocalFunctionLibrary())
 		{
@@ -232,9 +232,9 @@ URigVMEdGraphNode* URigVMEdGraphFunctionRefNodeSpawner::Invoke(URigVMEdGraph* Pa
 		{
 			if (UObject* BlueprintObj = GetResolvedAssetPath().TryLoad())
 			{
-				if (BlueprintObj->Implements<URigVMAssetInterface>())
+				if (BlueprintObj->Implements<URigVMEditorAssetInterface>())
 				{
-					FRigVMAssetInterfacePtr Blueprint = BlueprintObj;
+					FRigVMEditorAssetInterfacePtr Blueprint = BlueprintObj;
 					ReferencedPublicFunctionHeader = Blueprint->GetLocalFunctionLibrary()->FindFunction(ReferencedPublicFunctionHeader.Name)->GetFunctionHeader();
 				}
 			}
@@ -251,7 +251,7 @@ URigVMEdGraphNode* URigVMEdGraphFunctionRefNodeSpawner::Invoke(URigVMEdGraph* Pa
 		}
 #endif
 
-		FRigVMAssetInterfacePtr Blueprint = FRigVMBlueprintUtils::FindAssetForGraph(ParentGraph);
+		FRigVMEditorAssetInterfacePtr Blueprint = FRigVMBlueprintUtils::FindAssetForGraph(ParentGraph);
 		NewNode = SpawnNode(ParentGraph, Blueprint, ReferencedPublicFunctionHeader, Location);
 	}
 	else
@@ -269,7 +269,7 @@ URigVMEdGraphNode* URigVMEdGraphFunctionRefNodeSpawner::Invoke(URigVMEdGraph* Pa
 	return NewNode;
 }
 
-URigVMEdGraphNode* URigVMEdGraphFunctionRefNodeSpawner::SpawnNode(UEdGraph* ParentGraph, FRigVMAssetInterfacePtr Blueprint, FRigVMGraphFunctionHeader& InFunction, FVector2D const Location)
+URigVMEdGraphNode* URigVMEdGraphFunctionRefNodeSpawner::SpawnNode(UEdGraph* ParentGraph, FRigVMEditorAssetInterfacePtr Blueprint, FRigVMGraphFunctionHeader& InFunction, FVector2D const Location)
 {
 	URigVMEdGraphNode* NewNode = nullptr;
 	URigVMEdGraph* RigGraph = Cast<URigVMEdGraph>(ParentGraph);
@@ -356,7 +356,7 @@ bool URigVMEdGraphFunctionRefNodeSpawner::IsTemplateNodeFilteredOut(TArray<UObje
 		{
 			for (UObject* Blueprint : InAssets)
 			{
-				if (IRigVMAssetInterface* Asset = Cast<IRigVMAssetInterface>(Blueprint))
+				if (IRigVMEditorAssetInterface* Asset = Cast<IRigVMEditorAssetInterface>(Blueprint))
 				{
 					if (Asset->GetRigVMClientHost()->GetRigVMGraphFunctionHost() != ReferencedPublicFunctionHeader.GetFunctionHost())
 					{
diff --git a/Engine/Plugins/Runtime/RigVM/Source/RigVMEditor/Private/EdGraph/NodeSpawners/RigVMEdGraphInvokeEntryNodeSpawner.cpp b/Engine/Plugins/Runtime/RigVM/Source/RigVMEditor/Private/EdGraph/NodeSpawners/RigVMEdGraphInvokeEntryNodeSpawner.cpp
--- a/Engine/Plugins/Runtime/RigVM/Source/RigVMEditor/Private/EdGraph/NodeSpawners/RigVMEdGraphInvokeEntryNodeSpawner.cpp
+++ b/Engine/Plugins/Runtime/RigVM/Source/RigVMEditor/Private/EdGraph/NodeSpawners/RigVMEdGraphInvokeEntryNodeSpawner.cpp
@@ -21,7 +21,7 @@
 
 #define LOCTEXT_NAMESPACE "RigVMEdGraphInvokeEntryNodeSpawner"
 
-URigVMEdGraphInvokeEntryNodeSpawner::URigVMEdGraphInvokeEntryNodeSpawner(FRigVMAssetInterfacePtr InBlueprint, const FName& InEntryName, const FText& InMenuDesc, const FText& InCategory, const FText& InTooltip)
+URigVMEdGraphInvokeEntryNodeSpawner::URigVMEdGraphInvokeEntryNodeSpawner(FRigVMEditorAssetInterfacePtr InBlueprint, const FName& InEntryName, const FText& InMenuDesc, const FText& InCategory, const FText& InTooltip)
 {
 	Blueprint = InBlueprint;
 	EntryName = InEntryName;
@@ -48,7 +48,7 @@ URigVMEdGraphNode* URigVMEdGraphInvokeEntryNodeSpawner::Invoke(URigVMEdGraph* Pa
 	// First create a backing member for our node
 	URigVMEdGraph* RigGraph = Cast<URigVMEdGraph>(ParentGraph);
 	if(RigGraph == nullptr) return nullptr;
-	FRigVMAssetInterfacePtr RigBlueprint = FRigVMBlueprintUtils::FindAssetForGraph(ParentGraph);
+	FRigVMEditorAssetInterfacePtr RigBlueprint = FRigVMBlueprintUtils::FindAssetForGraph(ParentGraph);
 	check(RigBlueprint);
 
 #if WITH_EDITOR
diff --git a/Engine/Plugins/Runtime/RigVM/Source/RigVMEditor/Private/EdGraph/NodeSpawners/RigVMEdGraphTemplateNodeSpawner.cpp b/Engine/Plugins/Runtime/RigVM/Source/RigVMEditor/Private/EdGraph/NodeSpawners/RigVMEdGraphTemplateNodeSpawner.cpp
--- a/Engine/Plugins/Runtime/RigVM/Source/RigVMEditor/Private/EdGraph/NodeSpawners/RigVMEdGraphTemplateNodeSpawner.cpp
+++ b/Engine/Plugins/Runtime/RigVM/Source/RigVMEditor/Private/EdGraph/NodeSpawners/RigVMEdGraphTemplateNodeSpawner.cpp
@@ -182,14 +182,14 @@ URigVMEdGraphNode* URigVMEdGraphTemplateNodeSpawner::Invoke(URigVMEdGraph* Paren
 			return NewNode;
 		}
 
-		FRigVMAssetInterfacePtr Blueprint = FRigVMBlueprintUtils::FindAssetForGraph(ParentGraph);
+		FRigVMEditorAssetInterfacePtr Blueprint = FRigVMBlueprintUtils::FindAssetForGraph(ParentGraph);
 		NewNode = SpawnNode(ParentGraph, Blueprint, Template, Location);
 	}
 
 	return NewNode;
 }
 
-URigVMEdGraphNode* URigVMEdGraphTemplateNodeSpawner::SpawnNode(URigVMEdGraph* ParentGraph, FRigVMAssetInterfacePtr RigBlueprint, const FRigVMTemplate* Template, FVector2D const Location)
+URigVMEdGraphNode* URigVMEdGraphTemplateNodeSpawner::SpawnNode(URigVMEdGraph* ParentGraph, FRigVMEditorAssetInterfacePtr RigBlueprint, const FRigVMTemplate* Template, FVector2D const Location)
 {
 	URigVMEdGraphNode* NewNode = nullptr;
 	URigVMEdGraph* RigGraph = Cast<URigVMEdGraph>(ParentGraph);
diff --git a/Engine/Plugins/Runtime/RigVM/Source/RigVMEditor/Private/EdGraph/NodeSpawners/RigVMEdGraphUnitNodeSpawner.cpp b/Engine/Plugins/Runtime/RigVM/Source/RigVMEditor/Private/EdGraph/NodeSpawners/RigVMEdGraphUnitNodeSpawner.cpp
--- a/Engine/Plugins/Runtime/RigVM/Source/RigVMEditor/Private/EdGraph/NodeSpawners/RigVMEdGraphUnitNodeSpawner.cpp
+++ b/Engine/Plugins/Runtime/RigVM/Source/RigVMEditor/Private/EdGraph/NodeSpawners/RigVMEdGraphUnitNodeSpawner.cpp
@@ -127,14 +127,14 @@ URigVMEdGraphNode* URigVMEdGraphUnitNodeSpawner::Invoke(URigVMEdGraph* ParentGra
 		}
 #endif
 
-		FRigVMAssetInterfacePtr Blueprint = FRigVMBlueprintUtils::FindAssetForGraph(ParentGraph);
+		FRigVMEditorAssetInterfacePtr Blueprint = FRigVMBlueprintUtils::FindAssetForGraph(ParentGraph);
 		NewNode = SpawnNode(ParentGraph, Blueprint, StructTemplate, MethodName, Location);
 	}
 
 	return NewNode;
 }
 
-URigVMEdGraphNode* URigVMEdGraphUnitNodeSpawner::SpawnNode(URigVMEdGraph* ParentGraph, FRigVMAssetInterfacePtr RigBlueprint, UScriptStruct* StructTemplate, const FName& InMethodName, FVector2D const Location)
+URigVMEdGraphNode* URigVMEdGraphUnitNodeSpawner::SpawnNode(URigVMEdGraph* ParentGraph, FRigVMEditorAssetInterfacePtr RigBlueprint, UScriptStruct* StructTemplate, const FName& InMethodName, FVector2D const Location)
 {
 	URigVMEdGraphNode* NewNode = nullptr;
 	URigVMEdGraph* RigGraph = Cast<URigVMEdGraph>(ParentGraph);
@@ -350,7 +350,7 @@ URigVMEdGraphNode* URigVMEdGraphUnitNodeSpawner::SpawnNode(URigVMEdGraph* Parent
 	return NewNode;
 }
 
-void URigVMEdGraphUnitNodeSpawner::HookupMutableNode(URigVMNode* InModelNode, FRigVMAssetInterfacePtr InRigBlueprint)
+void URigVMEdGraphUnitNodeSpawner::HookupMutableNode(URigVMNode* InModelNode, FRigVMEditorAssetInterfacePtr InRigBlueprint)
 {
 	if(!InRigBlueprint->GetRigVMEditorSettings()->bAutoLinkMutableNodes)
 	{
diff --git a/Engine/Plugins/Runtime/RigVM/Source/RigVMEditor/Private/EdGraph/NodeSpawners/RigVMEdGraphVariableNodeSpawner.cpp b/Engine/Plugins/Runtime/RigVM/Source/RigVMEditor/Private/EdGraph/NodeSpawners/RigVMEdGraphVariableNodeSpawner.cpp
--- a/Engine/Plugins/Runtime/RigVM/Source/RigVMEditor/Private/EdGraph/NodeSpawners/RigVMEdGraphVariableNodeSpawner.cpp
+++ b/Engine/Plugins/Runtime/RigVM/Source/RigVMEditor/Private/EdGraph/NodeSpawners/RigVMEdGraphVariableNodeSpawner.cpp
@@ -14,6 +14,7 @@
 #include "RigVMModel/RigVMController.h"
 #include "RigVMBlueprintLegacy.h"
 #include "RigVMBlueprintUtils.h"
+#include "RigVMDeveloperTypeUtils.h"
 #include "RigVMModel/Nodes/RigVMVariableNode.h"
 
 #if WITH_EDITOR
@@ -22,7 +23,7 @@
 
 #define LOCTEXT_NAMESPACE "RigVMEdGraphVariableNodeSpawner"
 
-URigVMEdGraphVariableNodeSpawner::URigVMEdGraphVariableNodeSpawner(FRigVMAssetInterfacePtr InBlueprint, const FRigVMExternalVariable& InExternalVariable, bool bInIsGetter, const FText& InMenuDesc, const FText& InCategory, const FText& InTooltip)
+URigVMEdGraphVariableNodeSpawner::URigVMEdGraphVariableNodeSpawner(FRigVMEditorAssetInterfacePtr InBlueprint, const FRigVMExternalVariable& InExternalVariable, bool bInIsGetter, const FText& InMenuDesc, const FText& InCategory, const FText& InTooltip)
 {
 	Blueprint = InBlueprint;
 	ExternalVariable = InExternalVariable;
@@ -39,7 +40,7 @@ URigVMEdGraphVariableNodeSpawner::URigVMEdGraphVariableNodeSpawner(FRigVMAssetIn
 }
 
 URigVMEdGraphVariableNodeSpawner::URigVMEdGraphVariableNodeSpawner(
-	FRigVMAssetInterfacePtr InBlueprint, URigVMGraph* InGraphOwner, const FRigVMGraphVariableDescription& InLocalVariable,
+	FRigVMEditorAssetInterfacePtr InBlueprint, URigVMGraph* InGraphOwner, const FRigVMGraphVariableDescription& InLocalVariable,
 	bool bInIsGetter, const FText& InMenuDesc, const FText& InCategory, const FText& InTooltip)
 		: URigVMEdGraphVariableNodeSpawner(InBlueprint, InLocalVariable.ToExternalVariable(), bInIsGetter, InMenuDesc, InCategory, InTooltip)
 {
@@ -98,7 +99,7 @@ URigVMEdGraphNode* URigVMEdGraphVariableNodeSpawner::Invoke(URigVMEdGraph* Paren
 	// First create a backing member for our node
 	URigVMEdGraph* RigGraph = Cast<URigVMEdGraph>(ParentGraph);
 	if(RigGraph == nullptr) return nullptr;
-	FRigVMAssetInterfacePtr RigBlueprint = FRigVMBlueprintUtils::FindAssetForGraph(ParentGraph);
+	FRigVMEditorAssetInterfacePtr RigBlueprint = FRigVMBlueprintUtils::FindAssetForGraph(ParentGraph);
 	check(RigBlueprint);
 
 	FName MemberName = NAME_None;
diff --git a/Engine/Plugins/Runtime/RigVM/Source/RigVMEditor/Private/EdGraph/RigVMEdGraphConnectionDrawingPolicy.cpp b/Engine/Plugins/Runtime/RigVM/Source/RigVMEditor/Private/EdGraph/RigVMEdGraphConnectionDrawingPolicy.cpp
--- a/Engine/Plugins/Runtime/RigVM/Source/RigVMEditor/Private/EdGraph/RigVMEdGraphConnectionDrawingPolicy.cpp
+++ b/Engine/Plugins/Runtime/RigVM/Source/RigVMEditor/Private/EdGraph/RigVMEdGraphConnectionDrawingPolicy.cpp
@@ -4,7 +4,7 @@
 #include "EdGraph/RigVMEdGraph.h"
 #include "EdGraph/RigVMEdGraphNode.h"
 #include "RigVMHost.h"
-#include "RigVMAsset.h"
+#include "RigVMEditorAsset.h"
 #include "RigVMBlueprintUtils.h"
 #include "RigVMModel/RigVMController.h"
 #include "Kismet2/BlueprintEditorUtils.h"
@@ -323,7 +323,7 @@ void FRigVMEdGraphConnectionDrawingPolicy::DetermineWiringStyle(UEdGraphPin* Out
 
 			if (OutputInstructionIndex != INDEX_NONE && InputInstructionIndex != INDEX_NONE)
 			{
-				if (FRigVMAssetInterfacePtr RigBlueprint = FRigVMBlueprintUtils::FindAssetForNode(OutputNode))
+				if (FRigVMEditorAssetInterfacePtr RigBlueprint = FRigVMBlueprintUtils::FindAssetForNode(OutputNode))
 				{
 					if (URigVMHost* Host = Cast<URigVMHost>(RigBlueprint->GetObjectBeingDebugged()))
 					{
diff --git a/Engine/Plugins/Runtime/RigVM/Source/RigVMEditor/Private/EdGraph/RigVMEdGraphPanelNodeFactory.cpp b/Engine/Plugins/Runtime/RigVM/Source/RigVMEditor/Private/EdGraph/RigVMEdGraphPanelNodeFactory.cpp
--- a/Engine/Plugins/Runtime/RigVM/Source/RigVMEditor/Private/EdGraph/RigVMEdGraphPanelNodeFactory.cpp
+++ b/Engine/Plugins/Runtime/RigVM/Source/RigVMEditor/Private/EdGraph/RigVMEdGraphPanelNodeFactory.cpp
@@ -13,7 +13,7 @@
 
 FName FRigVMEdGraphPanelNodeFactory::GetFactoryName() const
 {
-	return IRigVMAssetInterface::RigVMPanelNodeFactoryName;
+	return IRigVMEditorAssetInterface::RigVMPanelNodeFactoryName;
 }
 
 TSharedPtr<SGraphNode> FRigVMEdGraphPanelNodeFactory::CreateNode(UEdGraphNode* Node) const
@@ -21,7 +21,7 @@ TSharedPtr<SGraphNode> FRigVMEdGraphPanelNodeFactory::CreateNode(UEdGraphNode* N
 	// we need to check if this is the right factory for the implementation
 	if(const URigVMEdGraph* EdGraph = Node->GetTypedOuter<URigVMEdGraph>())
 	{
-		if(const FRigVMAssetInterfacePtr Blueprint = EdGraph->GetAssetDefaultObject())
+		if(const FRigVMEditorAssetInterfacePtr Blueprint = EdGraph->GetAssetDefaultObject())
 		{
 			if(Blueprint->GetPanelNodeFactoryName() != GetFactoryName())
 			{
diff --git a/Engine/Plugins/Runtime/RigVM/Source/RigVMEditor/Private/EdGraph/RigVMEdGraphPanelPinFactory.cpp b/Engine/Plugins/Runtime/RigVM/Source/RigVMEditor/Private/EdGraph/RigVMEdGraphPanelPinFactory.cpp
--- a/Engine/Plugins/Runtime/RigVM/Source/RigVMEditor/Private/EdGraph/RigVMEdGraphPanelPinFactory.cpp
+++ b/Engine/Plugins/Runtime/RigVM/Source/RigVMEditor/Private/EdGraph/RigVMEdGraphPanelPinFactory.cpp
@@ -24,7 +24,7 @@
 
 FName FRigVMEdGraphPanelPinFactory::GetFactoryName() const
 {
-	return IRigVMAssetInterface::RigVMPanelPinFactoryName;
+	return IRigVMEditorAssetInterface::RigVMPanelPinFactoryName;
 }
 
 TSharedPtr<SGraphPin> FRigVMEdGraphPanelPinFactory::CreatePin(UEdGraphPin* InPin) const
@@ -39,7 +39,7 @@ TSharedPtr<SGraphPin> FRigVMEdGraphPanelPinFactory::CreatePin(UEdGraphPin* InPin
 	{
 		if(const URigVMEdGraph* EdGraph = EdGraphNode->GetTypedOuter<URigVMEdGraph>())
 		{
-			if(const FRigVMAssetInterfacePtr Blueprint = EdGraph->GetAssetDefaultObject())
+			if(const FRigVMEditorAssetInterfacePtr Blueprint = EdGraph->GetAssetDefaultObject())
 			{
 				if(Blueprint->GetPanelPinFactoryName() != GetFactoryName())
 				{
@@ -90,7 +90,7 @@ TSharedPtr<SGraphPin> FRigVMEdGraphPanelPinFactory::CreatePin_Internal(UEdGraphP
 			{
 				if (ModelPin->IsBoundToVariable())
 				{
-					if (FRigVMAssetInterfacePtr Blueprint = IRigVMAssetInterface::GetInterfaceOuter(RigGraph))
+					if (FRigVMEditorAssetInterfacePtr Blueprint = IRigVMEditorAssetInterface::GetInterfaceOuter(RigGraph))
 					{
 						return SNew(SRigVMGraphPinVariableBinding, InPin)
 							.ModelPins({ModelPin})
diff --git a/Engine/Plugins/Runtime/RigVM/Source/RigVMEditor/Private/Editor/RemoveUnusedMembers/RigVMEditorRemoveUnusedMembers.cpp b/Engine/Plugins/Runtime/RigVM/Source/RigVMEditor/Private/Editor/RemoveUnusedMembers/RigVMEditorRemoveUnusedMembers.cpp
--- a/Engine/Plugins/Runtime/RigVM/Source/RigVMEditor/Private/Editor/RemoveUnusedMembers/RigVMEditorRemoveUnusedMembers.cpp
+++ b/Engine/Plugins/Runtime/RigVM/Source/RigVMEditor/Private/Editor/RemoveUnusedMembers/RigVMEditorRemoveUnusedMembers.cpp
@@ -107,7 +107,7 @@ namespace UE::RigVMEditor::Private
 	/** Removes function members */
 	template <>
 	void RemoveMembers<URigVMFunctionReferenceNode>(
-		const TScriptInterface<IRigVMAssetInterface>& Asset,
+		const TScriptInterface<IRigVMEditorAssetInterface>& Asset,
 		const TArray<FName>& MemberNamesToRemove)
 	{
 		if (!Asset)
@@ -138,7 +138,7 @@ namespace UE::RigVMEditor::Private
 	/** Removes variable members */
 	template <>
 	void RemoveMembers<URigVMVariableNode>(
-		const TScriptInterface<IRigVMAssetInterface>& Asset,
+		const TScriptInterface<IRigVMEditorAssetInterface>& Asset,
 		const TArray<FName>& MemberNamesToRemove)
 	{
 		UObject* Object = Asset ? Asset->GetObject() : nullptr;
@@ -175,7 +175,7 @@ namespace UE::RigVMEditor::Private
 
 	template<>
 	void FindUnusedMemberNames<URigVMFunctionReferenceNode>(
-		const TScriptInterface<IRigVMAssetInterface>& InAsset,
+		const TScriptInterface<IRigVMEditorAssetInterface>& InAsset,
 		const TSharedRef<FRigVMEdGraphNodeRegistry>& InEdGraphNodeRegistry,
 		TMap<FRigVMUnusedMemberCategory, TArray<FName>>& OutCategoryToUnusedMemberNamesMap)
 	{
@@ -275,7 +275,7 @@ namespace UE::RigVMEditor::Private
 
 	template<>
 	void FindUnusedMemberNames<URigVMVariableNode>(
-		const TScriptInterface<IRigVMAssetInterface>& InAsset,
+		const TScriptInterface<IRigVMEditorAssetInterface>& InAsset,
 		const TSharedRef<FRigVMEdGraphNodeRegistry>& InEdGraphNodeRegistry,
 		TMap<FRigVMUnusedMemberCategory, TArray<FName>>& OutCategoryToUnusedMemberNamesMap)
 	{
diff --git a/Engine/Plugins/Runtime/RigVM/Source/RigVMEditor/Private/Editor/RemoveUnusedMembers/RigVMEditorRemoveUnusedMembers.h b/Engine/Plugins/Runtime/RigVM/Source/RigVMEditor/Private/Editor/RemoveUnusedMembers/RigVMEditorRemoveUnusedMembers.h
--- a/Engine/Plugins/Runtime/RigVM/Source/RigVMEditor/Private/Editor/RemoveUnusedMembers/RigVMEditorRemoveUnusedMembers.h
+++ b/Engine/Plugins/Runtime/RigVM/Source/RigVMEditor/Private/Editor/RemoveUnusedMembers/RigVMEditorRemoveUnusedMembers.h
@@ -4,7 +4,7 @@
 
 #include "Editor/RigVMEdGraphNodeRegistry.h"
 #include "Framework/Notifications/NotificationManager.h"
-#include "RigVMAsset.h"
+#include "RigVMEditorAsset.h"
 #include "RigVMEditorRemoveUnusedMembersCategory.h"
 #include "RigVMEditorRemoveUnusedMembersDialog.h"
 #include "RigVMModel/Nodes/RigVMFunctionReferenceNode.h"
@@ -32,7 +32,7 @@ namespace UE::RigVMEditor
 		/** Finds unused member names. Considers members with disconnected nodes as unused */
 		template <typename TRigVMNode> requires ConceptTypeSupportsRemoveUnusedMembers<TRigVMNode>
 		void FindUnusedMemberNames(
-			const TScriptInterface<IRigVMAssetInterface>& InAsset,
+			const TScriptInterface<IRigVMEditorAssetInterface>& InAsset,
 			const TSharedRef<FRigVMEdGraphNodeRegistry>& InEdGraphNodeRegistry,
 			TMap<FRigVMUnusedMemberCategory, TArray<FName>>& OutCategoryToUnusedMemberNamesMap);
 
@@ -46,7 +46,7 @@ namespace UE::RigVMEditor
 		/** Removes members */
 		template <typename TRigVMNode> requires ConceptTypeSupportsRemoveUnusedMembers<TRigVMNode>
 		void RemoveMembers(
-			const TScriptInterface<IRigVMAssetInterface>& Asset, 
+			const TScriptInterface<IRigVMEditorAssetInterface>& Asset, 
 			const TArray<FName>& MemberNamesToRemove);
 	}
 
@@ -59,7 +59,7 @@ namespace UE::RigVMEditor
 	 */
 	template <typename TRigVMNode> requires ConceptTypeSupportsRemoveUnusedMembers<TRigVMNode>
 	void RemoveUnusedMembersFromAsset(
-		const TScriptInterface<IRigVMAssetInterface>& Asset,
+		const TScriptInterface<IRigVMEditorAssetInterface>& Asset,
 		const FText& TransactionText = FText::GetEmpty(),
 		const FText& AllMembersReferencedMessage = FText::GetEmpty())
 	{
diff --git a/Engine/Plugins/Runtime/RigVM/Source/RigVMEditor/Private/Editor/RemoveUnusedMembers/RigVMEditorRemoveUnusedMembersCategory.h b/Engine/Plugins/Runtime/RigVM/Source/RigVMEditor/Private/Editor/RemoveUnusedMembers/RigVMEditorRemoveUnusedMembersCategory.h
--- a/Engine/Plugins/Runtime/RigVM/Source/RigVMEditor/Private/Editor/RemoveUnusedMembers/RigVMEditorRemoveUnusedMembersCategory.h
+++ b/Engine/Plugins/Runtime/RigVM/Source/RigVMEditor/Private/Editor/RemoveUnusedMembers/RigVMEditorRemoveUnusedMembersCategory.h
@@ -6,7 +6,7 @@
 #include "Internationalization/Text.h"
 #include "UObject/NameTypes.h"
 
-class IRigVMAssetInterface;
+class IRigVMEditorAssetInterface;
 template <typename InInterfaceType> class TScriptInterface;
 
 namespace UE::RigVMEditor
diff --git a/Engine/Plugins/Runtime/RigVM/Source/RigVMEditor/Private/Editor/RemoveUnusedMembers/RigVMEditorRemoveUnusedMembersDialog.h b/Engine/Plugins/Runtime/RigVM/Source/RigVMEditor/Private/Editor/RemoveUnusedMembers/RigVMEditorRemoveUnusedMembersDialog.h
--- a/Engine/Plugins/Runtime/RigVM/Source/RigVMEditor/Private/Editor/RemoveUnusedMembers/RigVMEditorRemoveUnusedMembersDialog.h
+++ b/Engine/Plugins/Runtime/RigVM/Source/RigVMEditor/Private/Editor/RemoveUnusedMembers/RigVMEditorRemoveUnusedMembersDialog.h
@@ -7,7 +7,7 @@
 #include "RigVMEditorRemoveUnusedMembersCategory.h"
 #include "UObject/NameTypes.h"
 
-class IRigVMAssetInterface;
+class IRigVMEditorAssetInterface;
 template <typename InInterfaceType> class TScriptInterface;
 
 namespace UE::RigVMEditor
diff --git a/Engine/Plugins/Runtime/RigVM/Source/RigVMEditor/Private/Editor/RigVMActionMenuBuilder.cpp b/Engine/Plugins/Runtime/RigVM/Source/RigVMEditor/Private/Editor/RigVMActionMenuBuilder.cpp
--- a/Engine/Plugins/Runtime/RigVM/Source/RigVMEditor/Private/Editor/RigVMActionMenuBuilder.cpp
+++ b/Engine/Plugins/Runtime/RigVM/Source/RigVMEditor/Private/Editor/RigVMActionMenuBuilder.cpp
@@ -15,6 +15,7 @@
 #include "BlueprintVariableNodeSpawner.h"
 #include "EditorCategoryUtils.h"
 #include "ObjectEditorUtils.h"
+#include "EdGraph/NodeSpawners/RigVMEdGraphNodeSpawner.h"
 
 #if WITH_RIGVMLEGACYEDITOR
 #include "BlueprintDragDropMenuItem.h"
diff --git a/Engine/Plugins/Runtime/RigVM/Source/RigVMEditor/Private/Editor/RigVMActionMenuItem.cpp b/Engine/Plugins/Runtime/RigVM/Source/RigVMEditor/Private/Editor/RigVMActionMenuItem.cpp
--- a/Engine/Plugins/Runtime/RigVM/Source/RigVMEditor/Private/Editor/RigVMActionMenuItem.cpp
+++ b/Engine/Plugins/Runtime/RigVM/Source/RigVMEditor/Private/Editor/RigVMActionMenuItem.cpp
@@ -90,7 +90,7 @@ static void FBlueprintMenuActionItemImpl::DirtyBlueprintFromNewNode(UEdGraphNode
 	UEdGraph const* const NodeGraph = SpawnedNode->GetGraph();
 	check(NodeGraph != nullptr);
 	
-	FRigVMAssetInterfacePtr Blueprint = FRigVMBlueprintUtils::FindAssetForGraph(NodeGraph);
+	FRigVMEditorAssetInterfacePtr Blueprint = FRigVMBlueprintUtils::FindAssetForGraph(NodeGraph);
 	check(Blueprint != nullptr);
 
 	if (!Blueprint || !SpawnedNode)
diff --git a/Engine/Plugins/Runtime/RigVM/Source/RigVMEditor/Private/Editor/RigVMActionMenuUtils.cpp b/Engine/Plugins/Runtime/RigVM/Source/RigVMEditor/Private/Editor/RigVMActionMenuUtils.cpp
--- a/Engine/Plugins/Runtime/RigVM/Source/RigVMEditor/Private/Editor/RigVMActionMenuUtils.cpp
+++ b/Engine/Plugins/Runtime/RigVM/Source/RigVMEditor/Private/Editor/RigVMActionMenuUtils.cpp
@@ -6,6 +6,7 @@
 #include "Editor/RigVMActionMenuBuilder.h"
 #include "Editor/RigVMActionMenuItem.h"
 #include "K2Node.h"
+#include "EdGraph/NodeSpawners/RigVMEdGraphNodeSpawner.h"
 
 
 class UEdGraph;
diff --git a/Engine/Plugins/Runtime/RigVM/Source/RigVMEditor/Private/Editor/RigVMCommentNodeDetailCustomization.cpp b/Engine/Plugins/Runtime/RigVM/Source/RigVMEditor/Private/Editor/RigVMCommentNodeDetailCustomization.cpp
--- a/Engine/Plugins/Runtime/RigVM/Source/RigVMEditor/Private/Editor/RigVMCommentNodeDetailCustomization.cpp
+++ b/Engine/Plugins/Runtime/RigVM/Source/RigVMEditor/Private/Editor/RigVMCommentNodeDetailCustomization.cpp
@@ -6,7 +6,7 @@
 #include "IDetailChildrenBuilder.h"
 #include "Widgets/Input/SCheckBox.h"
 #include "Widgets/Input/STextComboBox.h"
-#include "RigVMAsset.h"
+#include "RigVMEditorAsset.h"
 #include "RigVMCore/RigVM.h"
 #include "Widgets/Colors/SColorBlock.h"
 #include "Widgets/Colors/SColorPicker.h"
@@ -31,7 +31,10 @@ void FRigVMCommentNodeDetailCustomization::CustomizeDetails(IDetailLayoutBuilder
 
 	if(ObjectsBeingCustomized[0].IsValid())
 	{
-		BlueprintBeingCustomized = ObjectsBeingCustomized[0]->GetImplementingOuter<IRigVMAssetInterface>();
+		if (IRigVMEditorAssetInterface* InterfacePtr = ObjectsBeingCustomized[0]->GetImplementingOuter<IRigVMEditorAssetInterface>())
+		{
+			BlueprintBeingCustomized = InterfacePtr->GetObject();
+		}
 	}
 
 	IDetailCategoryBuilder& Category = DetailBuilder.EditCategory(TEXT("Comment Node"));
diff --git a/Engine/Plugins/Runtime/RigVM/Source/RigVMEditor/Private/Editor/RigVMDetailsViewWrapperObject.cpp b/Engine/Plugins/Runtime/RigVM/Source/RigVMEditor/Private/Editor/RigVMDetailsViewWrapperObject.cpp
--- a/Engine/Plugins/Runtime/RigVM/Source/RigVMEditor/Private/Editor/RigVMDetailsViewWrapperObject.cpp
+++ b/Engine/Plugins/Runtime/RigVM/Source/RigVMEditor/Private/Editor/RigVMDetailsViewWrapperObject.cpp
@@ -1,6 +1,7 @@
 // Copyright Epic Games, Inc. All Rights Reserved.
 
 #include "Editor/RigVMDetailsViewWrapperObject.h"
+#include "RigVMDeveloperModule.h"
 #include "RigVMCore/RigVMStruct.h"
 #include "RigVMModel/Nodes/RigVMUnitNode.h"
 #include "RigVMModel/Nodes/RigVMDispatchNode.h"
diff --git a/Engine/Plugins/Runtime/RigVM/Source/RigVMEditor/Private/Editor/RigVMEdGraphNodeRegistry.cpp b/Engine/Plugins/Runtime/RigVM/Source/RigVMEditor/Private/Editor/RigVMEdGraphNodeRegistry.cpp
--- a/Engine/Plugins/Runtime/RigVM/Source/RigVMEditor/Private/Editor/RigVMEdGraphNodeRegistry.cpp
+++ b/Engine/Plugins/Runtime/RigVM/Source/RigVMEditor/Private/Editor/RigVMEdGraphNodeRegistry.cpp
@@ -4,7 +4,7 @@
 
 #include "EdGraph/RigVMEdGraph.h"
 #include "Editor.h"
-#include "RigVMAsset.h"
+#include "RigVMEditorAsset.h"
 #include "RigVMModel/RigVMNode.h"
 
 namespace UE::RigVMEditor
@@ -12,7 +12,7 @@ namespace UE::RigVMEditor
 	TMap<FName, TWeakPtr<FRigVMEdGraphNodeRegistry>> FRigVMEdGraphNodeRegistry::TypeIDToRegistryMap;
 
 	TSharedRef<FRigVMEdGraphNodeRegistry> FRigVMEdGraphNodeRegistry::GetOrCreateRegistry(
-		const TScriptInterface<IRigVMAssetInterface>& InRigVMAssetInterface, 
+		const TScriptInterface<IRigVMEditorAssetInterface>& InRigVMAssetInterface, 
 		UClass* InRigVMNodeClass,
 		const bool bForceUpdate)
 	{		
@@ -64,7 +64,7 @@ namespace UE::RigVMEditor
 		}
 	}
 
-	FRigVMEdGraphNodeRegistry::FRigVMEdGraphNodeRegistry(const TScriptInterface<IRigVMAssetInterface>& InRigVMAssetInterface, UClass* InRigVMNodeClass)
+	FRigVMEdGraphNodeRegistry::FRigVMEdGraphNodeRegistry(const TScriptInterface<IRigVMEditorAssetInterface>& InRigVMAssetInterface, UClass* InRigVMNodeClass)
 		: WeakAssetInterface(InRigVMAssetInterface)
 		, RigVMNodeClass(InRigVMNodeClass)
 	{}
diff --git a/Engine/Plugins/Runtime/RigVM/Source/RigVMEditor/Private/Editor/RigVMEdGraphNodeRegistry.h b/Engine/Plugins/Runtime/RigVM/Source/RigVMEditor/Private/Editor/RigVMEdGraphNodeRegistry.h
--- a/Engine/Plugins/Runtime/RigVM/Source/RigVMEditor/Private/Editor/RigVMEdGraphNodeRegistry.h
+++ b/Engine/Plugins/Runtime/RigVM/Source/RigVMEditor/Private/Editor/RigVMEdGraphNodeRegistry.h
@@ -8,7 +8,7 @@
 #include "UObject/WeakInterfacePtr.h"
 #include "UObject/WeakObjectPtr.h"
 
-class IRigVMAssetInterface;
+class IRigVMEditorAssetInterface;
 class UClass;
 class UObject;
 class URigVMEdGraphNode;
@@ -38,7 +38,7 @@ namespace UE::RigVMEditor
 		 * @return								The registry instance
 		 */
 		static TSharedRef<FRigVMEdGraphNodeRegistry> GetOrCreateRegistry(
-			const TScriptInterface<IRigVMAssetInterface>& InRigVMAssetInterface, 
+			const TScriptInterface<IRigVMEditorAssetInterface>& InRigVMAssetInterface, 
 			UClass* InRigVMNodeClass,
 			const bool bForceUpdate = false);
 
@@ -56,7 +56,7 @@ namespace UE::RigVMEditor
 
 	private:
 		/** Private constructor, hidden on purpose. Use GetOrCreateRegistry to create instances. */
-		FRigVMEdGraphNodeRegistry(const TScriptInterface<IRigVMAssetInterface>& InRigVMAssetInterface, UClass* InRigVMNodeClass);
+		FRigVMEdGraphNodeRegistry(const TScriptInterface<IRigVMEditorAssetInterface>& InRigVMAssetInterface, UClass* InRigVMNodeClass);
 
 		/** Initializes the registry */
 		void Initialize();
@@ -83,7 +83,7 @@ namespace UE::RigVMEditor
 		TArray<TWeakObjectPtr<URigVMEdGraphNode>> WeakDisconnectedEdGraphNodes;
 
 		/** The asset interface this registry refers to */
-		TWeakInterfacePtr<IRigVMAssetInterface> WeakAssetInterface;
+		TWeakInterfacePtr<IRigVMEditorAssetInterface> WeakAssetInterface;
 
 		/** The Rig VM Node type that is tracked in this registry */
 		UClass* RigVMNodeClass = nullptr;
