Control Rig: Split IRigVMAssetInterface into IRigVMEditorAssetInterface and IRigVMRuntimeAssetInterface.

diff --git a/Engine/Plugins/Runtime/RigVM/Source/RigVMEditor/Private/Editor/RigVMEditor.cpp b/Engine/Plugins/Runtime/RigVM/Source/RigVMEditor/Private/Editor/RigVMEditor.cpp
--- a/Engine/Plugins/Runtime/RigVM/Source/RigVMEditor/Private/Editor/RigVMEditor.cpp
+++ b/Engine/Plugins/Runtime/RigVM/Source/RigVMEditor/Private/Editor/RigVMEditor.cpp
@@ -8,7 +8,7 @@
 #include "EditorModeManager.h"
 #include "EulerTransform.h"
 #include "InstancedPropertyBagStructureDataProvider.h"
-#include "RigVMAsset.h"
+#include "RigVMEditorAsset.h"
 #include "RigVMBlueprintUtils.h"
 #include "RigVMEditorCommands.h"
 #include "RigVMPythonUtils.h"
@@ -45,6 +45,8 @@
 #include "Overrides/OverrideStatusDetailsWidgetBuilder.h"
 #include "Overrides/OverrideStatusDetailsWidgetBuilder.h"
 #include "GraphEditorActions.h"
+#include "GraphEditorSettings.h"
+#include "RigVMDeveloperTypeUtils.h"
 
 #if WITH_RIGVMLEGACYEDITOR
 #include "SKismetInspector.h"
@@ -71,9 +73,10 @@ FRigVMEditorBase::FRigVMEditorBase()
 
 void FRigVMEditorBase::UnbindEditor()
 {
-	FRigVMAssetInterfacePtr RigVMBlueprint = GetRigVMAssetInterface();
+	FRigVMEditorAssetInterfacePtr RigVMBlueprint = GetRigVMAssetInterface();
 	RigVMEditorClosedDelegate.Broadcast(this, RigVMBlueprint);
 
+	GraphExplorerWidget = nullptr;
 	ClearDetailObject();
 
 	if(PropertyChangedHandle.IsValid())
@@ -93,7 +96,7 @@ void FRigVMEditorBase::UnbindEditor()
 		// clear editor related data from the debugged control rig instance 
 		RigVMBlueprint->SetObjectBeingDebugged(nullptr);
 
-		IRigVMAssetInterface::sCurrentlyOpenedRigVMBlueprints.Remove(RigVMBlueprint);
+		IRigVMEditorAssetInterface::sCurrentlyOpenedRigVMBlueprints.Remove(RigVMBlueprint);
 
 		RigVMBlueprint->OnRefreshEditor().RemoveAll(this);
 		RigVMBlueprint->OnVariableDropped().RemoveAll(this);
@@ -172,23 +175,22 @@ FRigVMEditorBase* FRigVMEditorBase::GetFromAssetEditorInstance(IAssetEditorInsta
 	return Editor;
 }
 
-void FRigVMEditorBase::InitRigVMEditor(const EToolkitMode::Type Mode, const TSharedPtr< class IToolkitHost >& InitToolkitHost, FRigVMAssetInterfacePtr InRigVMBlueprint)
+void FRigVMEditorBase::InitRigVMEditor(const EToolkitMode::Type Mode, const TSharedPtr< class IToolkitHost >& InitToolkitHost, FRigVMEditorAssetInterfacePtr InRigVMBlueprint)
 {
 	DECLARE_SCOPE_HIERARCHICAL_COUNTER_FUNC()
 
 	check(InRigVMBlueprint);
 
 #if WITH_RIGVMLEGACYEDITOR
-	FBlueprintCompilationManager::FlushCompilationQueue(nullptr);
+	if (InRigVMBlueprint->IsBlueprintAsset())
+	{
+		FBlueprintCompilationManager::FlushCompilationQueue(nullptr);
+	}
 #else
 	// TODO: sara-s remove once blueprint backend is replaced
 	// FRigVMBlueprintCompilationManager::FlushCompilationQueue(nullptr);
 #endif
 
-	Toolbox = SNew(SBorder)
-		.BorderImage(FAppStyle::GetBrush("ToolPanel.GroupBorder"))
-		.Padding(0.f);
-
 	CreateEditorToolbar();
 
 	FEditorDelegates::EndPIE.AddRaw(this, &FRigVMEditorBase::OnPIEStopped);
@@ -196,7 +198,14 @@ void FRigVMEditorBase::InitRigVMEditor(const EToolkitMode::Type Mode, const TSha
 
 	// Build up a list of objects being edited in this asset editor
 	TArray<UObject*> ObjectsBeingEdited;
-	ObjectsBeingEdited.Add(InRigVMBlueprint->GetObject());
+	if (TScriptInterface<IRigVMRuntimeAssetInterface> RuntimeAsset = InRigVMBlueprint->GetRuntimeAssetInterface())
+	{
+		ObjectsBeingEdited.Add(RuntimeAsset.GetObject());
+	}
+	else
+	{
+		ObjectsBeingEdited.Add(InRigVMBlueprint.GetObject());
+	}
 
 	// Initialize the asset editor and spawn tabs
 	const bool bCreateDefaultStandaloneMenu = true;
@@ -207,7 +216,7 @@ void FRigVMEditorBase::InitRigVMEditor(const EToolkitMode::Type Mode, const TSha
 
 	CreateDefaultCommands();
 
-	TArray<FRigVMAssetInterfacePtr> Blueprints;
+	TArray<FRigVMEditorAssetInterfacePtr> Blueprints;
 	Blueprints.Add(InRigVMBlueprint);
 	InRigVMBlueprint->InitializeModelIfRequired();
 
@@ -244,7 +253,7 @@ void FRigVMEditorBase::InitRigVMEditor(const EToolkitMode::Type Mode, const TSha
 
 	}
 
-	IRigVMAssetInterface::sCurrentlyOpenedRigVMBlueprints.AddUnique(InRigVMBlueprint);
+	IRigVMEditorAssetInterface::sCurrentlyOpenedRigVMBlueprints.AddUnique(InRigVMBlueprint);
 
 	InRigVMBlueprint->OnModified().AddSP(StaticCastSharedRef<FRigVMEditorBase>(SharedRef()), &FRigVMEditorBase::HandleModifiedEvent);
 	InRigVMBlueprint->OnVMCompiled().AddSP(StaticCastSharedRef<FRigVMEditorBase>(SharedRef()), &FRigVMEditorBase::HandleVMCompiledEvent);
@@ -402,7 +411,7 @@ void FRigVMEditorBase::InitRigVMEditor(const EToolkitMode::Type Mode, const TSha
 		{
 			// no need to do anything if the the CR is not opened anymore
 		   // (i.e. destructor has been called before that task actually got a chance to start)
-		   if (!IRigVMAssetInterface::sCurrentlyOpenedRigVMBlueprints.Contains(InRigVMBlueprint))
+		   if (!IRigVMEditorAssetInterface::sCurrentlyOpenedRigVMBlueprints.Contains(InRigVMBlueprint))
 		   {
 			   return;		
 		   }
@@ -536,22 +545,27 @@ URigVMBlueprint* FRigVMEditorBase::GetRigVMBlueprint() const
 	return Cast<URigVMBlueprint>(GetRigVMAssetInterface().GetObject());
 }
 
-FRigVMAssetInterfacePtr FRigVMEditorBase::GetRigVMAssetInterface() const
+FRigVMEditorAssetInterfacePtr FRigVMEditorBase::GetRigVMAssetInterface() const
 {
 	const TArray<UObject*>& EditingObjs = GetEditingBlueprints();
 	for (UObject* Obj : EditingObjs)
 	{
-		if (Obj->Implements<URigVMAssetInterface>()) 
+		if (Obj->Implements<URigVMEditorAssetInterface>()) 
 		{
 			return Obj;
 		}
+		else if (Obj->Implements<URigVMRuntimeAssetInterface>())
+		{
+			TScriptInterface<IRigVMRuntimeAssetInterface> RuntimeAsset = Obj;
+			return RuntimeAsset->GetEditorOnlyData();
+		}
 	}
 	return nullptr;
 }
 
 TSubclassOf<UEdGraphSchema> FRigVMEditorBase::GetDefaultSchemaClass() const
 {
-	if(const FRigVMAssetInterfacePtr RigVMBlueprint = GetRigVMAssetInterface())
+	if(const FRigVMEditorAssetInterfacePtr RigVMBlueprint = GetRigVMAssetInterface())
 	{
 		return RigVMBlueprint->GetRigVMClientHost()->GetRigVMEdGraphSchemaClass();
 	}
@@ -575,7 +589,7 @@ void FRigVMEditorBase::OnGraphEditorFocused(const TSharedRef<SGraphEditor>& InGr
 void FRigVMEditorBase::Tick(float DeltaTime)
 {
 	// tick the  rigvm host
-	if (const FRigVMAssetInterfacePtr Blueprint = GetRigVMAssetInterface())
+	if (const FRigVMEditorAssetInterfacePtr Blueprint = GetRigVMAssetInterface())
 	{
 		if (LastDebuggedHost != GetCustomDebugObjectLabel(Blueprint->GetObjectBeingDebugged()))
 		{
@@ -634,7 +648,7 @@ FLinearColor FRigVMEditorBase::GetWorldCentricTabColorScale() const
 
 void FRigVMEditorBase::InitToolMenuContextImpl(FToolMenuContext& MenuContext)
 {
-	if (FRigVMAssetInterfacePtr RigVMBlueprint = GetRigVMAssetInterface())
+	if (FRigVMEditorAssetInterfacePtr RigVMBlueprint = GetRigVMAssetInterface())
 	{
 		URigVMGraph* Model = nullptr;
 		URigVMNode* Node = nullptr;
@@ -667,7 +681,7 @@ void FRigVMEditorBase::InitToolMenuContextImpl(FToolMenuContext& MenuContext)
 
 bool FRigVMEditorBase::TransactionObjectAffectsBlueprintImpl(UObject* InTransactedObject)
 {
-	FRigVMAssetInterfacePtr RigVMBlueprint = GetRigVMAssetInterface();
+	FRigVMEditorAssetInterfacePtr RigVMBlueprint = GetRigVMAssetInterface();
 	if (RigVMBlueprint == nullptr)
 	{
 		return false;
@@ -701,7 +715,10 @@ void FRigVMEditorBase::OnAddNewLocalVariable()
 	}
 
 	FRigVMGraphVariableDescription LastTypeVar;
-	LastTypeVar.ChangeType(GetLastPinTypeUsed());
+	FString CPPType;
+	UObject* CPPTypeObject = nullptr;
+	RigVMTypeUtils::CPPTypeFromPinType(GetLastPinTypeUsed(), CPPType, &CPPTypeObject);
+	LastTypeVar.ChangeType(CPPType, CPPTypeObject);
 	FRigVMGraphVariableDescription NewVar = GetFocusedController()->AddLocalVariable(TEXT("NewLocalVar"), LastTypeVar.CPPType, LastTypeVar.CPPTypeObject, LastTypeVar.DefaultValue, true, true);
 	if(NewVar.Name.IsNone())
 	{
@@ -721,7 +738,10 @@ void FRigVMEditorBase::OnPasteNewLocalVariable(const FBPVariableDescription& Var
 	}
 
 	FRigVMGraphVariableDescription TypeVar;
-	TypeVar.ChangeType(VariableDescription.VarType);
+	FString CPPType;
+	UObject* CPPTypeObject = nullptr;
+	RigVMTypeUtils::CPPTypeFromPinType(VariableDescription.VarType, CPPType, &CPPTypeObject);
+	TypeVar.ChangeType(CPPType, CPPTypeObject);
 	FRigVMGraphVariableDescription NewVar = GetFocusedController()->AddLocalVariable(VariableDescription.VarName, TypeVar.CPPType, TypeVar.CPPTypeObject, VariableDescription.DefaultValue, true, true);
 	if(NewVar.Name.IsNone())
 	{
@@ -733,7 +753,7 @@ void FRigVMEditorBase::DeleteSelectedNodes()
 {
 	DECLARE_SCOPE_HIERARCHICAL_COUNTER_FUNC()
 
-	FRigVMAssetInterfacePtr RigVMBlueprint = GetRigVMAssetInterface();
+	FRigVMEditorAssetInterfacePtr RigVMBlueprint = GetRigVMAssetInterface();
 	if (RigVMBlueprint == nullptr)
 	{
 		return;
@@ -881,7 +901,7 @@ void FRigVMEditorBase::AddNewFunctionVariant(const UEdGraph* InOriginalFunction)
 		{
 			if (URigVMFunctionLibrary* FunctionLibrary = Cast<URigVMFunctionLibrary>(RigVMGraph->GetParentGraph()))
 			{
-				FRigVMAssetInterfacePtr RigVMBlueprint = GetRigVMAssetInterface();
+				FRigVMEditorAssetInterfacePtr RigVMBlueprint = GetRigVMAssetInterface();
 				URigVMController* Controller = RigVMBlueprint->GetController(FunctionLibrary);
 				if (const URigVMLibraryNode* VariantNode = Controller->CreateFunctionVariant(RigVMGraph->GetOuter()->GetFName(), NAME_None, true, true))
 				{
@@ -914,7 +934,7 @@ void FRigVMEditorBase::PostTransaction(bool bSuccess, const FTransaction* Transa
 		return;
 	}
 
-	const FRigVMAssetInterfacePtr Blueprint = GetRigVMAssetInterface();
+	const FRigVMEditorAssetInterfacePtr Blueprint = GetRigVMAssetInterface();
 	if(Blueprint == nullptr)
 	{
 			return;
@@ -976,20 +996,6 @@ bool FRigVMEditorBase::CanStopWatchingPin() const
 	return false;
 }
 
-void FRigVMEditorBase::OnToolkitHostingStarted(const TSharedRef<class IToolkit>& Toolkit)
-{
-	TSharedPtr<SWidget> InlineContent = Toolkit->GetInlineContent();
-	if (InlineContent.IsValid())
-	{
-		Toolbox->SetContent(InlineContent.ToSharedRef());
-	}
-}
-
-void FRigVMEditorBase::OnToolkitHostingFinished(const TSharedRef<class IToolkit>& Toolkit)
-{
-	Toolbox->SetContent(SNullWidget::NullWidget);
-}
-
 TStatId FRigVMEditorBase::GetStatId() const
 {
 	RETURN_QUICK_DECLARE_CYCLE_STAT(FRigVMEditorBase, STATGROUP_Tickables);
@@ -1033,7 +1039,7 @@ void FRigVMEditorBase::PasteNodes()
 
 URigVMHost* FRigVMEditorBase::GetRigVMHost() const
 {
-	if (FRigVMAssetInterfacePtr RigVMBlueprint = GetRigVMAssetInterface())
+	if (FRigVMEditorAssetInterfacePtr RigVMBlueprint = GetRigVMAssetInterface())
 	{
 		if(RigVMBlueprint->GetEditorHost() && IsValid(RigVMBlueprint->GetEditorHost()))
 		{
@@ -1045,7 +1051,16 @@ URigVMHost* FRigVMEditorBase::GetRigVMHost() const
 
 UObject* FRigVMEditorBase::GetOuterForHost() const
 {
-	return GetRigVMAssetInterface()->GetObject();
+	if (FRigVMEditorAssetInterfacePtr EditorAsset = GetRigVMAssetInterface())
+	{
+		if (EditorAsset->IsBlueprintAsset())
+		{
+			return EditorAsset.GetObject();
+		}
+		
+		return EditorAsset->GetRuntimeAssetInterface().GetObject();
+	}
+	return nullptr;
 }
 
 UClass* FRigVMEditorBase::GetDetailWrapperClass() const
@@ -1078,6 +1093,27 @@ bool FRigVMEditorBase::SelectLocalVariable(const UEdGraph* Graph, const FName& V
 	return false;
 }
 
+bool FRigVMEditorBase::SelectVariable(const FName& VariableName)
+{
+	if (FRigVMEditorAssetInterfacePtr Asset = GetRigVMAssetInterface())
+	{
+		const FProperty* Prop = Asset->FindGeneratedPropertyByName(VariableName);
+		UPropertyWrapper* PropWrap = (Prop ? const_cast<FProperty*>(Prop)->GetUPropertyWrapper() : nullptr);
+		if (TSharedPtr<SRigVMDetailsInspector> RigVMInspector = GetRigVMInspector())
+		{
+			SRigVMDetailsInspector::FShowDetailsOptions Options(FText::FromName(VariableName));
+			Options.bForceRefresh = true;
+			
+			if (RigVMInspector.IsValid())
+			{
+				RigVMInspector->ShowDetailsForSingleObject(PropWrap, Options);
+				return true;
+			}
+		}
+	}
+	return false;
+}
+
 void FRigVMEditorBase::CreateDefaultCommandsImpl()
 {
 	if (GetRigVMAssetInterface())
@@ -1110,7 +1146,7 @@ void FRigVMEditorBase::Compile()
 	{
 		DECLARE_SCOPE_HIERARCHICAL_COUNTER_FUNC()
 
-		FRigVMAssetInterfacePtr RigVMBlueprint = GetRigVMAssetInterface();
+		FRigVMEditorAssetInterfacePtr RigVMBlueprint = GetRigVMAssetInterface();
 		if (RigVMBlueprint == nullptr)
 		{
 			return;
@@ -1210,7 +1246,14 @@ void FRigVMEditorBase::Compile()
 			{
 				if (URigVMHost* DebuggedHost = Cast<URigVMHost>(DebugObject.Object))
 				{
-					DebuggedHost->CopyExternalVariableDefaultValuesFromCDO();
+					if (TScriptInterface<IRigVMRuntimeAssetInterface> RuntimeAsset = DebuggedHost->GetGeneratedByAsset())
+					{
+						RuntimeAsset->InitializeInstance(DebuggedHost);
+					}
+					else
+					{
+						DebuggedHost->CopyExternalVariableDefaultValuesFromCDO();
+					}
 				}
 			}
 		}
@@ -1247,7 +1290,7 @@ bool FRigVMEditorBase::IsEditable(UEdGraph* InGraph) const
 		return false;
 	}
 	
-	if(FRigVMAssetInterfacePtr RigVMBlueprint = GetRigVMAssetInterface())
+	if(FRigVMEditorAssetInterfacePtr RigVMBlueprint = GetRigVMAssetInterface())
 	{
 		// aggregate graphs are always read only
 		if(const URigVMGraph* Model = RigVMBlueprint->GetModel(InGraph))
@@ -1298,7 +1341,7 @@ void FRigVMEditorBase::OnSelectedNodesChangedImpl(const TSet<UObject*>& NewSelec
 
 	TGuardValue<bool> SelectGuard(RigVMEdGraph->bIsSelecting, true);
 
-	FRigVMAssetInterfacePtr RigVMBlueprint = GetRigVMAssetInterface();
+	FRigVMEditorAssetInterfacePtr RigVMBlueprint = GetRigVMAssetInterface();
 	if (RigVMBlueprint)
 	{
 		TArray<FName> NodeNamesToSelect;
@@ -1317,7 +1360,7 @@ void FRigVMEditorBase::OnSelectedNodesChangedImpl(const TSet<UObject*>& NewSelec
 	}
 }
 
-void FRigVMEditorBase::OnBlueprintChangedImpl(IRigVMAssetInterface* InBlueprint, bool bIsJustBeingCompiled)
+void FRigVMEditorBase::OnBlueprintChangedImpl(IRigVMEditorAssetInterface* InBlueprint, bool bIsJustBeingCompiled)
 {
 	DECLARE_SCOPE_HIERARCHICAL_COUNTER_FUNC()
 
@@ -1476,7 +1519,7 @@ void FRigVMEditorBase::BindCommands()
 		EditorCommands.TogglePreviewHere,
 		FExecuteAction::CreateLambda([ThisRef]()
 		{
-			if (const FRigVMAssetInterfacePtr RigVMAsset = ThisRef->GetRigVMAssetInterface())
+			if (const FRigVMEditorAssetInterfacePtr RigVMAsset = ThisRef->GetRigVMAssetInterface())
 			{
 				if (const FRigVMClient* Client = RigVMAsset->GetRigVMClient())
 				{
@@ -1493,14 +1536,14 @@ void FRigVMEditorBase::BindCommands()
 		EditorCommands.PreviewHereStepForward,
 		FExecuteAction::CreateLambda([ThisRef]()
 		{
-			if (const FRigVMAssetInterfacePtr RigVMAsset = ThisRef->GetRigVMAssetInterface())
+			if (const FRigVMEditorAssetInterfacePtr RigVMAsset = ThisRef->GetRigVMAssetInterface())
 			{
 				RigVMAsset->PreviewHereStepForward();
 			}
 		}),
 		FCanExecuteAction::CreateLambda([ThisRef]() -> bool
 		{
-			if (const FRigVMAssetInterfacePtr RigVMAsset = ThisRef->GetRigVMAssetInterface())
+			if (const FRigVMEditorAssetInterfacePtr RigVMAsset = ThisRef->GetRigVMAssetInterface())
 			{
 				return RigVMAsset->CanPreviewHereStepForward();
 			}
@@ -1529,7 +1572,7 @@ void FRigVMEditorBase::UnbindCommands()
 
 void FRigVMEditorBase::ToggleAutoCompileGraph()
 {
-	if (FRigVMAssetInterfacePtr RigVMBlueprint = GetRigVMAssetInterface())
+	if (FRigVMEditorAssetInterfacePtr RigVMBlueprint = GetRigVMAssetInterface())
 	{
 		RigVMBlueprint->SetAutoVMRecompile(!RigVMBlueprint->GetAutoVMRecompile());
 		if (RigVMBlueprint->GetAutoVMRecompile())
@@ -1541,7 +1584,7 @@ void FRigVMEditorBase::ToggleAutoCompileGraph()
 
 bool FRigVMEditorBase::IsAutoCompileGraphOn() const
 {
-	if (FRigVMAssetInterfacePtr RigVMBlueprint = GetRigVMAssetInterface())
+	if (FRigVMEditorAssetInterfacePtr RigVMBlueprint = GetRigVMAssetInterface())
 	{
 		return RigVMBlueprint->GetAutoVMRecompile();
 	}
@@ -1644,7 +1687,7 @@ void FRigVMEditorBase::OnNewDocumentClicked(ECreatedDocumentType GraphType)
 {
 	if (GraphType == FRigVMNewEditor::CGT_NewFunctionGraph)
 	{
-		if (FRigVMAssetInterfacePtr RigVMBlueprint = GetRigVMAssetInterface())
+		if (FRigVMEditorAssetInterfacePtr RigVMBlueprint = GetRigVMAssetInterface())
 		{
 			if (URigVMController* Controller = RigVMBlueprint->GetOrCreateController(RigVMBlueprint->GetLocalFunctionLibrary()))
 			{
@@ -1653,6 +1696,7 @@ void FRigVMEditorBase::OnNewDocumentClicked(ECreatedDocumentType GraphType)
 					if (const UEdGraph* NewGraph = RigVMBlueprint->GetEdGraph(FunctionNode->GetContainedGraph()))
 					{
 						OpenDocument(NewGraph, FDocumentTracker::OpenNewDocument);
+						GraphExplorerWidget->Refresh();
 						RenameNewlyAddedAction(FunctionNode->GetFName());
 					}
 				}
@@ -1661,7 +1705,7 @@ void FRigVMEditorBase::OnNewDocumentClicked(ECreatedDocumentType GraphType)
 	}
 	else if(GraphType == FRigVMNewEditor::CGT_NewEventGraph)
 	{
-		if (FRigVMAssetInterfacePtr RigVMBlueprint = GetRigVMAssetInterface())
+		if (FRigVMEditorAssetInterfacePtr RigVMBlueprint = GetRigVMAssetInterface())
 		{
 			const UClass* EdGraphSchemaClass = RigVMBlueprint->GetRigVMClientHost()->GetRigVMEdGraphSchemaClass();
 			const URigVMEdGraphSchema* SchemaCDO = CastChecked<URigVMEdGraphSchema>(EdGraphSchemaClass->GetDefaultObject());
@@ -1671,6 +1715,7 @@ void FRigVMEditorBase::OnNewDocumentClicked(ECreatedDocumentType GraphType)
 				if (const UEdGraph* NewGraph = RigVMBlueprint->GetEdGraph(Model))
 				{
 					OpenDocument(NewGraph, FDocumentTracker::OpenNewDocument);
+					GraphExplorerWidget->Refresh();
 					RenameNewlyAddedAction(NewGraph->GetFName());
 				}
 			}
@@ -1728,7 +1773,7 @@ FGraphAppearanceInfo FRigVMEditorBase::GetGraphAppearance(UEdGraph* InGraph) con
 {
 	FGraphAppearanceInfo AppearanceInfo = GetGraphAppearanceImpl(InGraph);
 
-	if (const FRigVMAssetInterfacePtr RigVMBlueprint = GetRigVMAssetInterface())
+	if (const FRigVMEditorAssetInterfacePtr RigVMBlueprint = GetRigVMAssetInterface())
 	{
 		AppearanceInfo.CornerText = LOCTEXT("AppearanceCornerText_RigVMEditor", "RigVM");
 
@@ -1760,7 +1805,7 @@ void FRigVMEditorBase::HandleModifiedEvent(ERigVMGraphNotifType InNotifType, URi
 {
 	DECLARE_SCOPE_HIERARCHICAL_COUNTER_FUNC()
 
-	FRigVMAssetInterfacePtr RigVMBlueprint = GetRigVMAssetInterface();
+	FRigVMEditorAssetInterfacePtr RigVMBlueprint = GetRigVMAssetInterface();
 	if (RigVMBlueprint == nullptr)
 	{
 		return;
@@ -1922,11 +1967,14 @@ public:
 
 void FRigVMEditorBase::HandleVMCompiledEvent(UObject* InCompiledObject, URigVM* InVM, FRigVMExtendedExecuteContext& InContext)
 {
-	if (InCompiledObject->Implements<URigVMAssetInterface>())
+	if (InCompiledObject->Implements<URigVMEditorAssetInterface>())
 	{
-		FRigVMAssetInterfacePtr RigVMBlueprint = InCompiledObject;
-		GetCompilerResultsListing()->ClearMessages();
-		GetCompilerResultsListing()->AddMessages(RigVMBlueprint->GetCompileLog().Messages);
+		FRigVMEditorAssetInterfacePtr RigVMBlueprint = InCompiledObject;
+		if (TSharedPtr<class IMessageLogListing> CompilerListing = GetCompilerResultsListing())
+		{
+			CompilerListing->ClearMessages();
+			CompilerListing->AddMessages(RigVMBlueprint->GetCompileLog().Messages);
+		}
 		RigVMBlueprint->GetCompileLog().Messages.Reset();
 		RigVMBlueprint->GetCompileLog().NumErrors = RigVMBlueprint->GetCompileLog().NumWarnings = 0;
 	}
@@ -1967,7 +2015,7 @@ void FRigVMEditorBase::HandleVMCompiledEvent(UObject* InCompiledObject, URigVM*
 
 void FRigVMEditorBase::HandleVMExecutedEvent(URigVMHost* InHost, const FName& InEventName)
 {
-	if (FRigVMAssetInterfacePtr RigVMBlueprint = GetRigVMAssetInterface())
+	if (FRigVMEditorAssetInterfacePtr RigVMBlueprint = GetRigVMAssetInterface())
 	{
 		URigVMHost* DebuggedHost = Cast<URigVMHost>(RigVMBlueprint->GetObjectBeingDebugged());
 		if (DebuggedHost == nullptr)
@@ -2082,7 +2130,7 @@ void FRigVMEditorBase::NotifyPreChange(FProperty* PropertyAboutToChange)
 {
 	NotifyPreChangeImpl(PropertyAboutToChange);
 
-	if (FRigVMAssetInterfacePtr RigVMBlueprint = GetRigVMAssetInterface())
+	if (FRigVMEditorAssetInterfacePtr RigVMBlueprint = GetRigVMAssetInterface())
 	{
 		RigVMBlueprint->GetObject()->Modify();
 	}
@@ -2096,7 +2144,7 @@ void FRigVMEditorBase::NotifyPostChange(const FPropertyChangedEvent& PropertyCha
 	// OnFinishedChangingProperties is called only for top level property changes.
 	// changes on a lower level property like transform under a user defined struct
 	// only go through this.
-	FRigVMAssetInterfacePtr RigVMBlueprint = GetRigVMAssetInterface();
+	FRigVMEditorAssetInterfacePtr RigVMBlueprint = GetRigVMAssetInterface();
 	if(GetRigVMHost() && RigVMBlueprint)
 	{
 		bool bUseDebuggedObject = true; 
@@ -2121,7 +2169,7 @@ void FRigVMEditorBase::OnFinishedChangingProperties(const FPropertyChangedEvent&
 {
 	DECLARE_SCOPE_HIERARCHICAL_COUNTER_FUNC()
 
-	FRigVMAssetInterfacePtr RigVMBlueprint = GetRigVMAssetInterface();
+	FRigVMEditorAssetInterfacePtr RigVMBlueprint = GetRigVMAssetInterface();
 
 	if (RigVMBlueprint)
 	{
@@ -2140,7 +2188,7 @@ void FRigVMEditorBase::OnFinishedChangingProperties(const FPropertyChangedEvent&
 
 void FRigVMEditorBase::OnPropertyChanged(UObject* InObject, FPropertyChangedEvent& InEvent)
 {
-	FRigVMAssetInterfacePtr RigVMBlueprint = GetRigVMAssetInterface();
+	FRigVMEditorAssetInterfacePtr RigVMBlueprint = GetRigVMAssetInterface();
 
 	if (RigVMBlueprint && InObject == RigVMBlueprint->GetObject())
 	{
@@ -2161,7 +2209,7 @@ void FRigVMEditorBase::OnWrappedPropertyChangedChainEvent(URigVMDetailsViewWrapp
 
 	TGuardValue<bool> SuspendDetailsPanelRefresh(bSuspendDetailsPanelRefresh, true);
 
-	FRigVMAssetInterfacePtr RigVMBlueprint = GetRigVMAssetInterface();
+	FRigVMEditorAssetInterfacePtr RigVMBlueprint = GetRigVMAssetInterface();
 
 	FString PropertyPath = InPropertyPath;
 	if(UScriptStruct* WrappedStruct = InWrapperObject->GetWrappedStruct())
@@ -2358,7 +2406,7 @@ void FRigVMEditorBase::OnRequestLocalizeFunctionDialog(FRigVMGraphFunctionIdenti
 	UE::RigVM::Editor::Tools::OnRequestLocalizeFunctionDialog(InFunction, InTargetController, InTargetFunctionHost, bForce);
 }
 
-FRigVMController_BulkEditResult FRigVMEditorBase::OnRequestBulkEditDialog(FRigVMAssetInterfacePtr InBlueprint,
+FRigVMController_BulkEditResult FRigVMEditorBase::OnRequestBulkEditDialog(FRigVMEditorAssetInterfacePtr InBlueprint,
 	URigVMController* InController, URigVMLibraryNode* InFunction, ERigVMControllerBulkEditType InEditType)
 {
 	if (bAllowBulkEdits)
@@ -2527,7 +2575,7 @@ TRigVMTypeIndex FRigVMEditorBase::OnRequestPinTypeSelectionDialog(const TArray<T
 
 void FRigVMEditorBase::HandleJumpToHyperlink(const UObject* InSubject)
 {
-	FRigVMAssetInterfacePtr RigBlueprint = GetRigVMAssetInterface();
+	FRigVMEditorAssetInterfacePtr RigBlueprint = GetRigVMAssetInterface();
 	if(RigBlueprint == nullptr)
 	{
 		return;
@@ -2568,13 +2616,18 @@ void FRigVMEditorBase::HandleJumpToHyperlink(const UObject* InSubject)
 
 	if (GraphToJumpTo && NodeToJumpTo)
 	{
-		if(IRigVMAssetInterface* OtherAsset = NodeToJumpTo->GetImplementingOuter<IRigVMAssetInterface>())
+		if(IRigVMEditorAssetInterface* OtherAsset = NodeToJumpTo->GetImplementingOuter<IRigVMEditorAssetInterface>())
 		{
 			if(OtherAsset != RigBlueprint.GetInterface())
 			{
-				if (GEditor->GetEditorSubsystem<UAssetEditorSubsystem>()->OpenEditorForAsset(OtherAsset->GetObject()))
+				UObject* OtherAssetObject = OtherAsset->GetObject();
+				if (TScriptInterface<IRigVMRuntimeAssetInterface> OtherRuntimeAsset = OtherAsset->GetRuntimeAssetInterface())
+				{
+					OtherAssetObject = OtherRuntimeAsset.GetObject();
+				}
+				if (GEditor->GetEditorSubsystem<UAssetEditorSubsystem>()->OpenEditorForAsset(OtherAssetObject))
 				{
-					IAssetEditorInstance* OtherEditor = GEditor->GetEditorSubsystem<UAssetEditorSubsystem>()->FindEditorForAsset(OtherAsset->GetObject(), /*bFocusIfOpen =*/true);
+					IAssetEditorInstance* OtherEditor = GEditor->GetEditorSubsystem<UAssetEditorSubsystem>()->FindEditorForAsset(OtherAssetObject, /*bFocusIfOpen =*/true);
 					if (FRigVMEditorBase* OtherRigVMEditor = FRigVMEditorBase::GetFromAssetEditorInstance(OtherEditor))
 					{
 						OtherRigVMEditor->HandleJumpToHyperlink(NodeToJumpTo);
@@ -2610,7 +2663,7 @@ void FRigVMEditorBase::HandleJumpToHyperlink(const UObject* InSubject)
 bool FRigVMEditorBase::UpdateDefaultValueForVariable(FRigVMGraphVariableDescription& InVariable, bool bUseDebuggedObject)
 {
 	bool bAnyValueChanged = false;
-	if (FRigVMAssetInterfacePtr RigVMBlueprint = GetRigVMAssetInterface())
+	if (FRigVMEditorAssetInterfacePtr RigVMBlueprint = GetRigVMAssetInterface())
 	{
 		FString NewDefaultValue = RigVMBlueprint->GetVariableDefaultValue(InVariable.Name, bUseDebuggedObject);
 		if (InVariable.DefaultValue != NewDefaultValue)
@@ -2627,8 +2680,8 @@ void FRigVMEditorBase::UpdateRigVMHost()
 {
 	DECLARE_SCOPE_HIERARCHICAL_COUNTER_FUNC()
 
-	FRigVMAssetInterfacePtr Blueprint = GetRigVMAssetInterface();
-	if(UClass* Class = Blueprint->GetRigVMGeneratedClass())
+	FRigVMEditorAssetInterfacePtr Blueprint = GetRigVMAssetInterface();
+
 	{
 		if (URigVMHost* CurrentHost = GetRigVMHost())
 		{
@@ -2642,7 +2695,7 @@ void FRigVMEditorBase::UpdateRigVMHost()
 			// if this control rig is from a temporary step,
 			// for example the reinstancing class, clear it 
 			// and create a new one!
-			if (CurrentHost->GetClass() != Class)
+			if (CurrentHost->GetClass() != Blueprint->GetRigVMGeneratedClass())
 			{
 				SetHost(nullptr);
 			}
@@ -2683,7 +2736,7 @@ void FRigVMEditorBase::UpdateRigVMHost()
 
 void FRigVMEditorBase::CacheNameLists()
 {
-	if (FRigVMAssetInterfacePtr RigVMBlueprint = GetRigVMAssetInterface())
+	if (FRigVMEditorAssetInterfacePtr RigVMBlueprint = GetRigVMAssetInterface())
 	{
 		TArray<UEdGraph*> EdGraphs;
 		RigVMBlueprint->GetAllEdGraphs(EdGraphs);
@@ -2709,17 +2762,27 @@ void FRigVMEditorBase::OnCreateComment()
 		{
 			if (URigVMEdGraph* RigVMEdGraph = Cast<URigVMEdGraph>(Graph))
 			{
-				if (FRigVMAssetInterfacePtr Blueprint = GetRigVMAssetInterface())
+				if (FRigVMEditorAssetInterfacePtr Blueprint = GetRigVMAssetInterface())
 				{
 					if (URigVMController* Controller = Blueprint->GetController(RigVMEdGraph))
 					{
 						Controller->OpenUndoBracket(TEXT("Create Comment"));
-						FEdGraphSchemaAction_K2AddComment CommentAction;
-						UEdGraphNode* EdNode = CommentAction.PerformAction(Graph, NULL, GraphEditor->GetPasteLocation2f());
-						if (UEdGraphNode_Comment* CommentNode = CastChecked<UEdGraphNode_Comment>(EdNode))
+						if (Blueprint->IsBlueprintAsset())
 						{
-							Controller->SetNodeColorByName(CommentNode->GetFName(), CommentNode->CommentColor, false);
-							Controller->SetNodePositionByName(CommentNode->GetFName(), FVector2D(CommentNode->NodePosX, CommentNode->NodePosY), false);
+							FEdGraphSchemaAction_K2AddComment CommentAction;
+							UEdGraphNode* EdNode = CommentAction.PerformAction(Graph, NULL, GraphEditor->GetPasteLocation2f());
+							if (UEdGraphNode_Comment* CommentNode = CastChecked<UEdGraphNode_Comment>(EdNode))
+							{
+								Controller->SetNodeColorByName(CommentNode->GetFName(), CommentNode->CommentColor, false);
+								Controller->SetNodePositionByName(CommentNode->GetFName(), FVector2D(CommentNode->NodePosX, CommentNode->NodePosY), false);
+							}
+						}
+						else
+						{
+							FSlateRect Bounds;
+							GraphEditor->GetBoundsForSelectedNodes(Bounds, 50.f);
+							FVector2D Size = Bounds.GetSize();
+							Controller->AddCommentNode(TEXT("Comment"), Bounds.GetTopLeft(), Size, GetDefault<UGraphEditorSettings>()->DefaultCommentNodeTitleColor, FString(), true, true);
 						}
 						Controller->CloseUndoBracket();
 					}
@@ -3333,19 +3396,45 @@ void FRigVMEditorBase::SetDetailViewForLocalVariable()
 	SelectLocalVariable(GetFocusedGraph(), VariableName);
 }
 
+void FRigVMEditorBase::SetDetailViewForVariable()
+{
+	FName VariableName;
+	if (TSharedPtr<SRigVMDetailsInspector> RigVMInspector = GetRigVMInspector())
+	{
+		TArray< TWeakObjectPtr<UObject> > SelectedObjects = RigVMInspector->GetSelectedObjects();
+		for (TWeakObjectPtr<UObject> SelectedObject : SelectedObjects)
+		{
+			if (SelectedObject.IsValid())
+			{
+				if(UPropertyWrapper* WrapperObject = Cast<UPropertyWrapper>(SelectedObject.Get()))
+				{
+					VariableName = WrapperObject->GetProperty()->GetFName();
+					break;
+				}
+			}
+		}
+	}
+
+	SelectVariable(VariableName);
+}
+
 void FRigVMEditorBase::RefreshDetailView()
 {
 	if(bSuspendDetailsPanelRefresh)
 	{
 		return;
 	}
-	if(DetailViewShowsAnyRigUnit())
+	if(DetailViewShowsLocalVariable())
 	{
-		SetDetailViewForFocusedGraph();
+		SetDetailViewForLocalVariable();
 	}
-	else if(DetailViewShowsLocalVariable())
+	else if(DetailViewShowsVariable())
 	{
-		SetDetailViewForLocalVariable();	
+		SetDetailViewForVariable();
+	}
+	else if(DetailViewShowsAnyRigUnit())
+	{
+		SetDetailViewForFocusedGraph();
 	}
 	else
 	{
@@ -3388,6 +3477,25 @@ bool FRigVMEditorBase::DetailViewShowsLocalVariable() const
 	return DetailViewShowsStruct(FRigVMGraphVariableDescription::StaticStruct());
 }
 
+bool FRigVMEditorBase::DetailViewShowsVariable() const
+{
+	if (TSharedPtr<SRigVMDetailsInspector> RigVMInspector = GetRigVMInspector())
+	{
+		TArray< TWeakObjectPtr<UObject> > SelectedObjects = RigVMInspector->GetSelectedObjects();
+		for (TWeakObjectPtr<UObject> SelectedObject : SelectedObjects)
+		{
+			if (SelectedObject.IsValid())
+			{
+				if(UPropertyWrapper* WrapperObject = Cast<UPropertyWrapper>(SelectedObject.Get()))
+				{
+					return true;
+				}
+			}
+		}
+	}
+	return false; 
+}
+
 bool FRigVMEditorBase::DetailViewShowsStruct(UScriptStruct* InStruct) const
 {
 	TArray< TWeakObjectPtr<UObject> > SelectedObjects = GetSelectedObjects();
@@ -3475,7 +3583,7 @@ void FRigVMEditorBase::ClearDetailsViewWrapperObjects()
 
 void FRigVMEditorBase::SetHost(URigVMHost* InHost)
 {
-	if (FRigVMAssetInterfacePtr RigVMBlueprint = GetRigVMAssetInterface())
+	if (FRigVMEditorAssetInterfacePtr RigVMBlueprint = GetRigVMAssetInterface())
 	{
 		if (IsValid(RigVMBlueprint->GetEditorHost()) && RigVMBlueprint->GetEditorHost()->GetOuter() == GetOuterForHost())
 		{
@@ -3492,7 +3600,7 @@ void FRigVMEditorBase::SetHost(URigVMHost* InHost)
 
 URigVMGraph* FRigVMEditorBase::GetFocusedModel() const
 {
-	const FRigVMAssetInterfacePtr Blueprint = GetRigVMAssetInterface();
+	const FRigVMEditorAssetInterfacePtr Blueprint = GetRigVMAssetInterface();
 	if (Blueprint == nullptr)
 	{
 		return nullptr;
@@ -3504,7 +3612,7 @@ URigVMGraph* FRigVMEditorBase::GetFocusedModel() const
 
 URigVMController* FRigVMEditorBase::GetFocusedController() const
 {
-	FRigVMAssetInterfacePtr Blueprint = GetRigVMAssetInterface();
+	FRigVMEditorAssetInterfacePtr Blueprint = GetRigVMAssetInterface();
 	if (Blueprint == nullptr)
 	{
 		return nullptr;
@@ -3631,7 +3739,7 @@ void FRigVMEditorBase::FillToolbar(FToolBarBuilder& ToolbarBuilder, bool bEndSec
 
 void FRigVMEditorBase::UpdateStaleWatchedPins()
 {
-	FRigVMAssetInterfacePtr RigVMBlueprint = GetRigVMAssetInterface();
+	FRigVMEditorAssetInterfacePtr RigVMBlueprint = GetRigVMAssetInterface();
 	if (RigVMBlueprint == nullptr)
 	{
 		return;
@@ -3698,14 +3806,14 @@ void FRigVMEditorBase::UpdateStaleWatchedPins()
 	}
 }
 
-void FRigVMEditorBase::HandleRefreshEditorFromBlueprint(FRigVMAssetInterfacePtr InBlueprint)
+void FRigVMEditorBase::HandleRefreshEditorFromBlueprint(FRigVMEditorAssetInterfacePtr InBlueprint)
 {
 	Compile();
 }
 
 void FRigVMEditorBase::HandleVariableDroppedFromBlueprint(UObject* InSubject, FProperty* InVariableToDrop, const FVector2D& InDropPosition, const FVector2D& InScreenPosition)
 {
-	FRigVMAssetInterfacePtr Blueprint = GetRigVMAssetInterface();
+	FRigVMEditorAssetInterfacePtr Blueprint = GetRigVMAssetInterface();
 	if (Blueprint == nullptr)
 	{
 		return;
@@ -3796,7 +3904,7 @@ void FRigVMEditorBase::OnGraphNodeClicked(URigVMEdGraphNode* InNode, const FGeom
 	}
 }
 
-void FRigVMEditorBase::OnNodeDoubleClicked(FRigVMAssetInterfacePtr InBlueprint, URigVMNode* InNode)
+void FRigVMEditorBase::OnNodeDoubleClicked(FRigVMEditorAssetInterfacePtr InBlueprint, URigVMNode* InNode)
 {
 	ensure(GetRigVMAssetInterface() == InBlueprint);
 
@@ -3904,7 +4012,7 @@ void FRigVMEditorBase::SwapAssetReferences()
 
 void FRigVMEditorBase::ToggleProfiling()
 {
-	TScriptInterface<IRigVMAssetInterface> Asset = GetRigVMAssetInterface();
+	TScriptInterface<IRigVMEditorAssetInterface> Asset = GetRigVMAssetInterface();
 	if (Asset && Asset->GetObject())
 	{
 		const FScopedTransaction Transaction(LOCTEXT("ToggleProfiler", "Toggle Profiler"));
@@ -3915,7 +4023,7 @@ void FRigVMEditorBase::ToggleProfiling()
 
 void FRigVMEditorBase::OnOpenSelectedNodesInNewTab()
 {
-	const FRigVMAssetInterfacePtr RigVMAsset = GetRigVMAssetInterface();
+	const FRigVMEditorAssetInterfacePtr RigVMAsset = GetRigVMAssetInterface();
 	if (RigVMAsset == nullptr)
 	{
 		return;
@@ -3982,7 +4090,7 @@ void FRigVMEditorBase::OnGraphNodeDropToPerform(TSharedPtr<FDragDropOperation> I
 			const URigVMEdGraphSchema* Schema = CastChecked<URigVMEdGraphSchema>(GetRigVMAssetInterface()->GetRigVMClientHost()->GetRigVMEdGraphSchemaClass()->GetDefaultObject());
 			if (ExplorerOp->GetElement().Type == ERigVMExplorerElementType::Function)
 			{
-				if (FRigVMAssetInterfacePtr TargetRigBlueprint = FRigVMBlueprintUtils::FindAssetForGraph(InGraph))
+				if (FRigVMEditorAssetInterfacePtr TargetRigBlueprint = FRigVMBlueprintUtils::FindAssetForGraph(InGraph))
 				{
 					if (URigVMFunctionLibrary* Library = ExplorerOp->GetRigVMAssetInterface()->GetLocalFunctionLibrary())
 					{
@@ -3995,15 +4103,15 @@ void FRigVMEditorBase::OnGraphNodeDropToPerform(TSharedPtr<FDragDropOperation> I
 			}
 			else if (ExplorerOp->GetElement().Type == ERigVMExplorerElementType::Variable)
 			{
-				if (FRigVMAssetInterfacePtr TargetRigBlueprint = FRigVMBlueprintUtils::FindAssetForGraph(InGraph))
+				if (FRigVMEditorAssetInterfacePtr TargetRigBlueprint = FRigVMBlueprintUtils::FindAssetForGraph(InGraph))
 				{
-					FProperty* Property = TargetRigBlueprint->FindGeneratedPropertyByName(*ExplorerOp->GetElement().Name);
-					Schema->RequestVariableDropOnPanel(InGraph, Property, InNodePosition, InScreenPosition);
+					const FProperty* Property = TargetRigBlueprint->FindGeneratedPropertyByName(*ExplorerOp->GetElement().Name);
+					Schema->RequestVariableDropOnPanel(InGraph, const_cast<FProperty*>(Property), InNodePosition, InScreenPosition);
 				}
 			}
 			else if (ExplorerOp->GetElement().Type == ERigVMExplorerElementType::LocalVariable)
 			{
-				if (FRigVMAssetInterfacePtr TargetRigBlueprint = FRigVMBlueprintUtils::FindAssetForGraph(InGraph))
+				if (FRigVMEditorAssetInterfacePtr TargetRigBlueprint = FRigVMBlueprintUtils::FindAssetForGraph(InGraph))
 				{
 					if (URigVMGraph* Graph = TargetRigBlueprint->GetFocusedModel())
 					{
@@ -4090,7 +4198,7 @@ void FRigVMEditorBase::UpdateGraphCompilerErrors()
 {
 	DECLARE_SCOPE_HIERARCHICAL_COUNTER_FUNC()
 
-	FRigVMAssetInterfacePtr Blueprint = GetRigVMAssetInterface();
+	FRigVMEditorAssetInterfacePtr Blueprint = GetRigVMAssetInterface();
 	URigVMHost* RigVMHost = GetRigVMHost();
 	if (Blueprint && RigVMHost && RigVMHost->GetVM())
 	{
@@ -4242,7 +4350,7 @@ FSlateIcon FRigVMEditorBase::GetEventQueueIcon() const
 
 void FRigVMEditorBase::GetDebugObjects(TArray<FRigVMCustomDebugObject>& DebugList) const
 {
-	FRigVMAssetInterfacePtr RigVMBlueprint = GetRigVMAssetInterface();
+	FRigVMEditorAssetInterfacePtr RigVMBlueprint = GetRigVMAssetInterface();
 	if (RigVMBlueprint == nullptr)
 	{
 		return;
@@ -4372,7 +4480,7 @@ void FRigVMEditorBase::HandleSetObjectBeingDebugged(UObject* InObject)
 		}
 	}
 
-	if (FRigVMAssetInterfacePtr RigBlueprint = GetRigVMAssetInterface())
+	if (FRigVMEditorAssetInterfacePtr RigBlueprint = GetRigVMAssetInterface())
 	{
 		if (URigVM* VM = RigBlueprint->GetVM(true))
 		{
@@ -4427,7 +4535,7 @@ FString FRigVMEditorBase::GetCustomDebugObjectLabel(UObject* ObjectBeingDebugged
 
 void FRigVMEditorBase::OnPIEStopped(bool bSimulation)
 {
-	if(FRigVMAssetInterfacePtr Blueprint = GetRigVMAssetInterface())
+	if(FRigVMEditorAssetInterfacePtr Blueprint = GetRigVMAssetInterface())
 	{
 		Blueprint->SetObjectBeingDebugged(GetRigVMHost());
 	}
@@ -4541,6 +4649,12 @@ void FRigVMEditorBase::OnGraphEditorTick(const FGeometry& AllottedGeometry, cons
 	}
 }
 
+void FRigVMEditorBase::SetGraphExplorerWidget(TSharedPtr<SRigVMEditorGraphExplorer> InWidget)
+{
+	GraphExplorerWidget = InWidget;
+	GetRigVMAssetInterface()->OnModified().AddSP(GraphExplorerWidget.ToSharedRef(), &SRigVMEditorGraphExplorer::HandleGraphModifiedEvent);
+}
+
 bool FRigVMEditorBase::IsEditingSingleBlueprint() const
 {
 	return GetRigVMAssetInterface() != nullptr;
diff --git a/Engine/Plugins/Runtime/RigVM/Source/RigVMEditor/Private/Editor/RigVMEditorMenuContext.cpp b/Engine/Plugins/Runtime/RigVM/Source/RigVMEditor/Private/Editor/RigVMEditorMenuContext.cpp
--- a/Engine/Plugins/Runtime/RigVM/Source/RigVMEditor/Private/Editor/RigVMEditorMenuContext.cpp
+++ b/Engine/Plugins/Runtime/RigVM/Source/RigVMEditor/Private/Editor/RigVMEditorMenuContext.cpp
@@ -2,7 +2,7 @@
 
 #include "Editor/RigVMEditorMenuContext.h"
 
-#include "RigVMAsset.h"
+#include "RigVMEditorAsset.h"
 #include "RigVMBlueprintLegacy.h"
 #include "Editor/RigVMNewEditor.h"
 #include "Framework/Application/SlateApplication.h"
@@ -17,11 +17,11 @@ void URigVMEditorMenuContext::Init(TWeakPtr<IRigVMEditor> InRigVMEditor, const F
 
 URigVMBlueprint* URigVMEditorMenuContext::GetRigVMBlueprint() const
 {
-	FRigVMAssetInterfacePtr Interface = GetRigVMAssetInterface();
+	FRigVMEditorAssetInterfacePtr Interface = GetRigVMAssetInterface();
 	return Cast<URigVMBlueprint>(Interface->GetObject());
 }
 
-FRigVMAssetInterfacePtr URigVMEditorMenuContext::GetRigVMAssetInterface() const
+FRigVMEditorAssetInterfacePtr URigVMEditorMenuContext::GetRigVMAssetInterface() const
 {
 	if (const TSharedPtr<IRigVMEditor> Editor = WeakRigVMEditor.Pin())
 	{
@@ -33,7 +33,7 @@ FRigVMAssetInterfacePtr URigVMEditorMenuContext::GetRigVMAssetInterface() const
 
 URigVMHost* URigVMEditorMenuContext::GetRigVMHost() const
 {
-	if (FRigVMAssetInterfacePtr RigBlueprint = GetRigVMAssetInterface())
+	if (FRigVMEditorAssetInterfacePtr RigBlueprint = GetRigVMAssetInterface())
 	{
 		if (URigVMHost* Host = RigBlueprint->GetDebuggedRigVMHost())
 		{
diff --git a/Engine/Plugins/Runtime/RigVM/Source/RigVMEditor/Private/Editor/RigVMFindReferences.cpp b/Engine/Plugins/Runtime/RigVM/Source/RigVMEditor/Private/Editor/RigVMFindReferences.cpp
--- a/Engine/Plugins/Runtime/RigVM/Source/RigVMEditor/Private/Editor/RigVMFindReferences.cpp
+++ b/Engine/Plugins/Runtime/RigVM/Source/RigVMEditor/Private/Editor/RigVMFindReferences.cpp
@@ -86,12 +86,12 @@ FRigVMFindResult::FRigVMFindResult(TWeakObjectPtr<URigVMBlueprint> InBlueprint,
 {
 }
 
-FRigVMFindResult::FRigVMFindResult(TWeakInterfacePtr<IRigVMAssetInterface> InBlueprint)
+FRigVMFindResult::FRigVMFindResult(TWeakInterfacePtr<IRigVMEditorAssetInterface> InBlueprint)
 	: WeakBlueprint(InBlueprint)
 {
 }
 
-FRigVMFindResult::FRigVMFindResult(TWeakInterfacePtr<IRigVMAssetInterface> InBlueprint, const FText& InDisplayText )
+FRigVMFindResult::FRigVMFindResult(TWeakInterfacePtr<IRigVMEditorAssetInterface> InBlueprint, const FText& InDisplayText )
 	: WeakBlueprint(InBlueprint), DisplayText(InDisplayText)
 {
 }
@@ -107,7 +107,7 @@ FReply FRigVMFindResult::OnClick()
 	return FReply::Handled();
 }
 
-UObject* FRigVMFindResult::GetObject(FRigVMAssetInterfacePtr InBlueprint) const
+UObject* FRigVMFindResult::GetObject(FRigVMEditorAssetInterfacePtr InBlueprint) const
 {
 	return GetRigVMAssetInterface()->GetObject();
 }
@@ -132,16 +132,16 @@ URigVMBlueprint* FRigVMFindResult::GetBlueprint() const
 	return Cast<URigVMBlueprint>(GetRigVMAssetInterface().GetObject());
 }
 
-FRigVMAssetInterfacePtr FRigVMFindResult::GetRigVMAssetInterface()
+FRigVMEditorAssetInterfacePtr FRigVMFindResult::GetRigVMAssetInterface()
 {
 	if (TStrongObjectPtr<UObject> Blueprint = WeakBlueprint.GetWeakObjectPtr().Pin())
 	{
-		return FRigVMAssetInterfacePtr(Blueprint.Get());
+		return FRigVMEditorAssetInterfacePtr(Blueprint.Get());
 	}
 	return nullptr;
 }
 
-const FRigVMAssetInterfacePtr FRigVMFindResult::GetRigVMAssetInterface() const
+const FRigVMEditorAssetInterfacePtr FRigVMFindResult::GetRigVMAssetInterface() const
 {
 	return const_cast<FRigVMFindResult*>(this)->GetRigVMAssetInterface();
 }
@@ -155,12 +155,12 @@ FText FRigVMFindResult::GetDisplayString() const
 // FRigVMFindReferencesGraphNode
 
 FRigVMFindReferencesGraphNode::FRigVMFindReferencesGraphNode(TWeakObjectPtr<URigVMBlueprint> InBlueprint)
-	: FRigVMFindResult(FRigVMAssetInterfacePtr(InBlueprint.Get())), Glyph(FAppStyle::GetAppStyleSetName(), "")
+	: FRigVMFindResult(FRigVMEditorAssetInterfacePtr(InBlueprint.Get())), Glyph(FAppStyle::GetAppStyleSetName(), "")
 	, Class(nullptr)
 {
 }
 
-FRigVMFindReferencesGraphNode::FRigVMFindReferencesGraphNode(FRigVMAssetInterfacePtr InBlueprint)
+FRigVMFindReferencesGraphNode::FRigVMFindReferencesGraphNode(FRigVMEditorAssetInterfacePtr InBlueprint)
 	: FRigVMFindResult(InBlueprint), Glyph(FAppStyle::GetAppStyleSetName(), "")
 	, Class(nullptr)
 {
@@ -168,7 +168,7 @@ FRigVMFindReferencesGraphNode::FRigVMFindReferencesGraphNode(FRigVMAssetInterfac
 
 FReply FRigVMFindReferencesGraphNode::OnClick()
 {
-	FRigVMAssetInterfacePtr Blueprint = GetRigVMAssetInterface();
+	FRigVMEditorAssetInterfacePtr Blueprint = GetRigVMAssetInterface();
 	if(Blueprint.GetObject())
 	{
 		UEdGraphNode* OutNode = NULL;
@@ -254,7 +254,7 @@ void FRigVMFindReferencesGraphNode::FinalizeSearchData()
 	}
 }
 
-UObject* FRigVMFindReferencesGraphNode::GetObject(FRigVMAssetInterfacePtr InBlueprint) const
+UObject* FRigVMFindReferencesGraphNode::GetObject(FRigVMEditorAssetInterfacePtr InBlueprint) const
 {
 	return FRigVMBlueprintUtils::GetNodeByGUID(InBlueprint, NodeGuid);
 }
@@ -268,13 +268,13 @@ FText FRigVMFindReferencesGraphNode::GetCategory() const
 // FRigVMFindReferencesPin
 
 FRigVMFindReferencesPin::FRigVMFindReferencesPin(TWeakObjectPtr<URigVMBlueprint> InBlueprint, FString InSchemaName)
-	: FRigVMFindResult(FRigVMAssetInterfacePtr(InBlueprint.Get()))
+	: FRigVMFindResult(FRigVMEditorAssetInterfacePtr(InBlueprint.Get()))
 	, SchemaName(InSchemaName)
 	, IconColor(FSlateColor::UseForeground())
 {
 }
 
-FRigVMFindReferencesPin::FRigVMFindReferencesPin(FRigVMAssetInterfacePtr InBlueprint, FString InSchemaName)
+FRigVMFindReferencesPin::FRigVMFindReferencesPin(FRigVMEditorAssetInterfacePtr InBlueprint, FString InSchemaName)
 	: FRigVMFindResult(InBlueprint)
 	, SchemaName(InSchemaName)
 	, IconColor(FSlateColor::UseForeground())
@@ -360,11 +360,11 @@ void FRigVMFindReferencesPin::FinalizeSearchData()
 // FRigVMFindReferencesVariable
 
 FRigVMFindReferencesVariable::FRigVMFindReferencesVariable(TWeakObjectPtr<URigVMBlueprint> InBlueprint)
-	: FRigVMFindResult(FRigVMAssetInterfacePtr(InBlueprint.Get()))
+	: FRigVMFindResult(FRigVMEditorAssetInterfacePtr(InBlueprint.Get()))
 {
 }
 
-FRigVMFindReferencesVariable::FRigVMFindReferencesVariable(FRigVMAssetInterfacePtr InBlueprint)
+FRigVMFindReferencesVariable::FRigVMFindReferencesVariable(FRigVMEditorAssetInterfacePtr InBlueprint)
 	: FRigVMFindResult(InBlueprint)
 {
 }
@@ -422,17 +422,17 @@ void FRigVMFindReferencesVariable::FinalizeSearchData()
 // FRigVMFindReferencesGraph
 
 FRigVMFindReferencesGraph::FRigVMFindReferencesGraph(TWeakObjectPtr<URigVMBlueprint> InBlueprint, EGraphType InGraphType)
-	: FRigVMFindResult(FRigVMAssetInterfacePtr(InBlueprint.Get())), GraphType(InGraphType)
+	: FRigVMFindResult(FRigVMEditorAssetInterfacePtr(InBlueprint.Get())), GraphType(InGraphType)
 {
 }
-FRigVMFindReferencesGraph::FRigVMFindReferencesGraph(FRigVMAssetInterfacePtr InBlueprint, EGraphType InGraphType)
+FRigVMFindReferencesGraph::FRigVMFindReferencesGraph(FRigVMEditorAssetInterfacePtr InBlueprint, EGraphType InGraphType)
 	: FRigVMFindResult(InBlueprint), GraphType(InGraphType)
 {
 }
 
 FReply FRigVMFindReferencesGraph::OnClick()
 {
-	FRigVMAssetInterfacePtr Blueprint = GetRigVMAssetInterface();
+	FRigVMEditorAssetInterfacePtr Blueprint = GetRigVMAssetInterface();
 	if(Blueprint.GetObject())
 	{
 		GEditor->GetEditorSubsystem<UAssetEditorSubsystem>()->OpenEditorForAsset(Blueprint->GetObject());
@@ -615,7 +615,7 @@ void SRigVMFindReferences::FindReferences(const FString& SearchTerms)
 		return;
 	}
 
-	FRigVMAssetInterfacePtr Blueprint = Editor->GetRigVMAssetInterface();
+	FRigVMEditorAssetInterfacePtr Blueprint = Editor->GetRigVMAssetInterface();
 	if (!Blueprint)
 	{
 		return;
@@ -820,7 +820,7 @@ void SRigVMFindReferences::FindReferences(const FString& SearchTerms)
 	TArray<FRigVMGraphVariableDescription> BPVariables = Blueprint->GetAssetVariables();
 	for (const FRigVMGraphVariableDescription& Variable : BPVariables)
 	{
-		if (FRigVMAssetInterfacePtr BP = Blueprint)
+		if (FRigVMEditorAssetInterfacePtr BP = Blueprint)
 		{
 			TArray<UBlueprintExtension::FSearchTagDataPair> BPVariableMetadata;
 			BP->AddVariableSearchMetaDataInfo(Variable.Name, BPVariableMetadata);
diff --git a/Engine/Plugins/Runtime/RigVM/Source/RigVMEditor/Private/Editor/RigVMGraphDetailCustomization.cpp b/Engine/Plugins/Runtime/RigVM/Source/RigVMEditor/Private/Editor/RigVMGraphDetailCustomization.cpp
--- a/Engine/Plugins/Runtime/RigVM/Source/RigVMEditor/Private/Editor/RigVMGraphDetailCustomization.cpp
+++ b/Engine/Plugins/Runtime/RigVM/Source/RigVMEditor/Private/Editor/RigVMGraphDetailCustomization.cpp
@@ -34,7 +34,9 @@
 #include "Editor/RigVMEditorTools.h"
 #include "ContentBrowserModule.h"
 #include "IContentBrowserSingleton.h"
+#include "RigVMDeveloperTypeUtils.h"
 #include "UObject/GarbageCollectionSchema.h"
+#include "UObject/ObjectPtr.h"
 #if WITH_RIGVMLEGACYEDITOR
 #include "Editor/RigVMLegacyEditor.h"
 #endif
@@ -755,7 +757,12 @@ TSharedRef<IDetailCustomization> FRigVMGraphDetailCustomization::MakeInstance(TS
 	const TArray<UObject*>* Objects = (InEditor.IsValid() ? InEditor->GetObjectsCurrentlyBeingEdited() : nullptr);
 	if (Objects && Objects->Num() == 1)
 	{
-		if (FRigVMAssetInterfacePtr RigVMBlueprint = (*Objects)[0])
+		UObject* EditorAssetObject = (*Objects)[0];
+		if (TScriptInterface<IRigVMRuntimeAssetInterface> RuntimeAsset = EditorAssetObject)
+		{
+			EditorAssetObject = RuntimeAsset->GetEditorOnlyData().Get();
+		}
+		if (FRigVMEditorAssetInterfacePtr RigVMBlueprint = EditorAssetObject)
 		{
 			if(RigVMBlueprint.GetObject()->GetClass() == InExpectedBlueprintClass)
 			{
@@ -764,7 +771,7 @@ TSharedRef<IDetailCustomization> FRigVMGraphDetailCustomization::MakeInstance(TS
 		}
 	}
 
-	return MakeShareable(new FRigVMGraphDetailCustomization((TSharedPtr<IRigVMEditor>)nullptr, FRigVMAssetInterfacePtr(nullptr)));
+	return MakeShareable(new FRigVMGraphDetailCustomization((TSharedPtr<IRigVMEditor>)nullptr, FRigVMEditorAssetInterfacePtr(nullptr)));
 }
 
 #if WITH_RIGVMLEGACYEDITOR
@@ -773,9 +780,9 @@ TSharedPtr<IDetailCustomization> FRigVMGraphDetailCustomization::MakeLegacyInsta
 	const TArray<UObject*>* Objects = (InBlueprintEditor.IsValid() ? InBlueprintEditor->GetObjectsCurrentlyBeingEdited() : nullptr);
 	if (Objects && Objects->Num() == 1)
 	{
-		if ((*Objects)[0]->Implements<URigVMAssetInterface>())
+		if ((*Objects)[0]->Implements<URigVMEditorAssetInterface>())
 		{
-			FRigVMAssetInterfacePtr RigVMBlueprint = (*Objects)[0];
+			FRigVMEditorAssetInterfacePtr RigVMBlueprint = (*Objects)[0];
 			if(RigVMBlueprint->GetObject()->GetClass() == InExpectedBlueprintClass)
 			{
 				return MakeShareable(new FRigVMGraphDetailCustomization(InBlueprintEditor, RigVMBlueprint));
@@ -783,10 +790,10 @@ TSharedPtr<IDetailCustomization> FRigVMGraphDetailCustomization::MakeLegacyInsta
 		}
 	}
 
-	return MakeShareable(new FRigVMGraphDetailCustomization((TSharedPtr<IRigVMEditor>)nullptr, FRigVMAssetInterfacePtr(nullptr)));
+	return MakeShareable(new FRigVMGraphDetailCustomization((TSharedPtr<IRigVMEditor>)nullptr, FRigVMEditorAssetInterfacePtr(nullptr)));
 }
 
-FRigVMGraphDetailCustomization::FRigVMGraphDetailCustomization(TSharedPtr<IBlueprintEditor> RigVMigEditor, FRigVMAssetInterfacePtr RigVMBlueprint)
+FRigVMGraphDetailCustomization::FRigVMGraphDetailCustomization(TSharedPtr<IBlueprintEditor> RigVMigEditor, FRigVMEditorAssetInterfacePtr RigVMBlueprint)
 	: RigVMBlueprintPtr(RigVMBlueprint)
 	, RigVMGraphDetailCustomizationImpl(MakeShared<FRigVMGraphDetailCustomizationImpl>())
 {
@@ -794,7 +801,7 @@ FRigVMGraphDetailCustomization::FRigVMGraphDetailCustomization(TSharedPtr<IBluep
 }
 #endif
 
-FRigVMGraphDetailCustomization::FRigVMGraphDetailCustomization(TSharedPtr<IRigVMEditor> RigVMigEditor, FRigVMAssetInterfacePtr RigVMBlueprint)
+FRigVMGraphDetailCustomization::FRigVMGraphDetailCustomization(TSharedPtr<IRigVMEditor> RigVMigEditor, FRigVMEditorAssetInterfacePtr RigVMBlueprint)
 	: RigVMEditorPtr(RigVMigEditor)
 	, RigVMBlueprintPtr(RigVMBlueprint)
 	, RigVMGraphDetailCustomizationImpl(MakeShared<FRigVMGraphDetailCustomizationImpl>())
@@ -809,7 +816,7 @@ void FRigVMGraphDetailCustomization::CustomizeDetails(IDetailLayoutBuilder& Deta
 	GraphPtr = CastChecked<URigVMEdGraph>(Objects[0].Get());
 	URigVMEdGraph* Graph = GraphPtr.Get();
 
-	FRigVMAssetInterfacePtr Blueprint = RigVMBlueprintPtr.GetObject();
+	FRigVMEditorAssetInterfacePtr Blueprint = RigVMBlueprintPtr.GetObject();
 	URigVMGraph* Model = nullptr;
 	URigVMController* Controller = nullptr;
 
@@ -1999,7 +2006,7 @@ void FRigVMWrappedNodeDetailCustomization::CustomizeDetails(IDetailLayoutBuilder
 		URigVMDetailsViewWrapperObject* WrapperObject = CastChecked<URigVMDetailsViewWrapperObject>(DetailObject.Get());
 		if(BlueprintBeingCustomized == nullptr)
 		{
-			BlueprintBeingCustomized = IRigVMAssetInterface::GetInterfaceOuter(WrapperObject);
+			BlueprintBeingCustomized = IRigVMEditorAssetInterface::GetInterfaceOuter(WrapperObject);
 		}
 
 		ObjectsBeingCustomized.Add(WrapperObject);
@@ -3047,7 +3054,7 @@ void FRigVMGraphEnumDetailCustomization::CustomizeHeader(TSharedRef<IPropertyHan
 
 		if(BlueprintBeingCustomized == nullptr)
 		{
-			BlueprintBeingCustomized = IRigVMAssetInterface::GetInterfaceOuter(Object);
+			BlueprintBeingCustomized = IRigVMEditorAssetInterface::GetInterfaceOuter(Object);
 		}
 
 		if(GraphBeingCustomized == nullptr)
@@ -3814,7 +3821,7 @@ void FRigVMGraphMathTypeDetailCustomization::CustomizeHeader(TSharedRef<IPropert
 	{
 		if(BlueprintBeingCustomized == nullptr)
 		{
-			BlueprintBeingCustomized = IRigVMAssetInterface::GetInterfaceOuter(Object);
+			BlueprintBeingCustomized = IRigVMEditorAssetInterface::GetInterfaceOuter(Object);
 		}
 
 		if(GraphBeingCustomized == nullptr)
@@ -3872,7 +3879,7 @@ void FRigVMGraphMathTypeDetailCustomization::CustomizeChildren(TSharedRef<IPrope
 	{
 		if(BlueprintBeingCustomized == nullptr)
 		{
-			BlueprintBeingCustomized = IRigVMAssetInterface::GetInterfaceOuter(Object);
+			BlueprintBeingCustomized = IRigVMEditorAssetInterface::GetInterfaceOuter(Object);
 		}
 
 		if(GraphBeingCustomized == nullptr)
diff --git a/Engine/Plugins/Runtime/RigVM/Source/RigVMEditor/Private/Editor/RigVMLegacyEditor.cpp b/Engine/Plugins/Runtime/RigVM/Source/RigVMEditor/Private/Editor/RigVMLegacyEditor.cpp
--- a/Engine/Plugins/Runtime/RigVM/Source/RigVMEditor/Private/Editor/RigVMLegacyEditor.cpp
+++ b/Engine/Plugins/Runtime/RigVM/Source/RigVMEditor/Private/Editor/RigVMLegacyEditor.cpp
@@ -38,10 +38,10 @@ void FRigVMLegacyEditor::CreateEditorToolbar()
 	}
 }
 
-void FRigVMLegacyEditor::CommonInitialization(const TArray<FRigVMAssetInterfacePtr>& InitBlueprints, bool bShouldOpenInDefaultsMode)
+void FRigVMLegacyEditor::CommonInitialization(const TArray<FRigVMEditorAssetInterfacePtr>& InitBlueprints, bool bShouldOpenInDefaultsMode)
 {
 	TArray<UBlueprint*> Blueprints;
-	Algo::Transform(InitBlueprints, Blueprints, [](FRigVMAssetInterfacePtr In)
+	Algo::Transform(InitBlueprints, Blueprints, [](FRigVMEditorAssetInterfacePtr In)
 	{
 		return Cast<UBlueprint>(In->GetObject());
 	});
diff --git a/Engine/Plugins/Runtime/RigVM/Source/RigVMEditor/Private/Editor/RigVMLocalVariableDetailCustomization.cpp b/Engine/Plugins/Runtime/RigVM/Source/RigVMEditor/Private/Editor/RigVMLocalVariableDetailCustomization.cpp
--- a/Engine/Plugins/Runtime/RigVM/Source/RigVMEditor/Private/Editor/RigVMLocalVariableDetailCustomization.cpp
+++ b/Engine/Plugins/Runtime/RigVM/Source/RigVMEditor/Private/Editor/RigVMLocalVariableDetailCustomization.cpp
@@ -13,6 +13,7 @@
 #include "RigVMBlueprintLegacy.h"
 #include "RigVMCore/RigVM.h"
 #include "InstancedPropertyBagStructureDataProvider.h"
+#include "RigVMDeveloperTypeUtils.h"
 
 #define LOCTEXT_NAMESPACE "LocalVariableDetails"
 
@@ -37,7 +38,7 @@ void FRigVMLocalVariableDetailCustomization::CustomizeDetails(IDetailLayoutBuild
 		}
 		if (ensure(GraphBeingCustomized))
 		{
-			BlueprintBeingCustomized = IRigVMAssetInterface::GetInterfaceOuter(GraphBeingCustomized);
+			BlueprintBeingCustomized = IRigVMEditorAssetInterface::GetInterfaceOuter(GraphBeingCustomized);
 		}
 
 		NameValidator = FRigVMLocalVariableNameValidator(BlueprintBeingCustomized, GraphBeingCustomized, VariableDescription.Name);
@@ -114,28 +115,29 @@ void FRigVMLocalVariableDetailCustomization::CustomizeDetails(IDetailLayoutBuild
 		MemoryStorage = MakeShared<FRigVMMemoryStorageStruct>();
 		MemoryStorage->AddProperties(PropertyDescriptions);
 
-
 		FRigVMMemoryStorageStruct* Memory = MemoryStorage.Get();
-		FProperty* Property = Memory->FindPropertyByName(*SourcePath);
+		const FName SanitizedName = FInstancedPropertyBag::SanitizePropertyName(SourcePath);
+		FProperty* Property = Memory->FindPropertyByName(SanitizedName);
 		if (Property)
 		{
 			IDetailCategoryBuilder& DefaultValueCategory = DetailBuilder.EditCategory(TEXT("DefaultValueCategory"), LOCTEXT("DefaultValueCategoryHeading", "Default Value"));
 			Property->ClearPropertyFlags(CPF_EditConst);
+			Property->PropertyFlags |= CPF_Edit;
 			
-			const FName SanitizedName = FRigVMPropertyDescription::SanitizeName(*SourcePath);
-			IDetailPropertyRow* Row = DefaultValueCategory.AddExternalStructureProperty(MakeShared<FInstancePropertyBagStructureDataProvider>(*Memory), SanitizedName);
-
-			Row->DisplayName(FText::FromName(VariableDescription.Name));
-
-			const FSimpleDelegate OnDefaultValueChanged = FSimpleDelegate::CreateLambda([this, Property]()
+			if (IDetailPropertyRow* Row = DefaultValueCategory.AddExternalStructureProperty(MakeShared<FInstancePropertyBagStructureDataProvider>(*Memory), SanitizedName))
 			{
-				VariableDescription.DefaultValue = MemoryStorage->GetDataAsString(MemoryStorage->GetPropertyIndex(Property));
-				DefaultValueHandle->SetValue(VariableDescription.DefaultValue);
-			});
-
-			TSharedPtr<IPropertyHandle> Handle = Row->GetPropertyHandle();
-			Handle->SetOnPropertyValueChanged(OnDefaultValueChanged);
-			Handle->SetOnChildPropertyValueChanged(OnDefaultValueChanged);
+				Row->DisplayName(FText::FromName(VariableDescription.Name));
+
+				const FSimpleDelegate OnDefaultValueChanged = FSimpleDelegate::CreateLambda([this, Property]()
+				{
+					VariableDescription.DefaultValue = MemoryStorage->GetDataAsString(MemoryStorage->GetPropertyIndex(Property));
+					DefaultValueHandle->SetValue(VariableDescription.DefaultValue);
+				});
+
+				TSharedPtr<IPropertyHandle> Handle = Row->GetPropertyHandle();
+				Handle->SetOnPropertyValueChanged(OnDefaultValueChanged);
+				Handle->SetOnChildPropertyValueChanged(OnDefaultValueChanged);
+			}
 		}
 	}
 }
@@ -173,14 +175,18 @@ FEdGraphPinType FRigVMLocalVariableDetailCustomization::OnGetPinInfo() const
 {
 	if (!VariableDescription.Name.IsNone())
 	{
-		return VariableDescription.ToPinType();
+		const FEdGraphPinType EdPinType = RigVMTypeUtils::PinTypeFromRigVMVariableDescription(VariableDescription);
+		return EdPinType;
 	}
 	return FEdGraphPinType();
 }
 
 void FRigVMLocalVariableDetailCustomization::HandlePinInfoChanged(const FEdGraphPinType& PinType)
 {
-	VariableDescription.ChangeType(PinType);
+	FString CPPType;
+	UObject* CPPTypeObject = nullptr;
+	RigVMTypeUtils::CPPTypeFromPinType(PinType, CPPType, &CPPTypeObject);
+	VariableDescription.ChangeType(CPPType, CPPTypeObject);
 	FRigVMBlueprintCompileScope CompileScope(BlueprintBeingCustomized);
 	TypeHandle->SetValue(VariableDescription.CPPType);
 	TypeObjectHandle->SetValue(VariableDescription.CPPTypeObject);	
diff --git a/Engine/Plugins/Runtime/RigVM/Source/RigVMEditor/Private/Editor/RigVMMinimalEnvironment.cpp b/Engine/Plugins/Runtime/RigVM/Source/RigVMEditor/Private/Editor/RigVMMinimalEnvironment.cpp
--- a/Engine/Plugins/Runtime/RigVM/Source/RigVMEditor/Private/Editor/RigVMMinimalEnvironment.cpp
+++ b/Engine/Plugins/Runtime/RigVM/Source/RigVMEditor/Private/Editor/RigVMMinimalEnvironment.cpp
@@ -74,7 +74,7 @@ void FRigVMMinimalEnvironment::SetFunctionNode(const FRigVMGraphFunctionIdentifi
 		{
 			if(const UClass* Class = AssetData.GetClass())
 			{
-				if(Class->ImplementsInterface(URigVMAssetInterface::StaticClass()))
+				if(Class->ImplementsInterface(URigVMEditorAssetInterface::StaticClass()))
 				{
 					SetSchemata(Class);
 				}
@@ -111,7 +111,7 @@ void FRigVMMinimalEnvironment::SetSchemata(const UClass* InRigVMBlueprintClass)
 {
 	check(InRigVMBlueprintClass);
 
-	const FRigVMAssetInterfacePtr CDO = InRigVMBlueprintClass->GetDefaultObject();
+	const FRigVMEditorAssetInterfacePtr CDO = InRigVMBlueprintClass->GetDefaultObject();
 	check(CDO);
 
 	EdGraphClass = CDO->GetRigVMClientHost()->GetRigVMEdGraphClass();
diff --git a/Engine/Plugins/Runtime/RigVM/Source/RigVMEditor/Private/Editor/RigVMNewEditor.cpp b/Engine/Plugins/Runtime/RigVM/Source/RigVMEditor/Private/Editor/RigVMNewEditor.cpp
--- a/Engine/Plugins/Runtime/RigVM/Source/RigVMEditor/Private/Editor/RigVMNewEditor.cpp
+++ b/Engine/Plugins/Runtime/RigVM/Source/RigVMEditor/Private/Editor/RigVMNewEditor.cpp
@@ -36,7 +36,7 @@
 #include "Editor/RigVMCompilerResultsTabSummoner.h"
 #include "Editor/RigVMDetailsInspectorTabSummoner.h"
 #include "Editor/RigVMFindReferencesTabSummoner.h"
-#include "Editor/RigVMGraphEditorSummoner.h"
+#include "Editor/RigVMGraphEditorTabSummoner.h"
 #include "Widgets/Input/SEditableTextBox.h"
 #include "Editor/SRigVMDetailsInspector.h"
 #include "Engine/LevelScriptBlueprint.h"
@@ -204,7 +204,28 @@ void FRigVMNewEditor::InitAssetEditor(const EToolkitMode::Type Mode, const TShar
 	FWorkflowCentricApplication::InitAssetEditor(Mode, InitToolkitHost, AppIdentifier, StandaloneDefaultLayout, bCreateDefaultStandaloneMenu, bCreateDefaultToolbar, ObjectsToEdit, bInIsToolbarFocusable, bInUseSmallToolbarIcons, InOpenMethod);
 }
 
-void FRigVMNewEditor::CommonInitialization(const TArray<FRigVMAssetInterfacePtr>& InitBlueprints, bool bShouldOpenInDefaultsMode)
+void FRigVMNewEditor::HandlePreVariablesChanged(FName InName)
+{
+	if (Inspector && GraphExplorerWidget)
+	{
+		if (GraphExplorerWidget->GetSelectedType() == ERigVMExplorerElementType::Variable)
+		{
+			ElementInDetailsPanel = GraphExplorerWidget->GetSelectedVariableName();
+			ClearDetailObject();
+		}
+	}
+}
+
+void FRigVMNewEditor::HandlePostVariablesChanged(FName InName)
+{
+	if (Inspector)
+	{
+		SelectVariable(ElementInDetailsPanel);
+		ElementInDetailsPanel = NAME_None;
+	}
+}
+
+void FRigVMNewEditor::CommonInitialization(const TArray<FRigVMEditorAssetInterfacePtr>& InitBlueprints, bool bShouldOpenInDefaultsMode)
 {
 	TSharedPtr<FRigVMNewEditor> ThisPtr = SharedThis(this);
 
@@ -217,8 +238,8 @@ void FRigVMNewEditor::CommonInitialization(const TArray<FRigVMAssetInterfacePtr>
 		//DocumentManager->RegisterDocumentFactory(MakeShareable(new FTimelineEditorSummoner(ThisPtr)));
 
 		TSharedRef<FRigVMNewEditor> SharedRef = StaticCastSharedRef<FRigVMNewEditor>(SharedThis(this));
-		TSharedRef<FDocumentTabFactory> GraphEditorFactory = MakeShareable(new FRigVMGraphEditorSummoner(ThisPtr,
-			FRigVMGraphEditorSummoner::FOnCreateGraphEditorWidget::CreateSP(SharedRef, &FRigVMNewEditor::CreateGraphEditorWidget)
+		TSharedRef<FDocumentTabFactory> GraphEditorFactory = MakeShareable(new FRigVMGraphEditorTabSummoner(SharedRef,
+			FRigVMGraphEditorTabSummoner::FOnCreateGraphEditorWidget::CreateSP(SharedRef, &FRigVMNewEditor::CreateGraphEditorWidget)
 			));
 
 		// Also store off a reference to the grapheditor factory so we can find all the tabs spawned by it later.
@@ -259,7 +280,7 @@ void FRigVMNewEditor::CommonInitialization(const TArray<FRigVMAssetInterfacePtr>
 
 		//FLoadObjectsFromAssetRegistryHelper::Load<UUserDefinedEnum>(UserDefinedEnumerators);
 
-		FRigVMAssetInterfacePtr InitBlueprint = InitBlueprints[0];
+		FRigVMEditorAssetInterfacePtr InitBlueprint = InitBlueprints[0];
 
 		// Update the blueprint if required
 		ERigVMAssetStatus OldStatus = InitBlueprint->GetAssetStatus();
@@ -277,6 +298,21 @@ void FRigVMNewEditor::CommonInitialization(const TArray<FRigVMAssetInterfacePtr>
 		InitBlueprint->OnChanged().AddSP(this, &FRigVMNewEditor::OnBlueprintChanged); // TODO sara-s
 		// InitBlueprint->OnCompiled().AddSP(this, &FRigVMNewEditor::OnBlueprintCompiled); // TODO sara-s
 		InitBlueprint->OnSetObjectBeingDebugged().AddSP(this, &FRigVMNewEditor::HandleSetObjectBeingDebugged);
+
+		InitBlueprint->GetVMCompileSettings().ASTSettings.ReportDelegate.BindSPLambda(this,
+			[InitBlueprint](EMessageSeverity::Type InSeverity, UObject* InSubject, const FString& InMessage)
+			{ InitBlueprint->HandleReportFromCompiler(InSeverity, InSubject, InMessage); });
+
+		InitBlueprint->OnExternalVariablesChanged().AddSPLambda(this, [this](const TArray<FRigVMExternalVariable>& Variables)
+			{
+				if (GraphExplorerWidget.IsValid())
+				{
+					GraphExplorerWidget->Refresh();
+				}
+			});
+
+		InitBlueprint->OnPreVariablesChanged().AddSP(this, &FRigVMNewEditor::HandlePreVariablesChanged);
+		InitBlueprint->OnPostVariablesChanged().AddSP(this, &FRigVMNewEditor::HandlePostVariablesChanged);
 	}
 
 	bWasOpenedInDefaultsMode = bShouldOpenInDefaultsMode;
@@ -288,6 +324,7 @@ void FRigVMNewEditor::CommonInitialization(const TArray<FRigVMAssetInterfacePtr>
 
 	// FKismetEditorUtilities::OnBlueprintUnloaded.AddSP(this, &FRigVMNewEditor::OnBlueprintUnloaded); // TODO sara-s
 
+
 	UEditorEngine* Editor = Cast<UEditorEngine>(GEngine);
 	if (Editor != nullptr)
 	{
@@ -350,7 +387,7 @@ void FRigVMNewEditor::SetCurrentMode(FName NewMode)
 void FRigVMNewEditor::PostLayoutBlueprintEditorInitialization()
 {
 	// TODO sara-s: Replicate what blueprint editor is doing
-	if (FRigVMAssetInterfacePtr Blueprint = GetRigVMAssetInterface())
+	if (FRigVMEditorAssetInterfacePtr Blueprint = GetRigVMAssetInterface())
 	{
 		// Refresh the graphs
 		ForceEditorRefresh();
@@ -524,7 +561,7 @@ FText FRigVMNewEditor::GetCompileStatusTooltip() const
 {
 	// Copied from FBlueprintEditorToolbar::GetStatusTooltip
 
-	FRigVMAssetInterfacePtr BlueprintObj = GetRigVMAssetInterface();
+	FRigVMEditorAssetInterfacePtr BlueprintObj = GetRigVMAssetInterface();
 	ERigVMAssetStatus Status = BlueprintObj->GetAssetStatus();
 
 	switch (Status)
@@ -546,7 +583,7 @@ FSlateIcon FRigVMNewEditor::GetCompileStatusImage() const
 {
 	// Copied from FBlueprintEditorToolbar::GetStatusImage
 
-	FRigVMAssetInterfacePtr BlueprintObj = GetRigVMAssetInterface();
+	FRigVMEditorAssetInterfacePtr BlueprintObj = GetRigVMAssetInterface();
 	ERigVMAssetStatus Status = BlueprintObj->GetAssetStatus();
 
 	
@@ -683,6 +720,7 @@ void FRigVMNewEditor::OnAddNewVariable()
 	else
 	{
 		RenameNewlyAddedAction(VarName);
+		GraphExplorerWidget->Refresh();
 	}
 }
 
@@ -711,7 +749,7 @@ void FRigVMNewEditor::RestoreEditedObjectState()
 {
 	check(IsEditingSingleBlueprint());
 
-	FRigVMAssetInterfacePtr Blueprint = GetRigVMAssetInterface();
+	FRigVMEditorAssetInterfacePtr Blueprint = GetRigVMAssetInterface();
 	
 	TSet<FSoftObjectPath> PathsToRemove;
 	for (int32 i = 0; i < Blueprint->GetLastEditedDocuments().Num(); i++)
@@ -844,7 +882,7 @@ void FRigVMNewEditor::SetUISelectionState(FName SelectionOwner)
 	}
 }
 
-void FRigVMNewEditor::AnalyticsTrackNodeEvent(IRigVMAssetInterface* Blueprint, UEdGraphNode* GraphNode, bool bNodeDelete) const
+void FRigVMNewEditor::AnalyticsTrackNodeEvent(IRigVMEditorAssetInterface* Blueprint, UEdGraphNode* GraphNode, bool bNodeDelete) const
 {
 	// TODO sara-s
 	// if(Blueprint && GraphNode && FEngineAnalytics::IsAvailable())
@@ -886,6 +924,15 @@ void FRigVMNewEditor::AnalyticsTrackNodeEvent(IRigVMAssetInterface* Blueprint, U
 	// }
 }
 
+TSharedRef<SWidget> FRigVMNewEditor::GetCompilerResults() const
+{
+	if (CompilerResults.IsValid())
+	{
+		return CompilerResults.ToSharedRef();
+	}
+	return SNullWidget::NullWidget;
+}
+
 void FRigVMNewEditor::JumpToHyperlink(const UObject* ObjectReference, bool bRequestRename)
 {
 	if (FRigVMEditorBase::JumpToHyperlinkImpl(ObjectReference, bRequestRename))
@@ -1007,7 +1054,7 @@ bool FRigVMNewEditor::IsSaveOnCompileEnabled() const
 {
 	// Copied from BlueprintEditorImpl::IsSaveOnCompileEnabled
 	
-	FRigVMAssetInterfacePtr Blueprint = GetRigVMAssetInterface();
+	FRigVMEditorAssetInterfacePtr Blueprint = GetRigVMAssetInterface();
 	bool const bIsLevelScript = (Cast<ULevelScriptBlueprint>(Blueprint->GetObject()) != nullptr);
 
 	return !bIsLevelScript;
@@ -1323,7 +1370,7 @@ void FRigVMNewEditor::CompileImpl()
 {
 	DECLARE_SCOPE_HIERARCHICAL_COUNTER_FUNC()
 
-	FRigVMAssetInterfacePtr BlueprintObj = GetRigVMAssetInterface();
+	FRigVMEditorAssetInterfacePtr BlueprintObj = GetRigVMAssetInterface();
 	if (BlueprintObj)
 	{
 		FMessageLog BlueprintLog("BlueprintLog");
@@ -1334,10 +1381,13 @@ void FRigVMNewEditor::CompileImpl()
 
 		FCompilerResultsLog LogResults = GetRigVMAssetInterface()->CompileBlueprint();
 
-		CompilerResultsListing->ClearMessages();
+		if (CompilerResultsListing)
+		{
+			CompilerResultsListing->ClearMessages();
 
-		// Note we dont mirror to the output log here as the compiler already does that
-		CompilerResultsListing->AddMessages(LogResults.Messages, false);
+			// Note we dont mirror to the output log here as the compiler already does that
+			CompilerResultsListing->AddMessages(LogResults.Messages, false);
+		}
 
 		const bool bForceMessageDisplay = (LogResults.NumWarnings > 0 || LogResults.NumErrors > 0) && !BlueprintObj->IsRegeneratingOnLoad();
 		if (bForceMessageDisplay)
@@ -1357,7 +1407,7 @@ bool FRigVMNewEditor::IsGraphInCurrentBlueprint(const UEdGraph* InGraph) const
 {
 	bool bEditable = true;
 
-	FRigVMAssetInterfacePtr EditingBP = GetRigVMAssetInterface();
+	FRigVMEditorAssetInterfacePtr EditingBP = GetRigVMAssetInterface();
 	if(EditingBP)
 	{
 		TArray<UEdGraph*> Graphs;
@@ -1423,7 +1473,7 @@ void FRigVMNewEditor::JumpToPin(const UEdGraphPin* Pin)
 	}
 }
 
-void FRigVMNewEditor::OnBlueprintChangedInnerImpl(IRigVMAssetInterface* InBlueprint, bool bIsJustBeingCompiled)
+void FRigVMNewEditor::OnBlueprintChangedInnerImpl(IRigVMEditorAssetInterface* InBlueprint, bool bIsJustBeingCompiled)
 {
 	if (InBlueprint)
 	{
@@ -1555,8 +1605,19 @@ void FRigVMNewEditor::SetupGraphEditorEventsImpl(UEdGraph* InGraph, SGraphEditor
 	}
 }
 
+void FRigVMNewEditor::OnClosedActionMenu(UEdGraph* InEdGraph) const
+{
+	if (TSharedPtr<SGraphEditor> GraphEditor = GetGraphEditor(InEdGraph))
+	{
+		if (SGraphPanel* GraphPanel = GraphEditor->GetGraphPanel())
+		{
+			GraphPanel->OnStopMakingConnection(true);
+		}
+	}
+}
+
 FActionMenuContent FRigVMNewEditor::OnCreateGraphActionMenu(UEdGraph* InGraph, const FVector2f& InNodePosition,
-	const TArray<UEdGraphPin*>& InDraggedPins, bool bAutoExpand, SGraphEditor::FActionMenuClosed InOnMenuClosed)
+                                                            const TArray<UEdGraphPin*>& InDraggedPins, bool bAutoExpand, SGraphEditor::FActionMenuClosed InOnMenuClosed)
 {
 	// TODO sara-s: Replicate what blueprint editor is doing
 	HasOpenActionMenu = InGraph;
@@ -1569,7 +1630,8 @@ FActionMenuContent FRigVMNewEditor::OnCreateGraphActionMenu(UEdGraph* InGraph, c
 		SNew(SRigVMActionMenu, SharedThis(this))
 		.GraphObj(InGraph)
 		.NewNodePosition(InNodePosition)
-		.DraggedFromPins(InDraggedPins);
+		.DraggedFromPins(InDraggedPins)
+		.OnClosedCallback( SGraphEditor::FActionMenuClosed::CreateSP(this, &FRigVMNewEditor::OnClosedActionMenu, InGraph));
 
 	return FActionMenuContent( ActionMenu, ActionMenu->GetFilterTextBox() );
 }
@@ -1615,7 +1677,7 @@ FGraphAppearanceInfo FRigVMNewEditor::GetGraphAppearanceImpl(UEdGraph* InGraph)
 
 	// Copy from FBlueprintEditor::GetPIEStatus
 	{
-		FRigVMAssetInterfacePtr CurrentBlueprint = GetRigVMAssetInterface();
+		FRigVMEditorAssetInterfacePtr CurrentBlueprint = GetRigVMAssetInterface();
 		UWorld *DebugWorld = nullptr;
 		ENetMode NetMode = NM_Standalone;
 		if (CurrentBlueprint)
@@ -1689,7 +1751,7 @@ void FRigVMNewEditor::NotifyPostChangeImpl(const FPropertyChangedEvent& Property
 	
 	if (IsEditingSingleBlueprint())
 	{
-		FRigVMAssetInterfacePtr Blueprint = GetRigVMAssetInterface();
+		FRigVMEditorAssetInterfacePtr Blueprint = GetRigVMAssetInterface();
 		UPackage* BlueprintPackage = Blueprint->GetObject()->GetOutermost();
 
 		// if any of the objects being edited are in our package, mark us as dirty
@@ -1771,13 +1833,14 @@ void FRigVMNewEditor::EditGlobalOptions_Clicked()
 	}
 	else
 	{
-		FRigVMAssetInterfacePtr Blueprint = GetRigVMAssetInterface();
+		FRigVMEditorAssetInterfacePtr Blueprint = GetRigVMAssetInterface();
 		if (Blueprint != nullptr)
 		{
 			// Show details for the Blueprint instance we're editing
 			if (Inspector.IsValid())
 			{
-				Inspector->ShowDetailsForSingleObject(Blueprint->GetObject());
+				SRigVMDetailsInspector::FShowDetailsOptions Options(FText::FromString(TEXT("Settings")), true);
+				Inspector->ShowDetailsForSingleObject(Blueprint->GetRuntimeAssetInterface().GetObject(), Options);
 			}
 
 			TryInvokingDetailsTab();
@@ -1973,9 +2036,9 @@ void FRigVMNewEditor::DumpMessagesToCompilerLog(const TArray<TSharedRef<class FT
 	}
 }
 
-void FRigVMNewEditor::CreateDefaultTabContents(const TArray<FRigVMAssetInterfacePtr> InBlueprints)
+void FRigVMNewEditor::CreateDefaultTabContents(const TArray<FRigVMEditorAssetInterfacePtr> InBlueprints)
 {
-	FRigVMAssetInterfacePtr InBlueprint = InBlueprints.Num() == 1 ? InBlueprints[0] : nullptr;
+	FRigVMEditorAssetInterfacePtr InBlueprint = InBlueprints.Num() == 1 ? InBlueprints[0] : nullptr;
 
 	// Cache off whether or not this is an interface, since it is used to govern multiple widget's behavior
 	//const bool bIsInterface = (InBlueprint && InBlueprint->BlueprintType == BPTYPE_Interface);
@@ -2000,7 +2063,43 @@ void FRigVMNewEditor::CreateDefaultTabContents(const TArray<FRigVMAssetInterface
 
 		FMessageLogModule& MessageLogModule = FModuleManager::LoadModuleChecked<FMessageLogModule>("MessageLog");
 		CompilerResults = MessageLogModule.CreateLogListingWidget( CompilerResultsListing.ToSharedRef() );
-	} 
+	}
+	else
+	{
+		// GetBlueprintMessageLog
+		{
+			FMessageLogModule& MessageLogModule = FModuleManager::LoadModuleChecked<FMessageLogModule>("MessageLog");
+
+			FName LogName;
+			//const FName LogName = GetBlueprintMessageLogName(InBlueprint);
+			{
+				if (InBlueprint != nullptr)
+				{
+					LogName = *FString::Printf(TEXT("%s_CompilerResultsLog"), *InBlueprint->GetObject()->GetName());
+				}
+				else
+				{
+					LogName = "BlueprintCompiler";
+				}
+			}
+
+			// Reuse any existing log, or create a new one (that is not held onto bey the message log system)
+			if(MessageLogModule.IsRegisteredLogListing(LogName))
+			{
+				CompilerResultsListing = MessageLogModule.GetLogListing(LogName);
+			}
+			else
+			{
+				FMessageLogInitializationOptions LogInitOptions;
+				LogInitOptions.bShowInLogWindow = false;
+				CompilerResultsListing = MessageLogModule.CreateLogListing(LogName, LogInitOptions);
+			}
+		}
+		CompilerResultsListing->OnMessageTokenClicked().AddSP(this, &FRigVMNewEditor::OnLogTokenClicked);
+
+		FMessageLogModule& MessageLogModule = FModuleManager::LoadModuleChecked<FMessageLogModule>("MessageLog");
+		CompilerResults = MessageLogModule.CreateLogListingWidget( CompilerResultsListing.ToSharedRef() );
+	}
 
 	FindResults = SNew(SRigVMFindReferences, SharedThis(this));
 	
diff --git a/Engine/Plugins/Runtime/RigVM/Source/RigVMEditor/Private/Editor/RigVMNewEditorMode.cpp b/Engine/Plugins/Runtime/RigVM/Source/RigVMEditor/Private/Editor/RigVMNewEditorMode.cpp
--- a/Engine/Plugins/Runtime/RigVM/Source/RigVMEditor/Private/Editor/RigVMNewEditorMode.cpp
+++ b/Engine/Plugins/Runtime/RigVM/Source/RigVMEditor/Private/Editor/RigVMNewEditorMode.cpp
@@ -7,6 +7,7 @@
 #include "Editor/RigVMEditor.h"
 #include "Editor/RigVMExecutionStackTabSummoner.h"
 #include "Editor/RigVMFindReferencesTabSummoner.h"
+#include "Editor/RigVMGraphEditorTabSummoner.h"
 #include "Editor/RigVMNewEditor.h"
 
 const FName FRigVMNewEditorApplicationModes::StandardRigVMEditorMode() { static FName Mode = TEXT("GraphName"); return Mode; };
@@ -23,7 +24,7 @@ FRigVMNewEditorMode::FRigVMNewEditorMode(const TSharedRef<FRigVMNewEditor>& InRi
 	TabFactories.RegisterFactory(MakeShared<FRigVMCompilerResultsTabSummoner>(InRigVMEditor));
 	TabFactories.RegisterFactory(MakeShared<FRigVMFindReferencesTabSummoner>(InRigVMEditor));
 
-	TabLayout = FTabManager::NewLayout("Standalone_RigVMEditMode_Layout_v1.6")
+	TabLayout = FTabManager::NewLayout("Standalone_RigVMEditMode_Layout_v1.7")
 		->AddArea
 		(
 			// Main application area
@@ -52,16 +53,16 @@ FRigVMNewEditorMode::FRigVMNewEditorMode(const TSharedRef<FRigVMNewEditor>& InRi
 						// Middle top
 						FTabManager::NewStack()
 						->SetSizeCoefficient(0.8f)
-						->AddTab("GraphEditor", ETabState::ClosedTab)
+						->AddTab(FRigVMGraphEditorTabSummoner::TabID(), ETabState::ClosedTab)
 					)
 					->Split
-						(
-							// Middle bottom - compiler results & find
-							FTabManager::NewStack()
-							->SetSizeCoefficient(0.2f)
-							->AddTab(FRigVMCompilerResultsTabSummoner::TabID(), ETabState::OpenedTab)
-							->AddTab(FRigVMFindReferencesTabSummoner::TabID(), ETabState::OpenedTab)
-						)
+					(
+						// Middle bottom - compiler results & find
+						FTabManager::NewStack()
+						->SetSizeCoefficient(0.2f)
+						->AddTab(FRigVMCompilerResultsTabSummoner::TabID(), ETabState::OpenedTab)
+						->AddTab(FRigVMFindReferencesTabSummoner::TabID(), ETabState::OpenedTab)
+					)
 				)
 				->Split
 				(
diff --git a/Engine/Plugins/Runtime/RigVM/Source/RigVMEditor/Private/Editor/RigVMVariableDetailCustomization.cpp b/Engine/Plugins/Runtime/RigVM/Source/RigVMEditor/Private/Editor/RigVMVariableDetailCustomization.cpp
--- a/Engine/Plugins/Runtime/RigVM/Source/RigVMEditor/Private/Editor/RigVMVariableDetailCustomization.cpp
+++ b/Engine/Plugins/Runtime/RigVM/Source/RigVMEditor/Private/Editor/RigVMVariableDetailCustomization.cpp
@@ -15,6 +15,7 @@
 #if WITH_RIGVMLEGACYEDITOR
 #include "Editor/RigVMLegacyEditor.h"
 #endif
+#include "InstancedPropertyBagStructureDataProvider.h"
 #include "RigVMModel/RigVMController.h"
 #include "Widgets/Input/SMultiLineEditableTextBox.h"
 #include "SPinTypeSelector.h"
@@ -23,6 +24,7 @@
 #include "Widgets/Input/SEditableTextBox.h"
 #include "Widgets/Input/SComboButton.h"
 #include "Widgets/SToolTip.h"
+#include "RigVMDeveloperTypeUtils.h"
 
 
 #define LOCTEXT_NAMESPACE "RigVMVariableDetailCustomization"
@@ -32,16 +34,21 @@ TSharedRef<IDetailCustomization> FRigVMVariableDetailCustomization::MakeInstance
 	const TArray<UObject*>* Objects = (InEditor.IsValid() ? InEditor->GetObjectsCurrentlyBeingEdited() : nullptr);
 	if (Objects && Objects->Num() == 1)
 	{
-		if (FRigVMAssetInterfacePtr Blueprint = (*Objects)[0])
+		if (FRigVMEditorAssetInterfacePtr Blueprint = (*Objects)[0])
 		{
 			return MakeShareable(new FRigVMVariableDetailCustomization(InEditor, Blueprint));
 		}
+		else if ((*Objects)[0]->Implements<URigVMRuntimeAssetInterface>())
+		{
+			TScriptInterface<IRigVMRuntimeAssetInterface> RuntimeAsset = (*Objects)[0];
+			return MakeShareable(new FRigVMVariableDetailCustomization(InEditor, RuntimeAsset->GetEditorOnlyData()));
+		}
 	}
 
-	return MakeShareable(new FRigVMVariableDetailCustomization((TSharedPtr<IRigVMEditor>)nullptr, FRigVMAssetInterfacePtr(nullptr)));
+	return MakeShareable(new FRigVMVariableDetailCustomization((TSharedPtr<IRigVMEditor>)nullptr, FRigVMEditorAssetInterfacePtr(nullptr)));
 }
 
-FRigVMVariableDetailCustomization::FRigVMVariableDetailCustomization(TSharedPtr<IRigVMEditor> InEditor, FRigVMAssetInterfacePtr Blueprint)
+FRigVMVariableDetailCustomization::FRigVMVariableDetailCustomization(TSharedPtr<IRigVMEditor> InEditor, FRigVMEditorAssetInterfacePtr Blueprint)
 		: EditorPtr(InEditor)
 		, BlueprintPtr(Blueprint)
 {}
@@ -107,12 +114,12 @@ void FRigVMVariableDetailCustomization::CustomizeDetails(IDetailLayoutBuilder& D
 
 			if (const URigVMEdGraphSchema* RigVMSchema = Cast<URigVMEdGraphSchema>(Schema))
 			{
-				FRigVMAssetInterfacePtr AssetPtr = BlueprintPtr.GetObject();
+				FRigVMEditorAssetInterfacePtr AssetPtr = BlueprintPtr.GetObject();
 				NameValidator = RigVMSchema->GetNameValidator(AssetPtr, CachedVariableName, nullptr);
 			}
 		}
 
-		FProperty* VariableProperty = CachedVariableProperty.Get();
+		const FProperty* VariableProperty = CachedVariableProperty.Get();
 
 		IDetailCategoryBuilder& Category = DetailLayout.EditCategory("Variable", LOCTEXT("VariableDetailsCategory", "Variable"));
 		const FSlateFontInfo DetailFontInfo = IDetailLayoutBuilder::GetDetailFont();
@@ -432,15 +439,12 @@ void FRigVMVariableDetailCustomization::CustomizeDetails(IDetailLayoutBuilder& D
 
 		if (TStrongObjectPtr<UObject> StrongObj = BlueprintPtr.GetWeakObjectPtr().Pin())
 		{
-			URigVMBlueprint* Blueprint = Cast<URigVMBlueprint>(StrongObj.Get());
-			if (!Blueprint)
-			{
-				checkf(false, TEXT("RigVM asset is not a blueprint"));
-			}
+			FRigVMEditorAssetInterfacePtr Asset = StrongObj.Get();
 			
 			// Add in default value editing for properties that can be edited, local properties cannot be edited
-			if (Blueprint && Blueprint->GeneratedClass != nullptr)
+			if (Asset->IsBlueprintAsset() && Asset->GetRigVMGeneratedClass() != nullptr)
 			{
+				URigVMBlueprint* Blueprint = Cast<URigVMBlueprint>(Asset.GetObject());
 				bool bVariableRenamed = false;
 				if (VariableProperty != nullptr)
 				{
@@ -475,57 +479,82 @@ void FRigVMVariableDetailCustomization::CustomizeDetails(IDetailLayoutBuilder& D
 						}
 					}
 				}
+			}
+			else 
+			{
+				VariableProperty = Asset->FindGeneratedPropertyByName(CachedVariableName);
+			}
 
-				// Find the class containing the variable
-				UClass* VariableClass = (VariableProperty ? VariableProperty->GetTypedOwner<UClass>() : nullptr);
+			// Find the class containing the variable
+			UClass* VariableClass = (VariableProperty ? VariableProperty->GetTypedOwner<UClass>() : nullptr);
 
-				FText ErrorMessage;
-				IDetailCategoryBuilder& DefaultValueCategory = DetailLayout.EditCategory(TEXT("DefaultValueCategory"), LOCTEXT("DefaultValueCategoryHeading", "Default Value"));
+			FText ErrorMessage;
+			IDetailCategoryBuilder& DefaultValueCategory = DetailLayout.EditCategory(TEXT("DefaultValueCategory"), LOCTEXT("DefaultValueCategoryHeading", "Default Value"));
 
-				if (VariableProperty == nullptr)
+			if (VariableProperty == nullptr)
+			{
+				if (Asset->GetAssetStatus() != RVMA_UpToDate)
 				{
-					if (Blueprint->Status != BS_UpToDate)
-					{
-						ErrorMessage = LOCTEXT("VariableMissing_DirtyBlueprint", "Please compile the blueprint");
-					}
-					else
-					{
-						ErrorMessage = LOCTEXT("VariableMissing_CleanBlueprint", "Failed to find variable property");
-					}
+					ErrorMessage = LOCTEXT("VariableMissing_DirtyBlueprint", "Please compile the blueprint");
+				}
+				else
+				{
+					ErrorMessage = LOCTEXT("VariableMissing_CleanBlueprint", "Failed to find variable property");
 				}
-				// Show the error message if something went wrong
-				if (!ErrorMessage.IsEmpty())
+			}
+			// Show the error message if something went wrong
+			if (!ErrorMessage.IsEmpty())
+			{
+				DefaultValueCategory.AddCustomRow( ErrorMessage )
+				[
+					SNew(STextBlock)
+					.ToolTipText(ErrorMessage)
+					.Text(ErrorMessage)
+					.Font(DetailFontInfo)
+				];
+			}
+			else 
+			{
+				TSharedPtr<IDetailsView> DetailsView;
+				if (Editor.IsValid())
 				{
-					DefaultValueCategory.AddCustomRow( ErrorMessage )
-					[
-						SNew(STextBlock)
-						.ToolTipText(ErrorMessage)
-						.Text(ErrorMessage)
-						.Font(DetailFontInfo)
-					];
+					DetailsView = Editor->GetRigVMInspector()->GetPropertyView();
 				}
-				else 
+
+				if (Asset->IsBlueprintAsset())
 				{
-					TSharedPtr<IDetailsView> DetailsView;
-					if (Editor.IsValid())
-					{
-						DetailsView = Editor->GetRigVMInspector()->GetPropertyView();
-					}
+					// Things are in order, show the property and allow it to be edited
+					TArray<UObject*> ObjectList;
+					ObjectList.Add(Asset->GetRigVMGeneratedClass()->GetDefaultObject());
+					IDetailPropertyRow* Row = DefaultValueCategory.AddExternalObjectProperty(ObjectList, VariableProperty->GetFName());
 
+				}
+				else
+				{
+					// Things are in order, show the property and allow it to be edited
+					if (FRigVMPropertyBag* InstancedPropertyBag = Asset->GetVariablesPropertyBag())
 					{
-						// Things are in order, show the property and allow it to be edited
-						TArray<UObject*> ObjectList;
-						ObjectList.Add(Blueprint->GeneratedClass->GetDefaultObject());
-						IDetailPropertyRow* Row = DefaultValueCategory.AddExternalObjectProperty(ObjectList, VariableProperty->GetFName());
-
-						// if (DetailsView.IsValid())
-						// {
-						// 	DetailsView->OnFinishedChangingProperties().AddSP(this, &FRigVMVariableDetailCustomization::OnFinishedChangingVariable);
-						// }
+						FStructView StructView = InstancedPropertyBag->GetMutableValue();
+						if (StructView.IsValid())
+						{
+							TSharedRef<FStructOnScope> StructOnScope = MakeShared<FStructOnScope>(
+								StructView.GetScriptStruct(),
+								StructView.GetMemory()
+								);
+							FAddPropertyParams Params;
+							Params.HideRootObjectNode(true);
+							Params.ForceShowProperty();
+							// Don't rely on the CachedVariableName because it is not sanitized (it allows spaces)
+							DefaultValueCategory.AddExternalStructureProperty(StructOnScope, VariableProperty->GetFName(), EPropertyLocation::Default, Params);
+						}
 					}
 				}
+
+				if (DetailsView.IsValid())
+				{
+					DetailsView->OnFinishedChangingProperties().AddSP(this, &FRigVMVariableDetailCustomization::OnFinishedChangingVariable);
+				}
 			}
-			
 		}
 	}
 }
@@ -539,7 +568,7 @@ void FRigVMVariableDetailCustomization::PopulateCategories()
 	{
 		return;
 	}
-	FRigVMAssetInterfacePtr Interface = Blueprint.Get();
+	FRigVMEditorAssetInterfacePtr Interface = Blueprint.Get();
 	
 	auto IsNewCategorySource = [&](const FText& NewCategory)
 	{
@@ -578,6 +607,35 @@ void FRigVMVariableDetailCustomization::PopulateCategories()
 	});
 }
 
+void FRigVMVariableDetailCustomization::OnFinishedChangingVariable(const FPropertyChangedEvent& PropertyChangedEvent)
+{
+	TSharedPtr<IRigVMEditor> Editor = EditorPtr.Pin();
+	if (!Editor)
+	{
+		return;
+	}
+
+	TStrongObjectPtr<UObject> Blueprint = BlueprintPtr.GetWeakObjectPtr().Pin();
+	if (!Blueprint)
+	{
+		return;
+	}
+
+	if (FRigVMEditorAssetInterfacePtr Asset = Blueprint.Get())
+	{
+		Asset->MarkAssetAsModified();
+		Asset->GetRuntimeAssetInterface()->SetAssetStatus(ERigVMAssetStatus::RVMA_Dirty);
+
+		// If the changed property is the root property, autocompile
+		// Otherwise wait until the user forces a compile, we do this in order to allow multi-parameter editing of structs like FVector (using tab)
+		if (PropertyChangedEvent.Property == PropertyChangedEvent.MemberProperty)
+		{
+			Asset->RecompileVM();
+		}
+	}
+	
+}
+
 void FRigVMVariableDetailCustomization::OnPostEditorRefresh()
 {
 	FName VariableName = GetVariableName();
@@ -624,7 +682,7 @@ void FRigVMVariableDetailCustomization::OnVarNameCommitted(const FText& InNewNam
 	ERigVMExplorerElementType::Type SelectedType = Explorer->GetSelectedType();
 	if (SelectedType == ERigVMExplorerElementType::Variable)
 	{
-		FRigVMAssetInterfacePtr Interface = Blueprint.Get();
+		FRigVMEditorAssetInterfacePtr Interface = Blueprint.Get();
 		Blueprint->Modify();
 		Interface->RenameMemberVariable(CachedVariableName, *InNewName.ToString());
 	}
@@ -644,7 +702,7 @@ FEdGraphPinType FRigVMVariableDetailCustomization::OnGetVarType() const
 	{
 		return Type;
 	}
-	FRigVMAssetInterfacePtr Interface = Blueprint.Get();
+	FRigVMEditorAssetInterfacePtr Interface = Blueprint.Get();
 
 	TSharedPtr<SRigVMEditorGraphExplorer> Explorer = Editor->GetGraphExplorerWidget();
 	if (!Explorer)
@@ -659,7 +717,7 @@ FEdGraphPinType FRigVMVariableDetailCustomization::OnGetVarType() const
 		{
 			if (Var.Name == CachedVariableName)
 			{
-				return Var.ToPinType();
+				return RigVMTypeUtils::PinTypeFromRigVMVariableDescription(Var);
 			}
 		}
 		return Type;
@@ -680,7 +738,7 @@ void FRigVMVariableDetailCustomization::OnVarTypeChanged(const FEdGraphPinType&
 	{
 		return;
 	}
-	FRigVMAssetInterfacePtr Interface = Blueprint.Get();
+	FRigVMEditorAssetInterfacePtr Interface = Blueprint.Get();
 
 	TSharedPtr<SRigVMEditorGraphExplorer> Explorer = Editor->GetGraphExplorerWidget();
 	if (!Explorer)
@@ -752,7 +810,7 @@ FText FRigVMVariableDetailCustomization::OnGetTooltipText() const
 	{
 		return ToolTip;
 	}
-	FRigVMAssetInterfacePtr Interface = Blueprint.Get();
+	FRigVMEditorAssetInterfacePtr Interface = Blueprint.Get();
 
 	TSharedPtr<SRigVMEditorGraphExplorer> Explorer = Editor->GetGraphExplorerWidget();
 	if (!Explorer)
@@ -790,7 +848,7 @@ void FRigVMVariableDetailCustomization::OnTooltipTextCommitted(const FText& NewT
 	{
 		return;
 	}
-	FRigVMAssetInterfacePtr Interface = Blueprint.Get();
+	FRigVMEditorAssetInterfacePtr Interface = Blueprint.Get();
 
 	TSharedPtr<SRigVMEditorGraphExplorer> Explorer = Editor->GetGraphExplorerWidget();
 	if (!Explorer)
@@ -840,7 +898,7 @@ FText FRigVMVariableDetailCustomization::OnGetCategoryText() const
 	{
 		return FText();
 	}
-	FRigVMAssetInterfacePtr Interface = Blueprint.Get();
+	FRigVMEditorAssetInterfacePtr Interface = Blueprint.Get();
 
 	if (CachedVariableName != NAME_None)
 	{
@@ -874,7 +932,7 @@ void FRigVMVariableDetailCustomization::OnCategoryTextCommitted(const FText& New
 	{
 		return;
 	}
-	FRigVMAssetInterfacePtr Interface = Blueprint.Get();
+	FRigVMEditorAssetInterfacePtr Interface = Blueprint.Get();
 	
 	if (InTextCommit == ETextCommit::OnEnter || InTextCommit == ETextCommit::OnUserMovedFocus)
 	{
@@ -898,7 +956,7 @@ void FRigVMVariableDetailCustomization::OnCategorySelectionChanged(TSharedPtr<FT
 	{
 		return;
 	}
-	FRigVMAssetInterfacePtr Interface = Blueprint.Get();
+	FRigVMEditorAssetInterfacePtr Interface = Blueprint.Get();
 	
 	if (ProposedSelection.IsValid() && CachedVariableName != NAME_None)
 	{
@@ -911,7 +969,7 @@ void FRigVMVariableDetailCustomization::OnCategorySelectionChanged(TSharedPtr<FT
 
 ECheckBoxState FRigVMVariableDetailCustomization::OnGetExposedToSpawnCheckboxState() const
 {
-	FProperty* Property = CachedVariableProperty.Get();
+	const FProperty* Property = CachedVariableProperty.Get();
 	if (Property)
 	{
 		return (Property && Property->GetBoolMetaData(FBlueprintMetadata::MD_ExposeOnSpawn) != false) ? ECheckBoxState::Checked : ECheckBoxState::Unchecked;
@@ -926,7 +984,7 @@ void FRigVMVariableDetailCustomization::OnExposedToSpawnChanged(ECheckBoxState I
 	{
 		return;
 	}
-	FRigVMAssetInterfacePtr Interface = Blueprint.Get();
+	FRigVMEditorAssetInterfacePtr Interface = Blueprint.Get();
 	
 	if (CachedVariableName != NAME_None)
 	{
@@ -937,7 +995,7 @@ void FRigVMVariableDetailCustomization::OnExposedToSpawnChanged(ECheckBoxState I
 
 ECheckBoxState FRigVMVariableDetailCustomization::OnGetPrivateCheckboxState() const
 {
-	FProperty* Property = CachedVariableProperty.Get();
+	const FProperty* Property = CachedVariableProperty.Get();
 	if (Property)
 	{
 		return (Property && Property->GetBoolMetaData(FBlueprintMetadata::MD_Private) != false) ? ECheckBoxState::Checked : ECheckBoxState::Unchecked;
@@ -952,7 +1010,7 @@ void FRigVMVariableDetailCustomization::OnPrivateChanged(ECheckBoxState InNewSta
 	{
 		return;
 	}
-	FRigVMAssetInterfacePtr Interface = Blueprint.Get();
+	FRigVMEditorAssetInterfacePtr Interface = Blueprint.Get();
 	
 	if (CachedVariableName != NAME_None)
 	{
@@ -963,7 +1021,7 @@ void FRigVMVariableDetailCustomization::OnPrivateChanged(ECheckBoxState InNewSta
 
 ECheckBoxState FRigVMVariableDetailCustomization::OnGetExposedToCinematicsCheckboxState() const
 {
-	FProperty* Property = CachedVariableProperty.Get();
+	const FProperty* Property = CachedVariableProperty.Get();
 	if (Property)
 	{
 		return Property && Property->HasAnyPropertyFlags(CPF_Interp) ? ECheckBoxState::Checked : ECheckBoxState::Unchecked;
@@ -978,7 +1036,7 @@ void FRigVMVariableDetailCustomization::OnExposedToCinematicsChanged(ECheckBoxSt
 	{
 		return;
 	}
-	FRigVMAssetInterfacePtr Interface = Blueprint.Get();
+	FRigVMEditorAssetInterfacePtr Interface = Blueprint.Get();
 	
 	// Toggle the flag on the blueprint's version of the variable description, based on state
 	const bool bExposeToCinematics = (InNewState == ECheckBoxState::Checked);
@@ -997,7 +1055,7 @@ FText FRigVMVariableDetailCustomization::OnGetMetaKeyValue(FName Key) const
 	{
 		return FText();
 	}
-	FRigVMAssetInterfacePtr Interface = Blueprint.Get();
+	FRigVMEditorAssetInterfacePtr Interface = Blueprint.Get();
 	
 	if (CachedVariableName != NAME_None)
 	{
@@ -1028,7 +1086,7 @@ void FRigVMVariableDetailCustomization::OnMetaKeyValueChanged(const FText& NewMi
 
 EVisibility FRigVMVariableDetailCustomization::RangeVisibility() const
 {
-	FProperty* VariableProperty = CachedVariableProperty.Get();
+	const FProperty* VariableProperty = CachedVariableProperty.Get();
 	if (VariableProperty)
 	{
 		const bool bIsInteger = VariableProperty->IsA(FIntProperty::StaticClass());
diff --git a/Engine/Plugins/Runtime/RigVM/Source/RigVMEditor/Private/Editor/RigVMVariantDetailCustomization.cpp b/Engine/Plugins/Runtime/RigVM/Source/RigVMEditor/Private/Editor/RigVMVariantDetailCustomization.cpp
--- a/Engine/Plugins/Runtime/RigVM/Source/RigVMEditor/Private/Editor/RigVMVariantDetailCustomization.cpp
+++ b/Engine/Plugins/Runtime/RigVM/Source/RigVMEditor/Private/Editor/RigVMVariantDetailCustomization.cpp
@@ -30,11 +30,17 @@ void FRigVMVariantDetailCustomization::CustomizeHeader(TSharedRef<IPropertyHandl
 
 	for (UObject* Object : Objects)
 	{
-		if (Object->Implements<URigVMAssetInterface>())
+		if (Object->Implements<URigVMEditorAssetInterface>())
 		{
 			BlueprintBeingCustomized = Object;
 			break;
 		}
+		else if (Object->Implements<URigVMRuntimeAssetInterface>())
+		{
+			TScriptInterface<IRigVMRuntimeAssetInterface> RuntimeInterface = Object;
+			BlueprintBeingCustomized = RuntimeInterface->GetEditorOnlyData();
+			break;
+		}
 	}
 
 	FRigVMVariantWidgetContext VariantContext;
diff --git a/Engine/Plugins/Runtime/RigVM/Source/RigVMEditor/Private/Editor/SRigVMDetailsInspector.cpp b/Engine/Plugins/Runtime/RigVM/Source/RigVMEditor/Private/Editor/SRigVMDetailsInspector.cpp
--- a/Engine/Plugins/Runtime/RigVM/Source/RigVMEditor/Private/Editor/SRigVMDetailsInspector.cpp
+++ b/Engine/Plugins/Runtime/RigVM/Source/RigVMEditor/Private/Editor/SRigVMDetailsInspector.cpp
@@ -12,6 +12,7 @@
 #include "Editor/RigVMCommentNodeDetailCustomization.h"
 #include "Framework/Application/SlateApplication.h"
 #include "K2Node.h"
+#include "Editor/RigVMClassSettingsDetailCustomization.h"
 
 class FNotifyHook;
 class FStructOnScope;
@@ -133,10 +134,11 @@ const TArray< TWeakObjectPtr<UObject> >& SRigVMDetailsInspector::GetSelectedObje
 	return SelectedObjects;
 }
 
-void SRigVMDetailsInspector::OnEditorClose(const IRigVMEditor* RigVMEditorBase, FRigVMAssetInterfacePtr RigVMBlueprint)
+void SRigVMDetailsInspector::OnEditorClose(const IRigVMEditor* RigVMEditorBase, FRigVMEditorAssetInterfacePtr RigVMBlueprint)
 {
 	PropertyView->UnregisterInstancedCustomPropertyLayout(URigVMEdGraph::StaticClass());
 	PropertyView->UnregisterInstancedCustomPropertyLayout(UPropertyWrapper::StaticClass());
+	PropertyView->UnregisterInstancedCustomPropertyLayout(URigVMRuntimeAsset::StaticClass());
 	PropertyView->UnregisterInstancedCustomPropertyLayout(URigVMCommentNode::StaticClass());
 }
 
@@ -248,6 +250,9 @@ void SRigVMDetailsInspector::Construct(const FArguments& InArgs)
 	FOnGetDetailCustomizationInstance LayoutVariableDetails = FOnGetDetailCustomizationInstance::CreateStatic(&FRigVMVariableDetailCustomization::MakeInstance, Editor);
 	PropertyView->RegisterInstancedCustomPropertyLayout(UPropertyWrapper::StaticClass(), LayoutVariableDetails);
 
+	FOnGetDetailCustomizationInstance LayoutClassSettingsDetails = FOnGetDetailCustomizationInstance::CreateStatic(&FRigVMClassSettingsDetailCustomization::MakeInstance, Editor);
+	PropertyView->RegisterInstancedCustomPropertyLayout(URigVMRuntimeAsset::StaticClass(), LayoutClassSettingsDetails);
+
 	FOnGetDetailCustomizationInstance CommentNodeDetails = FOnGetDetailCustomizationInstance::CreateStatic(&FRigVMCommentNodeDetailCustomization::MakeInstance);
 	PropertyView->RegisterInstancedCustomPropertyLayout(URigVMCommentNode::StaticClass(), CommentNodeDetails);
 
@@ -509,7 +514,7 @@ bool SRigVMDetailsInspector::IsPropertyVisible( const FPropertyAndParent& Proper
 		// this is necessary because of Instanced objects, which will have a different owning class yet are conceptually contained in this blueprint
 		bool bVariableAddedInCurrentBlueprint = false;
 		TSharedPtr<IRigVMEditor> Editor = WeakEditor.Pin();
-		FRigVMAssetInterfacePtr Blueprint = Editor.IsValid() ? Editor->GetRigVMAssetInterface() : nullptr;
+		FRigVMEditorAssetInterfacePtr Blueprint = Editor.IsValid() ? Editor->GetRigVMAssetInterface() : nullptr;
 
 		auto WasAddedInThisBlueprint = [Blueprint](const FProperty* Property)
 		{
diff --git a/Engine/Plugins/Runtime/RigVM/Source/RigVMEditor/Private/Editor/SRigVMGraphTitleBar.cpp b/Engine/Plugins/Runtime/RigVM/Source/RigVMEditor/Private/Editor/SRigVMGraphTitleBar.cpp
--- a/Engine/Plugins/Runtime/RigVM/Source/RigVMEditor/Private/Editor/SRigVMGraphTitleBar.cpp
+++ b/Engine/Plugins/Runtime/RigVM/Source/RigVMEditor/Private/Editor/SRigVMGraphTitleBar.cpp
@@ -217,7 +217,7 @@ void SRigVMGraphTitleBar::Construct( const FArguments& InArgs )
 	RebuildBreadcrumbTrail();
 	BreadcrumbTrailScrollBox->ScrollToEnd();
 
-	FRigVMAssetInterfacePtr BlueprintObj = FRigVMBlueprintUtils::FindAssetForGraph(this->EdGraphObj);
+	FRigVMEditorAssetInterfacePtr BlueprintObj = FRigVMBlueprintUtils::FindAssetForGraph(this->EdGraphObj);
 	if (BlueprintObj)
 	{
 		BlueprintTitle = FText::FromString(BlueprintObj->GetObject()->GetName());
@@ -287,7 +287,7 @@ void SRigVMGraphTitleBar::Refresh()
 	// Refresh UI on request
 	if (EdGraphObj)
 	{
-		if (FRigVMAssetInterfacePtr BlueprintObj = FRigVMBlueprintUtils::FindAssetForGraph(this->EdGraphObj))
+		if (FRigVMEditorAssetInterfacePtr BlueprintObj = FRigVMBlueprintUtils::FindAssetForGraph(this->EdGraphObj))
 		{
 			BlueprintTitle = FText::FromString(BlueprintObj->GetObject()->GetName());
 			RebuildBreadcrumbTrail();
diff --git a/Engine/Plugins/Runtime/RigVM/Source/RigVMEditor/Private/RigVMEditorBlueprintLibrary.cpp b/Engine/Plugins/Runtime/RigVM/Source/RigVMEditor/Private/RigVMEditorBlueprintLibrary.cpp
--- a/Engine/Plugins/Runtime/RigVM/Source/RigVMEditor/Private/RigVMEditorBlueprintLibrary.cpp
+++ b/Engine/Plugins/Runtime/RigVM/Source/RigVMEditor/Private/RigVMEditorBlueprintLibrary.cpp
@@ -1,6 +1,7 @@
 // Copyright Epic Games, Inc. All Rights Reserved.
 
 #include "RigVMEditorBlueprintLibrary.h"
+#include "RigVMDeveloperModule.h"
 #include "AssetRegistry/AssetRegistryModule.h"
 #include "Misc/ScopedSlowTask.h"
 #include "UObject/UObjectIterator.h"
@@ -312,7 +313,7 @@ TArray<FAssetData> URigVMEditorBlueprintLibrary::GetAssetsWithFilter(
 				continue;
 			}
 
-			if (Class->ImplementsInterface(URigVMAssetInterface::StaticClass()))
+			if (Class->ImplementsInterface(URigVMEditorAssetInterface::StaticClass()))
 			{
 				Filter.ClassPaths.Add(Class->GetClassPathName());
 			}
