Control Rig: Split IRigVMAssetInterface into IRigVMEditorAssetInterface and IRigVMRuntimeAssetInterface.

diff --git a/Engine/Plugins/Runtime/RigVM/Source/RigVMEditor/Private/RigVMEditorModule.cpp b/Engine/Plugins/Runtime/RigVM/Source/RigVMEditor/Private/RigVMEditorModule.cpp
--- a/Engine/Plugins/Runtime/RigVM/Source/RigVMEditor/Private/RigVMEditorModule.cpp
+++ b/Engine/Plugins/Runtime/RigVM/Source/RigVMEditor/Private/RigVMEditorModule.cpp
@@ -58,7 +58,7 @@ DEFINE_LOG_CATEGORY(LogRigVMEditor);
 IMPLEMENT_MODULE(FRigVMEditorModule, RigVMEditor)
 
 TAutoConsoleVariable<bool> CVarRigVMUseNewEditor(TEXT("RigVM.UseNewEditor"), 0, TEXT("Use editor which does not inherit from BlueprintEditor."));
-TAutoConsoleVariable<bool> CVarRigVMUseDualEditor(TEXT("RigVM.UseDualEditor"), 0, TEXT("Use both legacy and new editor for A/B comparison."));
+TAutoConsoleVariable<bool> CVarRigVMBlueprintIndependentAssets(TEXT("RigVM.BlueprintIndependentAssets"), 0, TEXT("Create assets that do not inherit from blueprint."));
 
 UE::RigVMEditor::FRigVMEditorOnRequestFindNodeReferencesDelegate FRigVMEditorModule::OnRequestFindNodeReferences;
 
@@ -143,9 +143,9 @@ void FRigVMEditorModule::StartupModule()
 									FPaths::Split(PackageName, PathName, ObjectName, Extension);
 
 									UObject* DuplicateAsset = AssetToolsModule.Get().DuplicateAsset(ObjectName, PathName, SelectedObject);
-									if (IRigVMAssetInterface* DuplicateBlueprint = Cast<IRigVMAssetInterface>(DuplicateAsset))
+									if (IRigVMEditorAssetInterface* DuplicateBlueprint = Cast<IRigVMEditorAssetInterface>(DuplicateAsset))
 									{
-										DuplicateBlueprint->GetAssetVariant() = Cast<IRigVMAssetInterface>(SelectedObject)->GetAssetVariant();
+										DuplicateBlueprint->GetAssetVariant() = Cast<IRigVMEditorAssetInterface>(SelectedObject)->GetAssetVariant();
 									}
 								}
 							}));
@@ -171,10 +171,10 @@ void FRigVMEditorModule::StartupModuleCommon()
 	MenuExtensibilityManager = MakeShareable(new FExtensibilityManager);
     ToolBarExtensibilityManager = MakeShareable(new FExtensibilityManager);
 
-	const UClass* BlueprintClass = GetRigVMAssetClass();
 	const UClass* EdGraphSchemaClass = GetRigVMEdGraphSchemaClass();
 	const URigVMEdGraphSchema* SchemaCDO = CastChecked<URigVMEdGraphSchema>(EdGraphSchemaClass->GetDefaultObject());
 #if WITH_RIGVMLEGACYEDITOR
+	const UClass* BlueprintClass = GetRigVMLegacyBlueprintClass();
 	FBlueprintEditorModule& BlueprintEditorModule = FModuleManager::LoadModuleChecked<FBlueprintEditorModule>("Kismet");
 	BlueprintEditorModule.RegisterGraphCustomization(SchemaCDO, FOnGetGraphCustomizationInstance::CreateStatic(&FRigVMGraphDetailCustomization::MakeLegacyInstance, BlueprintClass));
 #endif
@@ -236,11 +236,45 @@ void FRigVMEditorModule::ShutdownModuleCommon()
 	FKismetEditorUtilities::UnregisterAutoBlueprintNodeCreation(this);
 }
 
-void FRigVMEditorModule::GetTypeActions(FRigVMAssetInterfacePtr RigVMBlueprint, FBlueprintActionDatabaseRegistrar& ActionRegistrar)
+void FRigVMEditorModule::GetTypeActions(FRigVMEditorAssetInterfacePtr RigVMBlueprint, FBlueprintActionDatabaseRegistrar& ActionRegistrar)
+{
+	// actions get registered under specific object-keys; the idea is that 
+	// actions might have to be updated (or deleted) if their object-key is  
+	// mutated (or removed)... here we use the class (so if the class 
+	// type disappears, then the action should go with it)
+	UClass* ActionKey = RigVMBlueprint->GetObject()->GetClass();
+	// to keep from needlessly instantiating a UBlueprintNodeSpawner, first   
+	// check to make sure that the registrar is looking for actions of this type
+	// (could be regenerating actions for a specific asset, and therefore the 
+	// registrar would only accept actions corresponding to that asset)
+	if (!ActionRegistrar.IsOpenForRegistration(ActionKey))
+	{
+		return;
+	}
+
+	TArray<TSharedRef<URigVMEdGraphNodeSpawner>> NodeSpawners;
+	GetTypeActions(RigVMBlueprint, NodeSpawners);
+
+	for (TSharedRef<URigVMEdGraphNodeSpawner>& NodeSpawner : NodeSpawners)
+	{
+		URigVMEdGraphNodeBlueprintSpawner* BlueprintSpawner = URigVMEdGraphNodeBlueprintSpawner::CreateFromRigVMSpawner(NodeSpawner);
+		ActionRegistrar.AddBlueprintAction(ActionKey, BlueprintSpawner);
+	}
+}
+
+void FRigVMEditorModule::GetTypeActions(FRigVMEditorAssetInterfacePtr RigVMBlueprint, TArray<TSharedRef<URigVMEdGraphNodeSpawner>>& OutGraphNodeSpawners)
 {
 	// Only register actions for ourselves
 	UClass* InClass = RigVMBlueprint.GetObject()->GetClass();
+	if (TScriptInterface<IRigVMRuntimeAssetInterface> RuntimeAsset = RigVMBlueprint->GetRuntimeAssetInterface())
+	{
+		InClass = RuntimeAsset.GetObject()->GetClass();
+	}
 	UClass* BlueprintClass = GetRigVMAssetClass();
+	if (RigVMBlueprint->IsBlueprintAsset() && CVarRigVMBlueprintIndependentAssets->GetBool())
+	{
+		BlueprintClass = GetRigVMLegacyBlueprintClass();
+	}
 
 	if (BlueprintClass->GetDefaultObject()->IsA<UInterface>())
 	{
@@ -261,14 +295,6 @@ void FRigVMEditorModule::GetTypeActions(FRigVMAssetInterfacePtr RigVMBlueprint,
 	// mutated (or removed)... here we use the class (so if the class 
 	// type disappears, then the action should go with it)
 	UClass* ActionKey = RigVMBlueprint->GetObject()->GetClass();
-	// to keep from needlessly instantiating a UBlueprintNodeSpawner, first   
-	// check to make sure that the registrar is looking for actions of this type
-	// (could be regenerating actions for a specific asset, and therefore the 
-	// registrar would only accept actions corresponding to that asset)
-	if (!ActionRegistrar.IsOpenForRegistration(ActionKey))
-	{
-		return;
-	}
 
 	FRigVMRegistry& Registry = FRigVMRegistry::Get();
 	Registry.RefreshEngineTypes();
@@ -298,14 +324,15 @@ void FRigVMEditorModule::GetTypeActions(FRigVMAssetInterfacePtr RigVMBlueprint,
 			continue;
 		}
 
+		ensure (!Template.GetNotation().ToString().Contains(TEXT("AnimNext")));
+
 		FText NodeCategory = FText::FromString(Template.GetCategory());
 		FText MenuDesc = FText::FromName(Template.GetName());
 		FText ToolTip = Template.GetTooltipText();
 
 		TSharedRef<URigVMEdGraphTemplateNodeSpawner> NodeSpawner = MakeShared<URigVMEdGraphTemplateNodeSpawner>(Template.GetNotation(), MenuDesc, NodeCategory, ToolTip);
 		NodeSpawner->SetRelatedBlueprintClass(BlueprintClass);
-		URigVMEdGraphNodeBlueprintSpawner* BlueprintSpawner = URigVMEdGraphNodeBlueprintSpawner::CreateFromRigVMSpawner(NodeSpawner);
-		ActionRegistrar.AddBlueprintAction(ActionKey, BlueprintSpawner);
+		OutGraphNodeSpawners.Add(NodeSpawner);
 	};
 
 	for (const FRigVMDispatchFactory* Factory : Registry.GetFactories())
@@ -327,8 +354,7 @@ void FRigVMEditorModule::GetTypeActions(FRigVMAssetInterfacePtr RigVMBlueprint,
 
 		TSharedRef<URigVMEdGraphTemplateNodeSpawner> NodeSpawner = MakeShared<URigVMEdGraphTemplateNodeSpawner>(Template->GetNotation(), MenuDesc, NodeCategory, ToolTip);
 		NodeSpawner->SetRelatedBlueprintClass(BlueprintClass);
-		URigVMEdGraphNodeBlueprintSpawner* BlueprintSpawner = URigVMEdGraphNodeBlueprintSpawner::CreateFromRigVMSpawner(NodeSpawner);
-		ActionRegistrar.AddBlueprintAction(ActionKey, BlueprintSpawner);
+		OutGraphNodeSpawners.Add(NodeSpawner);
 	};
 
 	// Add all rig units
@@ -376,8 +402,7 @@ void FRigVMEditorModule::GetTypeActions(FRigVMAssetInterfacePtr RigVMBlueprint,
 
 		TSharedRef<URigVMEdGraphUnitNodeSpawner> NodeSpawner = MakeShared<URigVMEdGraphUnitNodeSpawner>(Struct, Function.GetMethodName(), MenuDesc, NodeCategory, ToolTip);
 		NodeSpawner->SetRelatedBlueprintClass(BlueprintClass);
-		URigVMEdGraphNodeBlueprintSpawner* BlueprintSpawner = URigVMEdGraphNodeBlueprintSpawner::CreateFromRigVMSpawner(NodeSpawner);
-		ActionRegistrar.AddBlueprintAction(ActionKey, BlueprintSpawner);
+		OutGraphNodeSpawners.Add(NodeSpawner);
 	};
 
 	for (TObjectIterator<UEnum> EnumIt; EnumIt; ++EnumIt)
@@ -405,8 +430,7 @@ void FRigVMEditorModule::GetTypeActions(FRigVMAssetInterfacePtr RigVMBlueprint,
 
 		TSharedRef<URigVMEdGraphEnumNodeSpawner> NodeSpawner = MakeShared<URigVMEdGraphEnumNodeSpawner>(EnumToConsider, MenuDesc, NodeCategory, ToolTip);
 		NodeSpawner->SetRelatedBlueprintClass(BlueprintClass);
-		URigVMEdGraphNodeBlueprintSpawner* BlueprintSpawner = URigVMEdGraphNodeBlueprintSpawner::CreateFromRigVMSpawner(NodeSpawner);
-		ActionRegistrar.AddBlueprintAction(ActionKey, BlueprintSpawner);
+		OutGraphNodeSpawners.Add(NodeSpawner);
 	}
 
 	FArrayProperty* PublicGraphFunctionsProperty = CastField<FArrayProperty>(URigVMBlueprint::StaticClass()->FindPropertyByName(TEXT("PublicGraphFunctions")));
@@ -438,6 +462,10 @@ void FRigVMEditorModule::GetTypeActions(FRigVMAssetInterfacePtr RigVMBlueprint,
 			}
 		};
 		AddClassToFilter(RigVMBlueprint->GetObject()->GetClass());
+		if (TScriptInterface<IRigVMRuntimeAssetInterface> RuntimeAsset = RigVMBlueprint->GetRuntimeAssetInterface())
+		{
+			AddClassToFilter(RuntimeAsset.GetObject()->GetClass());
+		}
 		AddClassToFilter(RigVMBlueprint->GetRigVMGeneratedClassPrototype());
 
 		AssetRegistryModule.Get().GetAssets(ControlRigAssetFilter, ControlRigAssetDatas);
@@ -462,9 +490,13 @@ void FRigVMEditorModule::GetTypeActions(FRigVMAssetInterfacePtr RigVMBlueprint,
 			if (ControlRigAssetData.IsAssetLoaded())
 			{
 				UObject* AssetObject = ControlRigAssetData.GetAsset();
-				if (IRigVMAssetInterface* Asset = Cast<IRigVMAssetInterface>(AssetObject))
+				if (IRigVMEditorAssetInterface* EditorAsset = Cast<IRigVMEditorAssetInterface>(AssetObject))
+				{
+					PublicFunctions = EditorAsset->GetPublicGraphFunctions();
+				}
+				if (IRigVMRuntimeAssetInterface* RuntimeAsset = Cast<IRigVMRuntimeAssetInterface>(AssetObject))
 				{
-					PublicFunctions = Asset->GetPublicGraphFunctions();
+					PublicFunctions = TScriptInterface<IRigVMEditorAssetInterface>(RuntimeAsset->GetEditorOnlyData())->GetPublicGraphFunctions();
 				}
 				else if(URigVMBlueprintGeneratedClass* GeneratedClass = Cast<URigVMBlueprintGeneratedClass>(AssetObject))
 				{
@@ -508,8 +540,7 @@ void FRigVMEditorModule::GetTypeActions(FRigVMAssetInterfacePtr RigVMBlueprint,
 					{
 						TSharedRef<URigVMEdGraphFunctionRefNodeSpawner> NodeSpawner = MakeShared<URigVMEdGraphFunctionRefNodeSpawner>(ControlRigAssetData, PublicFunction);
 						NodeSpawner->SetRelatedBlueprintClass(BlueprintClass);
-						URigVMEdGraphNodeBlueprintSpawner* BlueprintSpawner = URigVMEdGraphNodeBlueprintSpawner::CreateFromRigVMSpawner(NodeSpawner);
-						ActionRegistrar.AddBlueprintAction(ActionKey, BlueprintSpawner);
+						OutGraphNodeSpawners.Add(NodeSpawner);
 					}
 				}
 
@@ -534,138 +565,168 @@ void FRigVMEditorModule::GetTypeActions(FRigVMAssetInterfacePtr RigVMBlueprint,
 			{
 				if (PublicFunction.LibraryPointer.IsValid())
 				{
-					TSharedPtr<URigVMEdGraphFunctionRefNodeSpawner> NodeSpawner = MakeShared<URigVMEdGraphFunctionRefNodeSpawner>(ControlRigAssetData, PublicFunction);
+					TSharedRef<URigVMEdGraphFunctionRefNodeSpawner> NodeSpawner = MakeShared<URigVMEdGraphFunctionRefNodeSpawner>(ControlRigAssetData, PublicFunction);
 					NodeSpawner->SetRelatedBlueprintClass(BlueprintClass);
-					URigVMEdGraphNodeBlueprintSpawner* BlueprintSpawner = URigVMEdGraphNodeBlueprintSpawner::CreateFromRigVMSpawner(NodeSpawner);
-					ActionRegistrar.AddBlueprintAction(ActionKey, BlueprintSpawner);
+					OutGraphNodeSpawners.Add(NodeSpawner);
 				}
 			}
 		}
 	}
 }
 
-void FRigVMEditorModule::GetInstanceActions(FRigVMAssetInterfacePtr RigVMBlueprint, FBlueprintActionDatabaseRegistrar& ActionRegistrar)
+void FRigVMEditorModule::GetInstanceActions(FRigVMEditorAssetInterfacePtr RigVMBlueprint, FBlueprintActionDatabaseRegistrar& ActionRegistrar)
+{
+	if (UClass* GeneratedClass = RigVMBlueprint->GetRigVMGeneratedClass())
+	{
+		if (!ActionRegistrar.IsOpenForRegistration(GeneratedClass))
+		{
+			return;
+		}
+
+		TArray<TSharedRef<URigVMEdGraphNodeSpawner>> NodeSpawners;
+		GetInstanceActions(RigVMBlueprint, NodeSpawners);
+
+		for (TSharedRef<URigVMEdGraphNodeSpawner>& NodeSpawner : NodeSpawners)
+		{
+			URigVMEdGraphNodeBlueprintSpawner* BlueprintNodeSpawner = URigVMEdGraphNodeBlueprintSpawner::CreateFromRigVMSpawner(NodeSpawner);
+			ActionRegistrar.AddBlueprintAction(GeneratedClass, BlueprintNodeSpawner);
+		}
+	}
+}
+
+void FRigVMEditorModule::GetInstanceActions(FRigVMEditorAssetInterfacePtr RigVMBlueprint, TArray<TSharedRef<URigVMEdGraphNodeSpawner>>& OutGraphNodeSpawners)
 {
 	// Only register actions for ourselves
+	UClass* InClass = RigVMBlueprint.GetObject()->GetClass();
+	if (TScriptInterface<IRigVMRuntimeAssetInterface> RuntimeAsset = RigVMBlueprint->GetRuntimeAssetInterface())
+	{
+		InClass = RuntimeAsset.GetObject()->GetClass();
+	}
 	UClass* BlueprintClass = GetRigVMAssetClass();
-	if(RigVMBlueprint->GetObject()->GetClass() != BlueprintClass)
+	if (RigVMBlueprint->IsBlueprintAsset() && CVarRigVMBlueprintIndependentAssets->GetBool())
 	{
-		return;
+		BlueprintClass = GetRigVMLegacyBlueprintClass();
 	}
 
-	if (UClass* GeneratedClass = RigVMBlueprint->GetRigVMGeneratedClass())
+	if (BlueprintClass->GetDefaultObject()->IsA<UInterface>())
 	{
-		if (!ActionRegistrar.IsOpenForRegistration(GeneratedClass))
+		if (!InClass->ImplementsInterface(BlueprintClass))
 		{
 			return;
 		}
-		static const FText NodeCategory = LOCTEXT("Variables", "Variables");
-
-		TArray<FRigVMGraphVariableDescription> Variables = RigVMBlueprint->GetAssetVariables();
-		for (const FRigVMGraphVariableDescription& Variable : Variables)
-		{
-			const FRigVMExternalVariable ExternalVariable = Variable.ToExternalVariable();
-			FText MenuDesc = FText::FromName(ExternalVariable.Name);
-			FText ToolTip = FText::FromString(FString::Printf(TEXT("Get the value of variable %s"), *ExternalVariable.Name.ToString()));
-			TSharedPtr<URigVMEdGraphVariableNodeSpawner> GetNodeSpawner = MakeShared<URigVMEdGraphVariableNodeSpawner>(RigVMBlueprint, ExternalVariable, true, MenuDesc, NodeCategory, ToolTip);
-			GetNodeSpawner->SetRelatedBlueprintClass(BlueprintClass);
-			URigVMEdGraphNodeBlueprintSpawner* BlueprintGetSpawner = URigVMEdGraphNodeBlueprintSpawner::CreateFromRigVMSpawner(GetNodeSpawner);
-			ActionRegistrar.AddBlueprintAction(GeneratedClass, BlueprintGetSpawner);
-
-			ToolTip = FText::FromString(FString::Printf(TEXT("Set the value of variable %s"), *ExternalVariable.Name.ToString()));
-			TSharedPtr<URigVMEdGraphVariableNodeSpawner> SetNodeSpawner = MakeShared<URigVMEdGraphVariableNodeSpawner>(RigVMBlueprint, ExternalVariable, false, MenuDesc, NodeCategory, ToolTip);
-			SetNodeSpawner->SetRelatedBlueprintClass(BlueprintClass);
-			URigVMEdGraphNodeBlueprintSpawner* BlueprintSetSpawner = URigVMEdGraphNodeBlueprintSpawner::CreateFromRigVMSpawner(SetNodeSpawner);
-			ActionRegistrar.AddBlueprintAction(GeneratedClass, BlueprintSetSpawner);
-		}
+	}
+	else if(InClass != BlueprintClass)
+	{
+		return;
+	}
+
+	static const FText NodeCategory = LOCTEXT("Variables", "Variables");
+
+	TArray<FRigVMGraphVariableDescription> Variables = RigVMBlueprint->GetAssetVariables();
+	for (const FRigVMGraphVariableDescription& Variable : Variables)
+	{
+		const FRigVMExternalVariable ExternalVariable = Variable.ToExternalVariable();
+		FText MenuDesc = FText::FromName(ExternalVariable.Name);
+		FText ToolTip = FText::FromString(FString::Printf(TEXT("Get the value of variable %s"), *ExternalVariable.Name.ToString()));
+		TSharedRef<URigVMEdGraphVariableNodeSpawner> GetNodeSpawner = MakeShared<URigVMEdGraphVariableNodeSpawner>(RigVMBlueprint, ExternalVariable, true, MenuDesc, NodeCategory, ToolTip);
+		GetNodeSpawner->SetRelatedBlueprintClass(BlueprintClass);
+		OutGraphNodeSpawners.Add(GetNodeSpawner);
+
+		ToolTip = FText::FromString(FString::Printf(TEXT("Set the value of variable %s"), *ExternalVariable.Name.ToString()));
+		TSharedRef<URigVMEdGraphVariableNodeSpawner> SetNodeSpawner = MakeShared<URigVMEdGraphVariableNodeSpawner>(RigVMBlueprint, ExternalVariable, false, MenuDesc, NodeCategory, ToolTip);
+		SetNodeSpawner->SetRelatedBlueprintClass(BlueprintClass);
+		OutGraphNodeSpawners.Add(SetNodeSpawner);
+	}
 
-		if (URigVMFunctionLibrary* LocalFunctionLibrary = RigVMBlueprint->GetLocalFunctionLibrary())
+	if (URigVMFunctionLibrary* LocalFunctionLibrary = RigVMBlueprint->GetLocalFunctionLibrary())
+	{
+		TArray<URigVMLibraryNode*> Functions = LocalFunctionLibrary->GetFunctions();
+		const FSoftObjectPath LocalLibrarySoftPath = LocalFunctionLibrary->GetFunctionHostObjectPath();
+		for (URigVMLibraryNode* Function : Functions)
 		{
-			TArray<URigVMLibraryNode*> Functions = LocalFunctionLibrary->GetFunctions();
-			const FSoftObjectPath LocalLibrarySoftPath = LocalFunctionLibrary->GetFunctionHostObjectPath();
-			for (URigVMLibraryNode* Function : Functions)
+			// Avoid adding functions that are already added by the GetTypeActions functions (public functions that are already saved into the blueprint tag)
+			if (RigVMBlueprint->GetPublicGraphFunctions().ContainsByPredicate([LocalLibrarySoftPath, Function](const FRigVMGraphFunctionHeader& Header) -> bool
 			{
-				// Avoid adding functions that are already added by the GetTypeActions functions (public functions that are already saved into the blueprint tag)
-				if (RigVMBlueprint->GetPublicGraphFunctions().ContainsByPredicate([LocalLibrarySoftPath, Function](const FRigVMGraphFunctionHeader& Header) -> bool
-				{
-					return FRigVMGraphFunctionIdentifier(LocalLibrarySoftPath, Function->GetPathName()) == Header.LibraryPointer;
-				}))
-				{
-					continue;
-				}
-				
-				TSharedPtr<URigVMEdGraphFunctionRefNodeSpawner> NodeSpawner = MakeShared<URigVMEdGraphFunctionRefNodeSpawner>(Function);
-				NodeSpawner->SetRelatedBlueprintClass(BlueprintClass);
-				URigVMEdGraphNodeBlueprintSpawner* BlueprintSpawner = URigVMEdGraphNodeBlueprintSpawner::CreateFromRigVMSpawner(NodeSpawner);
-				ActionRegistrar.AddBlueprintAction(GeneratedClass, BlueprintSpawner);
-			}
-
-			static const FText LocalNodeCategory = LOCTEXT("LocalVariables", "Local Variables");
-			for (URigVMLibraryNode* Function : Functions)
+				return FRigVMGraphFunctionIdentifier(LocalLibrarySoftPath, Function->GetPathName()) == Header.LibraryPointer;
+			}))
 			{
-				for (const FRigVMGraphVariableDescription& LocalVariable : Function->GetContainedGraph()->GetLocalVariables())
-				{
-					FText MenuDesc = FText::FromName(LocalVariable.Name);
-					FText ToolTip = FText::FromString(FString::Printf(TEXT("Get the value of variable %s"), *LocalVariable.Name.ToString()));
-					TSharedPtr<URigVMEdGraphVariableNodeSpawner> GetNodeSpawner = MakeShared<URigVMEdGraphVariableNodeSpawner>(RigVMBlueprint, Function->GetContainedGraph(), LocalVariable, true, MenuDesc, LocalNodeCategory, ToolTip);
-					GetNodeSpawner->SetRelatedBlueprintClass(BlueprintClass);
-					URigVMEdGraphNodeBlueprintSpawner* BlueprintGetSpawner = URigVMEdGraphNodeBlueprintSpawner::CreateFromRigVMSpawner(GetNodeSpawner);
-					ActionRegistrar.AddBlueprintAction(GeneratedClass, BlueprintGetSpawner);
-
-					ToolTip = FText::FromString(FString::Printf(TEXT("Set the value of variable %s"), *LocalVariable.Name.ToString()));
-					TSharedPtr<URigVMEdGraphVariableNodeSpawner> SetNodeSpawner = MakeShared<URigVMEdGraphVariableNodeSpawner>(RigVMBlueprint, Function->GetContainedGraph(), LocalVariable, false, MenuDesc, LocalNodeCategory, ToolTip);
-					SetNodeSpawner->SetRelatedBlueprintClass(BlueprintClass);
-					URigVMEdGraphNodeBlueprintSpawner* BlueprintSetSpawner = URigVMEdGraphNodeBlueprintSpawner::CreateFromRigVMSpawner(SetNodeSpawner);
-					ActionRegistrar.AddBlueprintAction(GeneratedClass, BlueprintSetSpawner);
-				}
+				continue;
 			}
+			
+			TSharedRef<URigVMEdGraphFunctionRefNodeSpawner> NodeSpawner = MakeShared<URigVMEdGraphFunctionRefNodeSpawner>(Function);
+			NodeSpawner->SetRelatedBlueprintClass(BlueprintClass);
+			OutGraphNodeSpawners.Add(NodeSpawner);
 		}
 
-		for (URigVMGraph* Graph : RigVMBlueprint->GetAllModels())
+		static const FText LocalNodeCategory = LOCTEXT("LocalVariables", "Local Variables");
+		for (URigVMLibraryNode* Function : Functions)
 		{
-			if (Graph->GetEntryNode())
+			for (const FRigVMGraphVariableDescription& LocalVariable : Function->GetContainedGraph()->GetLocalVariables())
 			{
-				static const FText InputNodeCategory = LOCTEXT("InputArguments", "Input Arguments");
-				for (const FRigVMGraphVariableDescription& InputArgument : Graph->GetInputArguments())
-				{
-					FText MenuDesc = FText::FromName(InputArgument.Name);
-					FText ToolTip = FText::FromString(FString::Printf(TEXT("Get the value of input %s"), *InputArgument.Name.ToString()));
-					TSharedPtr<URigVMEdGraphVariableNodeSpawner> NodeSpawner = MakeShared<URigVMEdGraphVariableNodeSpawner>(RigVMBlueprint, Graph, InputArgument, true, MenuDesc, InputNodeCategory, ToolTip);
-					NodeSpawner->SetRelatedBlueprintClass(BlueprintClass);
-					URigVMEdGraphNodeBlueprintSpawner* BlueprintSpawner = URigVMEdGraphNodeBlueprintSpawner::CreateFromRigVMSpawner(NodeSpawner);
-					ActionRegistrar.AddBlueprintAction(GeneratedClass, BlueprintSpawner);
-				}			
+				FText MenuDesc = FText::FromName(LocalVariable.Name);
+				FText ToolTip = FText::FromString(FString::Printf(TEXT("Get the value of variable %s"), *LocalVariable.Name.ToString()));
+				TSharedRef<URigVMEdGraphVariableNodeSpawner> GetNodeSpawner = MakeShared<URigVMEdGraphVariableNodeSpawner>(RigVMBlueprint, Function->GetContainedGraph(), LocalVariable, true, MenuDesc, LocalNodeCategory, ToolTip);
+				GetNodeSpawner->SetRelatedBlueprintClass(BlueprintClass);
+				OutGraphNodeSpawners.Add(GetNodeSpawner);
+
+				ToolTip = FText::FromString(FString::Printf(TEXT("Set the value of variable %s"), *LocalVariable.Name.ToString()));
+				TSharedRef<URigVMEdGraphVariableNodeSpawner> SetNodeSpawner = MakeShared<URigVMEdGraphVariableNodeSpawner>(RigVMBlueprint, Function->GetContainedGraph(), LocalVariable, false, MenuDesc, LocalNodeCategory, ToolTip);
+				SetNodeSpawner->SetRelatedBlueprintClass(BlueprintClass);
+				OutGraphNodeSpawners.Add(SetNodeSpawner);
 			}
 		}
+	}
 
-		const TArray<FName> EntryNames = RigVMBlueprint->GetRigVMClient()->GetEntryNames();
-		if(!EntryNames.IsEmpty())
+	for (URigVMGraph* Graph : RigVMBlueprint->GetAllModels())
+	{
+		if (Graph->GetEntryNode())
 		{
-			static const FText EventNodeCategory = LOCTEXT("Events", "Events");
-			for (const FName& EntryName : EntryNames)
+			static const FText InputNodeCategory = LOCTEXT("InputArguments", "Input Arguments");
+			for (const FRigVMGraphVariableDescription& InputArgument : Graph->GetInputArguments())
 			{
-				static constexpr TCHAR EventStr[] = TEXT("Event");
-				static const FString EventSuffix = FString::Printf(TEXT(" %s"), EventStr);
-				FString Suffix = EntryName.ToString().EndsWith(EventStr) ? FString() : EventSuffix;
-				FText MenuDesc = FText::FromString(FString::Printf(TEXT("Run %s%s"), *EntryName.ToString(), *Suffix));
-				FText ToolTip = FText::FromString(FString::Printf(TEXT("Runs the %s%s"), *EntryName.ToString(), *Suffix));
-				TSharedPtr<URigVMEdGraphInvokeEntryNodeSpawner> NodeSpawner = MakeShared<URigVMEdGraphInvokeEntryNodeSpawner>(RigVMBlueprint, EntryName, MenuDesc, EventNodeCategory, ToolTip);
+				FText MenuDesc = FText::FromName(InputArgument.Name);
+				FText ToolTip = FText::FromString(FString::Printf(TEXT("Get the value of input %s"), *InputArgument.Name.ToString()));
+				TSharedRef<URigVMEdGraphVariableNodeSpawner> NodeSpawner = MakeShared<URigVMEdGraphVariableNodeSpawner>(RigVMBlueprint, Graph, InputArgument, true, MenuDesc, InputNodeCategory, ToolTip);
 				NodeSpawner->SetRelatedBlueprintClass(BlueprintClass);
-				URigVMEdGraphNodeBlueprintSpawner* BlueprintSpawner = URigVMEdGraphNodeBlueprintSpawner::CreateFromRigVMSpawner(NodeSpawner);
-				ActionRegistrar.AddBlueprintAction(GeneratedClass, BlueprintSpawner);
+				OutGraphNodeSpawners.Add(NodeSpawner);
 			}
 		}
 	}
+
+	const TArray<FName> EntryNames = RigVMBlueprint->GetRigVMClient()->GetEntryNames();
+	if(!EntryNames.IsEmpty())
+	{
+		static const FText EventNodeCategory = LOCTEXT("Events", "Events");
+		for (const FName& EntryName : EntryNames)
+		{
+			static constexpr TCHAR EventStr[] = TEXT("Event");
+			static const FString EventSuffix = FString::Printf(TEXT(" %s"), EventStr);
+			FString Suffix = EntryName.ToString().EndsWith(EventStr) ? FString() : EventSuffix;
+			FText MenuDesc = FText::FromString(FString::Printf(TEXT("Run %s%s"), *EntryName.ToString(), *Suffix));
+			FText ToolTip = FText::FromString(FString::Printf(TEXT("Runs the %s%s"), *EntryName.ToString(), *Suffix));
+			TSharedRef<URigVMEdGraphInvokeEntryNodeSpawner> NodeSpawner = MakeShared<URigVMEdGraphInvokeEntryNodeSpawner>(RigVMBlueprint, EntryName, MenuDesc, EventNodeCategory, ToolTip);
+			NodeSpawner->SetRelatedBlueprintClass(BlueprintClass);
+			OutGraphNodeSpawners.Add(NodeSpawner);
+		}
+	}
 }
 
 void FRigVMEditorModule::GetNodeContextMenuActions(TScriptInterface<IRigVMClientHost> RigVMClientHost, const URigVMEdGraphNode* EdGraphNode, URigVMNode* ModelNode, UToolMenu* Menu) const
 {
 	// Only register menu actions for ourselves if we are a blueprint
-	if(FRigVMAssetInterfacePtr RigVMBlueprint = RigVMClientHost.GetObject())
+	if(FRigVMEditorAssetInterfacePtr RigVMBlueprint = RigVMClientHost.GetObject())
 	{
 		UClass* BPClass = RigVMBlueprint->GetObject()->GetClass();
 		UClass* BaseClass = GetRigVMAssetClass();
-		if(BPClass != GetRigVMAssetClass() && !BPClass->IsChildOf(BaseClass))
+		if (RigVMBlueprint->IsBlueprintAsset() && CVarRigVMBlueprintIndependentAssets->GetBool())
+		{
+			BaseClass = GetRigVMLegacyBlueprintClass();
+		}
+		else if (BaseClass->IsChildOf(URigVMRuntimeAsset::StaticClass()))
+		{
+			BPClass = RigVMBlueprint->GetRuntimeAssetInterface().GetObject()->GetClass();
+		}
+		if(BPClass != BaseClass && !BPClass->IsChildOf(BaseClass))
 		{
 			return;
 		}
@@ -687,7 +748,7 @@ void FRigVMEditorModule::GetNodeContextMenuActions(TScriptInterface<IRigVMClient
 
 void FRigVMEditorModule::GetNodeActionsContextMenuActions(TScriptInterface<IRigVMClientHost> RigVMClientHost, const URigVMEdGraphNode* EdGraphNode, URigVMNode* ModelNode, UToolMenu* Menu) const
 {
-	const TScriptInterface<IRigVMAssetInterface> RigVMAsset = RigVMClientHost.GetObject();
+	const TScriptInterface<IRigVMEditorAssetInterface> RigVMAsset = RigVMClientHost.GetObject();
 	if (!EdGraphNode || 
 		!ModelNode || 
 		!Menu || 
@@ -696,7 +757,7 @@ void FRigVMEditorModule::GetNodeActionsContextMenuActions(TScriptInterface<IRigV
 		return;
 	}
 
-	const TWeakInterfacePtr<IRigVMAssetInterface> WeakRigVMAsset = RigVMAsset;
+	const TWeakInterfacePtr<IRigVMEditorAssetInterface> WeakRigVMAsset = RigVMAsset;
 	const TWeakObjectPtr<URigVMNode> WeakModelNode = ModelNode;
 	const TWeakObjectPtr<const URigVMEdGraphNode> WeakEdGraphNode = EdGraphNode;
 
@@ -844,7 +905,7 @@ void FRigVMEditorModule::GetNodeWorkflowContextMenuActions(TScriptInterface<IRig
 
 void FRigVMEditorModule::GetNodeEventsContextMenuActions(TScriptInterface<IRigVMClientHost> RigVMClientHost, const URigVMEdGraphNode* EdGraphNode, URigVMNode* ModelNode, UToolMenu* Menu) const
 {
-	FRigVMAssetInterfacePtr RigVMBlueprint = RigVMClientHost.GetObject();
+	FRigVMEditorAssetInterfacePtr RigVMBlueprint = RigVMClientHost.GetObject();
 	if(RigVMBlueprint && ModelNode->IsEvent())
 	{
 		const FName& EventName = ModelNode->GetEventName();
@@ -1036,7 +1097,7 @@ void FRigVMEditorModule::GetNodeConversionContextMenuActions(TScriptInterface<IR
 
 void FRigVMEditorModule::GetNodeDebugContextMenuActions(TScriptInterface<IRigVMClientHost> RigVMClientHost, const URigVMEdGraphNode* EdGraphNode, URigVMNode* ModelNode, UToolMenu* Menu) const
 {
-	if(FRigVMAssetInterfacePtr RigVMAsset = RigVMClientHost.GetObject())
+	if(FRigVMEditorAssetInterfacePtr RigVMAsset = RigVMClientHost.GetObject())
 	{
 		FToolMenuSection& DebugSection = Menu->AddSection("RigVMEditorContextMenuDebug", LOCTEXT("DebugHeader", "Debug"));
 		bool bAnyNodeIsEarlyExit = false;
@@ -1092,7 +1153,7 @@ void FRigVMEditorModule::GetNodeDebugContextMenuActions(TScriptInterface<IRigVMC
 
 void FRigVMEditorModule::GetNodeVariablesContextMenuActions(TScriptInterface<IRigVMClientHost> RigVMClientHost, const URigVMEdGraphNode* EdGraphNode, URigVMNode* ModelNode, UToolMenu* Menu) const
 {
-	if (FRigVMAssetInterfacePtr RigVMBlueprint = RigVMClientHost.GetObject())
+	if (FRigVMEditorAssetInterfacePtr RigVMBlueprint = RigVMClientHost.GetObject())
 	{
 		const URigVMGraph* Model = ModelNode->GetGraph();
 		URigVMController* Controller = RigVMBlueprint->GetController(Model);
@@ -1134,7 +1195,7 @@ void FRigVMEditorModule::GetNodeVariablesContextMenuActions(TScriptInterface<IRi
 
 								if (URigVMLibraryNode* LibraryNode = FunctionReferenceNode->LoadReferencedNode())
 								{
-									FRigVMAssetInterfacePtr ReferencedBlueprint = IRigVMAssetInterface::GetInterfaceOuter(LibraryNode);
+									FRigVMEditorAssetInterfacePtr ReferencedBlueprint = IRigVMEditorAssetInterface::GetInterfaceOuter(LibraryNode);
 								   // ReferencedBlueprint != RigVMBlueprint - since only FunctionReferenceNodes from other assets have the potential to be unmapped
                             
 								   for(const FRigVMExternalVariable& ExternalVariable : ExternalVariables)
@@ -1197,7 +1258,7 @@ void FRigVMEditorModule::GetNodeTemplatesContextMenuActions(TScriptInterface<IRi
 
 void FRigVMEditorModule::GetNodeOrganizationContextMenuActions(TScriptInterface<IRigVMClientHost> RigVMClientHost, const URigVMEdGraphNode* EdGraphNode, URigVMNode* ModelNode, UToolMenu* Menu) const
 {
-	FRigVMAssetInterfacePtr RigVMBlueprint = RigVMClientHost.GetObject();
+	FRigVMEditorAssetInterfacePtr RigVMBlueprint = RigVMClientHost.GetObject();
 	URigVMGraph* Model = ModelNode->GetGraph();
 	URigVMController* Controller = RigVMClientHost->GetRigVMClient()->GetController(Model);
 
@@ -1607,7 +1668,7 @@ void FRigVMEditorModule::GetNodeTestContextMenuActions(TScriptInterface<IRigVMCl
 
 void FRigVMEditorModule::GetNodeDisplayContextMenuActions(TScriptInterface<IRigVMClientHost> RigVMClientHost, const URigVMEdGraphNode* EdGraphNode, URigVMNode* ModelNode, UToolMenu* Menu) const
 {
-	FRigVMAssetInterfacePtr Blueprint = RigVMClientHost.GetObject();
+	FRigVMEditorAssetInterfacePtr Blueprint = RigVMClientHost.GetObject();
 	if(Blueprint == nullptr)
 	{
 		return;
@@ -1743,7 +1804,7 @@ void FRigVMEditorModule::GetPinWorkflowContextMenuActions(TScriptInterface<IRigV
 
 void FRigVMEditorModule::GetPinDebugContextMenuActions(TScriptInterface<IRigVMClientHost> RigVMClientHost, const UEdGraphPin* EdGraphPin, URigVMPin* ModelPin, UToolMenu* Menu) const
 {
-	if (FRigVMAssetInterfacePtr RigVMBlueprint = RigVMClientHost.GetObject())
+	if (FRigVMEditorAssetInterfacePtr RigVMBlueprint = RigVMClientHost.GetObject())
 	{
 		const bool bIsEditablePin = !ModelPin->IsExecuteContext() && !ModelPin->IsWildCard();
 
@@ -1990,7 +2051,7 @@ void FRigVMEditorModule::GetPinConversionContextMenuActions(TScriptInterface<IRi
 
 void FRigVMEditorModule::GetPinVariableContextMenuActions(TScriptInterface<IRigVMClientHost> RigVMClientHost, const UEdGraphPin* EdGraphPin, URigVMPin* ModelPin, UToolMenu* Menu) const
 {
-	if(FRigVMAssetInterfacePtr RigVMBlueprint = RigVMClientHost.GetObject())
+	if(FRigVMEditorAssetInterfacePtr RigVMBlueprint = RigVMClientHost.GetObject())
 	{
 		const bool bIsEditablePin = !ModelPin->IsExecuteContext() && !ModelPin->IsWildCard();
 
@@ -2488,7 +2549,7 @@ void FRigVMEditorModule::PostChange(const UUserDefinedStruct* Changed, FStructur
 		return;
 	}
 
-	TArray<FRigVMAssetInterfacePtr> BlueprintsToRefresh;
+	TArray<FRigVMEditorAssetInterfacePtr> BlueprintsToRefresh;
 	TArray<UEdGraph*> EdGraphsToRefresh;
 
 	for (TObjectIterator<URigVMPin> It(RF_Transient | RF_ClassDefaultObject, /** bIncludeDerivedClasses */ true, /** InternalExcludeFlags */ EInternalObjectFlags::Garbage); It; ++It)
@@ -2497,7 +2558,7 @@ void FRigVMEditorModule::PostChange(const UUserDefinedStruct* Changed, FStructur
 		// GetCPPTypeObject also makes sure the pin's type information is update to date
 		if (Pin && Pin->GetCPPTypeObject() == Changed)
 		{
-			if (FRigVMAssetInterfacePtr RigVMBlueprint = IRigVMAssetInterface::GetInterfaceOuter(Pin))
+			if (FRigVMEditorAssetInterfacePtr RigVMBlueprint = IRigVMEditorAssetInterface::GetInterfaceOuter(Pin))
 			{
 				BlueprintsToRefresh.AddUnique(RigVMBlueprint);
 				
@@ -2516,7 +2577,7 @@ void FRigVMEditorModule::PostChange(const UUserDefinedStruct* Changed, FStructur
 							continue;
 						}
 						
-						if (FRigVMAssetInterfacePtr FunctionUserBlueprint = IRigVMAssetInterface::GetInterfaceOuter(RefNode))
+						if (FRigVMEditorAssetInterfacePtr FunctionUserBlueprint = IRigVMEditorAssetInterface::GetInterfaceOuter(RefNode))
 						{
 							BlueprintsToRefresh.AddUnique(FunctionUserBlueprint);
 						}
@@ -2531,10 +2592,10 @@ void FRigVMEditorModule::PostChange(const UUserDefinedStruct* Changed, FStructur
 		}
 	}
 
-	for (FRigVMAssetInterfacePtr RigVMBlueprint : BlueprintsToRefresh)
+	for (FRigVMEditorAssetInterfacePtr RigVMBlueprint : BlueprintsToRefresh)
 	{
 		RigVMBlueprint->OnRigVMRegistryChanged();
-		(void)RigVMBlueprint->MarkPackageDirty();
+		(void)RigVMBlueprint.GetObject()->MarkPackageDirty();
 	}
 
 	// Avoid slate crashing after pins get repopulated
@@ -2543,7 +2604,7 @@ void FRigVMEditorModule::PostChange(const UUserDefinedStruct* Changed, FStructur
 		Graph->NotifyGraphChanged();
 	}
 
-	for (FRigVMAssetInterfacePtr RigVMBlueprint : BlueprintsToRefresh)
+	for (FRigVMEditorAssetInterfacePtr RigVMBlueprint : BlueprintsToRefresh)
 	{
 		// this should make sure variables in BP are updated with the latest struct object
 		// otherwise RigVMCompiler validation would complain about variable type - pin type mismatch
@@ -2584,7 +2645,7 @@ TArray<FRigVMEditorModule::FRigVMEditorToolbarExtender>& FRigVMEditorModule::Get
 	return RigVMEditorToolbarExtenders;
 }
 
-void FRigVMEditorModule::CreateRootGraphIfRequired(FRigVMAssetInterfacePtr InBlueprint) const
+void FRigVMEditorModule::CreateRootGraphIfRequired(FRigVMEditorAssetInterfacePtr InBlueprint) const
 {
 	if(InBlueprint == nullptr)
 	{
@@ -2616,7 +2677,7 @@ void FRigVMEditorModule::CreateRootGraphIfRequired(FRigVMAssetInterfacePtr InBlu
 
 void FRigVMEditorModule::HandleNewBlueprintCreated(UBlueprint* InBlueprint)
 {
-	CreateRootGraphIfRequired(FRigVMAssetInterfacePtr(InBlueprint));
+	CreateRootGraphIfRequired(FRigVMEditorAssetInterfacePtr(InBlueprint));
 }
 
 bool FRigVMEditorModule::ShowWorkflowOptionsDialog(URigVMUserWorkflowOptions* InOptions) const
@@ -2653,7 +2714,7 @@ FConnectionDrawingPolicy* FRigVMEditorModule::CreateConnectionDrawingPolicy(int3
 bool FRigVMEditorModule::IsRigVMEditorModuleBase() const
 {
 	UClass* Class = GetRigVMAssetClass();
-	if (Class == URigVMAssetInterface::StaticClass() || Class->ImplementsInterface(URigVMAssetInterface::StaticClass()))
+	if (Class == URigVMEditorAssetInterface::StaticClass() || Class->ImplementsInterface(URigVMEditorAssetInterface::StaticClass()))
 	{
 		return true;
 	}
diff --git a/Engine/Plugins/Runtime/RigVM/Source/RigVMEditor/Private/Widgets/RigVMTreeToolkitNode.cpp b/Engine/Plugins/Runtime/RigVM/Source/RigVMEditor/Private/Widgets/RigVMTreeToolkitNode.cpp
--- a/Engine/Plugins/Runtime/RigVM/Source/RigVMEditor/Private/Widgets/RigVMTreeToolkitNode.cpp
+++ b/Engine/Plugins/Runtime/RigVM/Source/RigVMEditor/Private/Widgets/RigVMTreeToolkitNode.cpp
@@ -11,13 +11,13 @@ FAssetData FRigVMTreeNode::GetAssetData() const
 	return FRigVMTreeContext::FindAssetFromAnyPath(GetPath(), false);
 }
 
-IRigVMAssetInterface* FRigVMTreeNode::GetAsset() const
+IRigVMEditorAssetInterface* FRigVMTreeNode::GetAsset() const
 {
 	const FAssetData AssetData = FRigVMTreeContext::FindAssetFromAnyPath(GetPath(), true);
 	if(AssetData.IsAssetLoaded())
 	{
 		UObject* Asset = AssetData.GetAsset();
-		return Cast<IRigVMAssetInterface>(Asset);
+		return Cast<IRigVMEditorAssetInterface>(Asset);
 	}
 	return nullptr;
 }
@@ -381,7 +381,7 @@ const TArray<FRigVMTag>& FRigVMTreePackageNode::GetTags() const
 	{
 		if(IsLoaded())
 		{
-			if(const IRigVMAssetInterface* Blueprint = GetAsset())
+			if(const IRigVMEditorAssetInterface* Blueprint = GetAsset())
 			{
 				Tags = Blueprint->GetAssetVariant().Tags;
 			}
diff --git a/Engine/Plugins/Runtime/RigVM/Source/RigVMEditor/Private/Widgets/SRigVMBulkEditWidget.cpp b/Engine/Plugins/Runtime/RigVM/Source/RigVMEditor/Private/Widgets/SRigVMBulkEditWidget.cpp
--- a/Engine/Plugins/Runtime/RigVM/Source/RigVMEditor/Private/Widgets/SRigVMBulkEditWidget.cpp
+++ b/Engine/Plugins/Runtime/RigVM/Source/RigVMEditor/Private/Widgets/SRigVMBulkEditWidget.cpp
@@ -1,6 +1,7 @@
 // Copyright Epic Games, Inc. All Rights Reserved.
 
 #include "Widgets/SRigVMBulkEditWidget.h"
+#include "RigVMDeveloperModule.h"
 #include "Widgets/SWindow.h"
 #include "Widgets/SRigVMLogWidget.h"
 #include "Widgets/Input/SButton.h"
@@ -251,7 +252,7 @@ void SRigVMBulkEditWidget::Tick(const FGeometry& AllottedGeometry, const double
 				const TSharedRef<FRigVMTreePhase>& Phase = GetActivePhase();
 
 				TSharedPtr<TGuardValue<FRigVMReportDelegate>> ReportDelegateGuard; 
-				if(FRigVMAssetInterfacePtr Blueprint = Task->GetRigVMAsset(Phase))
+				if(FRigVMEditorAssetInterfacePtr Blueprint = Task->GetRigVMAsset(Phase))
 				{
 					ReportDelegateGuard = MakeShareable(new TGuardValue<FRigVMReportDelegate>(
 						Blueprint->GetVMCompileSettings().ASTSettings.ReportDelegate,
diff --git a/Engine/Plugins/Runtime/RigVM/Source/RigVMEditor/Private/Widgets/SRigVMEditorGraphExplorer.cpp b/Engine/Plugins/Runtime/RigVM/Source/RigVMEditor/Private/Widgets/SRigVMEditorGraphExplorer.cpp
--- a/Engine/Plugins/Runtime/RigVM/Source/RigVMEditor/Private/Widgets/SRigVMEditorGraphExplorer.cpp
+++ b/Engine/Plugins/Runtime/RigVM/Source/RigVMEditor/Private/Widgets/SRigVMEditorGraphExplorer.cpp
@@ -37,16 +37,17 @@
 #include "EdGraph/RigVMEdGraph.h"
 #include "Editor/SRigVMDetailsInspector.h"
 #include "WorkflowOrientedApp/WorkflowTabManager.h"
+#include "RigVMDeveloperTypeUtils.h"
 
 #define LOCTEXT_NAMESPACE "RigVMGraphExplorer"
 
 TSharedRef<FRigVMGraphExplorerDragDropOp> FRigVMGraphExplorerDragDropOp::New(const FRigVMExplorerElementKey& InElement, TObjectPtr<URigVMBlueprint> InBlueprint)
 {
-	const FRigVMAssetInterfacePtr Asset(InBlueprint);
+	const FRigVMEditorAssetInterfacePtr Asset(InBlueprint);
 	return New(InElement, Asset);
 }
 
-TSharedRef<FRigVMGraphExplorerDragDropOp> FRigVMGraphExplorerDragDropOp::New(const ::FRigVMExplorerElementKey& InElement, FRigVMAssetInterfacePtr InBlueprint)
+TSharedRef<FRigVMGraphExplorerDragDropOp> FRigVMGraphExplorerDragDropOp::New(const ::FRigVMExplorerElementKey& InElement, FRigVMEditorAssetInterfacePtr InBlueprint)
 {
 	TSharedRef<FRigVMGraphExplorerDragDropOp> Operation = MakeShared<FRigVMGraphExplorerDragDropOp>();
 	Operation->Element = InElement;
@@ -105,10 +106,16 @@ FReply FRigVMGraphExplorerDragDropOp::DroppedOnPanel(const TSharedRef<SWidget>&
 {
 	if (URigVMEdGraph* TargetRigGraph = Cast<URigVMEdGraph>(&Graph))
 	{
-		if(FRigVMAssetInterfacePtr Blueprint = TargetRigGraph->GetAsset())
+		if(FRigVMEditorAssetInterfacePtr Blueprint = TargetRigGraph->GetAsset())
 		{
 			//  Find the appropriate asset editor where the drop is happening
-			const TArray<IAssetEditorInstance*> AssetEditors = GEditor->GetEditorSubsystem<UAssetEditorSubsystem>()->FindEditorsForAsset(Blueprint->GetObject());
+			UObject* AssetObj = Blueprint.GetObject();
+			if (TScriptInterface<IRigVMRuntimeAssetInterface> RuntimeAsset = Blueprint->GetRuntimeAssetInterface())
+			{
+				AssetObj = RuntimeAsset.GetObject();
+			}
+			
+			const TArray<IAssetEditorInstance*> AssetEditors = GEditor->GetEditorSubsystem<UAssetEditorSubsystem>()->FindEditorsForAsset(AssetObj);
 			IAssetEditorInstance* FocusedAssetEditor = nullptr;
 			for (IAssetEditorInstance* AssetEditor : AssetEditors)
 			{
@@ -134,8 +141,8 @@ FReply FRigVMGraphExplorerDragDropOp::DroppedOnPanel(const TSharedRef<SWidget>&
 
 			if (!FocusedAssetEditor)
 			{
-				GEditor->GetEditorSubsystem<UAssetEditorSubsystem>()->OpenEditorForAsset(Blueprint->GetObject());
-				FocusedAssetEditor = GEditor->GetEditorSubsystem<UAssetEditorSubsystem>()->FindEditorForAsset(Blueprint->GetObject(), /*bFocusIfOpen =*/true);
+				GEditor->GetEditorSubsystem<UAssetEditorSubsystem>()->OpenEditorForAsset(AssetObj);
+				FocusedAssetEditor = GEditor->GetEditorSubsystem<UAssetEditorSubsystem>()->FindEditorForAsset(AssetObj, /*bFocusIfOpen =*/true);
 			}
 			if (FocusedAssetEditor)
 			{
@@ -153,7 +160,7 @@ FReply FRigVMGraphExplorerDragDropOp::DroppedOnPin(const FVector2f& ScreenPositi
 {
 	UEdGraph* TargetGraph = GetHoveredGraph();
 	UEdGraphPin* TargetPin = GetHoveredPin();
-	FRigVMAssetInterfacePtr Blueprint = GetRigVMAssetInterface();
+	FRigVMEditorAssetInterfacePtr Blueprint = GetRigVMAssetInterface();
 	
 	const URigVMEdGraphSchema* EdSchema = GetDefault<URigVMEdGraphSchema>();
 
@@ -228,7 +235,7 @@ void SRigVMEditorGraphExplorer::Construct(const FArguments& InArgs, TWeakPtr<IRi
 		Editor->OnRefresh().AddSP(this, &SRigVMEditorGraphExplorer::Refresh);
 
 		// Acquire a registry for function and variable nodes
-		if (const TScriptInterface<IRigVMAssetInterface> AssetInterface = Editor->GetRigVMAssetInterface())
+		if (const TScriptInterface<IRigVMEditorAssetInterface> AssetInterface = Editor->GetRigVMAssetInterface())
 		{
 			EdGraphNodeFunctionRegistry = FRigVMEdGraphNodeRegistry::GetOrCreateRegistry(AssetInterface, URigVMFunctionReferenceNode::StaticClass());
 			EdGraphNodeVariableRegistry = FRigVMEdGraphNodeRegistry::GetOrCreateRegistry(AssetInterface, URigVMVariableNode::StaticClass());
@@ -284,6 +291,42 @@ void SRigVMEditorGraphExplorer::ClearSelection()
 	TreeView->ClearSelection();
 }
 
+void SRigVMEditorGraphExplorer::HandleGraphModifiedEvent(ERigVMGraphNotifType InNotifType, URigVMGraph* InGraph, UObject* InSubject)
+{
+	switch (InNotifType)
+	{
+		case ERigVMGraphNotifType::GraphChanged:
+		case ERigVMGraphNotifType::VariableAdded:
+		case ERigVMGraphNotifType::VariableRenamed:
+		case ERigVMGraphNotifType::VariableRemoved:
+		case ERigVMGraphNotifType::LocalVariableAdded:
+		case ERigVMGraphNotifType::LocalVariableRenamed:
+		case ERigVMGraphNotifType::LocalVariableRemoved:
+		case ERigVMGraphNotifType::LocalVariableTypeChanged:
+			{
+				Refresh();
+				break;
+			}
+		case ERigVMGraphNotifType::NodeAdded:
+		case ERigVMGraphNotifType::NodeRemoved:
+		case ERigVMGraphNotifType::NodeRenamed:
+			{
+				if (URigVMCollapseNode* FunctionLibrary = Cast<URigVMCollapseNode>(InSubject))
+				{
+					Refresh();
+				}
+				if (URigVMUnitNode* Node = Cast<URigVMUnitNode>(InSubject))
+				{
+					if (Node->IsEvent())
+					{
+						Refresh();
+					}
+				}
+				break;
+			}
+	}
+}
+
 void SRigVMEditorGraphExplorer::RegisterCommands()
 {
 	TSharedPtr<IRigVMEditor> Editor = RigVMEditor.Pin();
@@ -704,7 +747,7 @@ TArray<FRigVMGraphVariableDescription> SRigVMEditorGraphExplorer::GetLocalVariab
 		return Variables;
 	}
 
-	if (FRigVMAssetInterfacePtr Blueprint = Editor->GetRigVMAssetInterface())
+	if (FRigVMEditorAssetInterfacePtr Blueprint = Editor->GetRigVMAssetInterface())
 	{
 		if (URigVMGraph* Function = Blueprint->GetFocusedModel())
 		{
@@ -791,13 +834,20 @@ FText SRigVMEditorGraphExplorer::GetGraphTooltip(const FString& InGraphPath) con
 		return Tooltip;
 	}
 
-	FRigVMClient* RigVMClient = Editor->GetRigVMAssetInterface()->GetRigVMClient();
+	FRigVMEditorAssetInterfacePtr Asset = Editor->GetRigVMAssetInterface();
+	FRigVMClient* RigVMClient = Asset->GetRigVMClient();
 	if (const URigVMGraph* Graph = RigVMClient->GetModel(InGraphPath))
 	{
 		const URigVMEdGraphSchema* EdSchema = GetDefault<URigVMEdGraphSchema>();
 		FGraphDisplayInfo DisplayInfo;
-		UEdGraph* EdGraph = Editor->GetRigVMAssetInterface()->GetEdGraph(Graph);
-		EdSchema->GetGraphDisplayInformation(*EdGraph, DisplayInfo);
+		if (UEdGraph* EdGraph = Asset->GetEdGraph(Graph))
+		{
+			EdSchema->GetGraphDisplayInformation(*EdGraph, DisplayInfo);
+		}
+		else
+		{
+			ensure(false);
+		}
 
 		Tooltip = DisplayInfo.Tooltip;
 	}
@@ -905,7 +955,7 @@ void SRigVMEditorGraphExplorer::OnVariableClicked(const FRigVMExplorerElementKey
 		return;
 	}
 
-	FRigVMAssetInterfacePtr Blueprint = Editor->GetRigVMAssetInterface();
+	FRigVMEditorAssetInterfacePtr Blueprint = Editor->GetRigVMAssetInterface();
 	if (InVariable.Type == ERigVMExplorerElementType::Variable)
 	{
 		FProperty* Prop = Blueprint->FindGeneratedPropertyByName(*InVariable.Name);
@@ -925,6 +975,7 @@ void SRigVMEditorGraphExplorer::OnVariableClicked(const FRigVMExplorerElementKey
 		{
 			SRigVMDetailsInspector::FShowDetailsOptions Options(FText::FromString(InVariable.Name));
 			Options.bForceRefresh = true;
+			
 			if (RigVMInspector.IsValid())
 			{
 				RigVMInspector->ShowDetailsForSingleObject(PropWrap, Options);
@@ -966,7 +1017,7 @@ void SRigVMEditorGraphExplorer::OnEventDoubleClicked(const FString& InEventPath)
 		return;
 	}
 
-	FRigVMAssetInterfacePtr Blueprint = Editor->GetRigVMAssetInterface();
+	FRigVMEditorAssetInterfacePtr Blueprint = Editor->GetRigVMAssetInterface();
 	FRigVMClient* RigVMClient = Blueprint->GetRigVMClient();
 
 	if(const URigVMNode* ModelNode = RigVMClient->FindNode(InEventPath))
@@ -974,7 +1025,6 @@ void SRigVMEditorGraphExplorer::OnEventDoubleClicked(const FString& InEventPath)
 		Blueprint->OnRequestJumpToHyperlink().Execute(ModelNode);
 	}
 }
-
 void SRigVMEditorGraphExplorer::OnFunctionDoubleClicked(const FString& InFunctionPath)
 {
 	TSharedPtr<IRigVMEditor> Editor = RigVMEditor.Pin();
@@ -1102,14 +1152,14 @@ FEdGraphPinType SRigVMEditorGraphExplorer::OnGetVariablePinType(const FRigVMExpl
 		return FEdGraphPinType();
 	}
 
-	FRigVMAssetInterfacePtr Blueprint = Editor->GetRigVMAssetInterface();
+	FRigVMEditorAssetInterfacePtr Blueprint = Editor->GetRigVMAssetInterface();
 	if (InVariable.Type == ERigVMExplorerElementType::Variable)
 	{
 		for (FRigVMGraphVariableDescription& Variable : Blueprint->GetAssetVariables())
 		{
 			if (Variable.Name == *InVariable.Name)
 			{
-				return Variable.ToPinType();
+				return RigVMTypeUtils::PinTypeFromRigVMVariableDescription(Variable);
 			}
 		}
 	}
@@ -1121,7 +1171,7 @@ FEdGraphPinType SRigVMEditorGraphExplorer::OnGetVariablePinType(const FRigVMExpl
 			{
 				if (Variable.Name == *InVariable.Name)
 				{
-					return Variable.ToPinType();
+					return RigVMTypeUtils::PinTypeFromRigVMVariableDescription(Variable);
 				}
 			}
 		}
@@ -1138,7 +1188,7 @@ bool SRigVMEditorGraphExplorer::OnSetVariablePinType(const FRigVMExplorerElement
 		return false;
 	}
 
-	FRigVMAssetInterfacePtr Blueprint = Editor->GetRigVMAssetInterface();
+	FRigVMEditorAssetInterfacePtr Blueprint = Editor->GetRigVMAssetInterface();
 
 	if (InVariable.Type == ERigVMExplorerElementType::Variable)
 	{
@@ -1149,8 +1199,9 @@ bool SRigVMEditorGraphExplorer::OnSetVariablePinType(const FRigVMExplorerElement
 				FString CPPType;
 				UObject* CPPTypeObject = nullptr;
 				RigVMTypeUtils::CPPTypeFromPinType(InType, CPPType, &CPPTypeObject);
-				Blueprint->ChangeMemberVariableType(*InVariable.Name, CPPType);
+				Blueprint->ChangeMemberVariableType(*InVariable.Name, CPPType, Variable.bPublic);
 				SetLastPinTypeUsed(InType);
+				Refresh();
 				return true;
 			}
 		}
@@ -1166,6 +1217,8 @@ bool SRigVMEditorGraphExplorer::OnSetVariablePinType(const FRigVMExplorerElement
 				RigVMTypeUtils::CPPTypeFromPinType(InType, NewCPPType, &NewCPPTypeObject);
 				Controller->SetLocalVariableType(*InVariable.Name, NewCPPType, NewCPPTypeObject, true, true);
 				SetLastPinTypeUsed(InType);
+				Refresh();
+				Editor->SelectLocalVariable(Blueprint->GetEdGraph(Function), *InVariable.Name);
 				return true;
 			}
 		}
@@ -1182,15 +1235,12 @@ bool SRigVMEditorGraphExplorer::OnIsVariablePublic(const FString& InVariableName
 		return false;
 	}
 
-	FRigVMAssetInterfacePtr Blueprint = Editor->GetRigVMAssetInterface();
-	for (FRigVMGraphVariableDescription& Variable : Blueprint->GetAssetVariables())
+	FRigVMEditorAssetInterfacePtr Blueprint = Editor->GetRigVMAssetInterface();
+	if (const FProperty* Property = Blueprint->FindGeneratedPropertyByName(*InVariableName))
 	{
-		if (Variable.Name == *InVariableName)
-		{
-			return Variable.bPublic;
-		}
+		return !Property->HasAnyPropertyFlags(CPF_DisableEditOnInstance);
 	}
-
+	
 	return false;
 }
 
@@ -1202,7 +1252,7 @@ bool SRigVMEditorGraphExplorer::OnToggleVariablePublic(const FString& InVariable
 		return false;
 	}
 
-	if (FRigVMAssetInterfacePtr Blueprint = Editor->GetRigVMAssetInterface())
+	if (FRigVMEditorAssetInterfacePtr Blueprint = Editor->GetRigVMAssetInterface())
 	{
 		// Toggle the flag on the blueprint's version of the variable description, based on state
 	   const bool bVariableIsExposed = OnIsVariablePublic(InVariableName);
@@ -1379,7 +1429,7 @@ void SRigVMEditorGraphExplorer::OnCopy()
 
 	FString OutputString;
 	const URigVMEdGraphSchema* Schema = GetDefault<URigVMEdGraphSchema>();
-	FRigVMAssetInterfacePtr Blueprint = Editor->GetRigVMAssetInterface();
+	FRigVMEditorAssetInterfacePtr Blueprint = Editor->GetRigVMAssetInterface();
 
 	switch (Selection[0].Type)
 	{
@@ -1566,11 +1616,12 @@ void SRigVMEditorGraphExplorer::OnPasteVariable()
 	{
 		return;
 	}
-	FRigVMAssetInterfacePtr Blueprint = Editor->GetRigVMAssetInterface();
+	FRigVMEditorAssetInterfacePtr Blueprint = Editor->GetRigVMAssetInterface();
 	
 	FString ClipboardText;
 	FPlatformApplicationMisc::ClipboardPaste(ClipboardText);
 	Blueprint->OnPasteVariable(ClipboardText);
+	Refresh();
 }
 
 bool SRigVMEditorGraphExplorer::CanPasteVariable() const
@@ -1593,6 +1644,17 @@ bool SRigVMEditorGraphExplorer::CanPasteVariable() const
 		return Errors.IsEmpty();
 	}
 
+	if (ClipboardText.StartsWith(TEXT("RigVMVar"), ESearchCase::CaseSensitive))
+	{
+		FRigVMGraphVariableDescription Description;
+
+		FStringOutputDevice Errors;
+		const TCHAR* Import = ClipboardText.GetCharArray().GetData() + FCString::Strlen(TEXT("RigVMVar"));
+		FRigVMGraphVariableDescription::StaticStruct()->ImportText(Import, &Description, nullptr, 0, &Errors, FRigVMGraphVariableDescription::StaticStruct()->GetName());
+
+		return Errors.IsEmpty();
+	}
+
 	return false;
 }
 
@@ -1603,23 +1665,40 @@ void SRigVMEditorGraphExplorer::OnPasteLocalVariable()
 	{
 		return;
 	}
-	FRigVMAssetInterfacePtr Blueprint = Editor->GetRigVMAssetInterface();
+	FRigVMEditorAssetInterfacePtr Blueprint = Editor->GetRigVMAssetInterface();
 	
 	FString ClipboardText;
 	FPlatformApplicationMisc::ClipboardPaste(ClipboardText);
-	if (!ensure(ClipboardText.StartsWith(TEXT("BPVar"), ESearchCase::CaseSensitive)))
+	if (!ensure(ClipboardText.StartsWith(TEXT("BPVar"), ESearchCase::CaseSensitive) ||
+							ClipboardText.StartsWith(TEXT("RigVMVar"), ESearchCase::CaseSensitive)))
 	{
 		return;
 	}
 
-	FBPVariableDescription Description;
-	FStringOutputDevice Errors;
-	const TCHAR* Import = ClipboardText.GetCharArray().GetData() + FCString::Strlen(TEXT("BPVar"));
-	FBPVariableDescription::StaticStruct()->ImportText(Import, &Description, nullptr, PPF_None, &Errors, FBPVariableDescription::StaticStruct()->GetName());
-	if (Errors.IsEmpty())
+	if (ClipboardText.StartsWith(TEXT("BPVar"), ESearchCase::CaseSensitive))
 	{
-		Editor->OnPasteNewLocalVariable(Description);
-		Refresh();
+		FStringOutputDevice Errors;
+		FBPVariableDescription Description;
+		const TCHAR* Import = ClipboardText.GetCharArray().GetData() + FCString::Strlen(TEXT("BPVar"));
+		FBPVariableDescription::StaticStruct()->ImportText(Import, &Description, nullptr, PPF_None, &Errors, FBPVariableDescription::StaticStruct()->GetName());
+		if (Errors.IsEmpty())
+		{
+			Editor->OnPasteNewLocalVariable(Description);
+			Refresh();
+		}
+	}
+	else
+	{
+		FStringOutputDevice Errors;
+		FRigVMGraphVariableDescription Description;
+		const TCHAR* Import = ClipboardText.GetCharArray().GetData() + FCString::Strlen(TEXT("RigVMVar"));
+		FRigVMGraphVariableDescription::StaticStruct()->ImportText(Import, &Description, nullptr, PPF_None, &Errors, FRigVMGraphVariableDescription::StaticStruct()->GetName());
+		if (Errors.IsEmpty())
+		{
+			ensure(false);
+			//Editor->OnPasteNewLocalVariable(Description);
+			Refresh();
+		}
 	}
 }
 
@@ -1799,7 +1878,7 @@ void SRigVMEditorGraphExplorer::OnDeleteEntry()
 
 	if (Editor)
 	{
-		FRigVMAssetInterfacePtr Blueprint = Editor->GetRigVMAssetInterface();
+		FRigVMEditorAssetInterfacePtr Blueprint = Editor->GetRigVMAssetInterface();
 		TArray<FRigVMExplorerElementKey> Selected = TreeView->GetSelectedKeys();
 		const URigVMEdGraphSchema* Schema = GetDefault<URigVMEdGraphSchema>();
 		for (const FRigVMExplorerElementKey& Key : Selected)
@@ -2047,7 +2126,7 @@ bool SRigVMEditorGraphExplorer::OnRenameVariable(const FRigVMExplorerElementKey&
 
 	const FScopedTransaction Transaction( LOCTEXT( "RenameVariable", "Rename Variable" ));
 
-	FRigVMAssetInterfacePtr Blueprint = Editor->GetRigVMAssetInterface();
+	FRigVMEditorAssetInterfacePtr Blueprint = Editor->GetRigVMAssetInterface();
 
 	if (InOldKey.Type == ERigVMExplorerElementType::Variable)
 	{
@@ -2064,6 +2143,7 @@ bool SRigVMEditorGraphExplorer::OnRenameVariable(const FRigVMExplorerElementKey&
 			}
 		}
 	}
+	Refresh();
 	return true;
 }
 
@@ -2111,7 +2191,7 @@ void SRigVMEditorGraphExplorer::OnFindReferences(bool bSearchInAllBlueprints)
 
 	FString SearchTerm;
 
-	const TScriptInterface<IRigVMAssetInterface> Asset = PinnedEditor->GetRigVMAssetInterface();
+	const TScriptInterface<IRigVMEditorAssetInterface> Asset = PinnedEditor->GetRigVMAssetInterface();
 	const URigVMEdGraphSchema* Schema = GetDefault<URigVMEdGraphSchema>();
 	if (!Asset || !Schema)
 	{
@@ -2158,7 +2238,7 @@ bool SRigVMEditorGraphExplorer::CanFindReferences() const
 
 void SRigVMEditorGraphExplorer::OnRemoveUnusedFunctions()
 {
-	const TScriptInterface<IRigVMAssetInterface> Asset = RigVMEditor.IsValid() ? RigVMEditor.Pin()->GetRigVMAssetInterface() : nullptr;
+	const TScriptInterface<IRigVMEditorAssetInterface> Asset = RigVMEditor.IsValid() ? RigVMEditor.Pin()->GetRigVMAssetInterface() : nullptr;
 	if (Asset)
 	{
 		UE::RigVMEditor::RemoveUnusedMembersFromAsset<URigVMFunctionReferenceNode>(
@@ -2177,7 +2257,7 @@ bool SRigVMEditorGraphExplorer::CanRemoveUnusedFunctions() const
 
 void SRigVMEditorGraphExplorer::OnRemoveUnusedVariables()
 {
-	const TScriptInterface<IRigVMAssetInterface> Asset = RigVMEditor.IsValid() ? RigVMEditor.Pin()->GetRigVMAssetInterface() : nullptr;
+	const TScriptInterface<IRigVMEditorAssetInterface> Asset = RigVMEditor.IsValid() ? RigVMEditor.Pin()->GetRigVMAssetInterface() : nullptr;
 	if (Asset)
 	{
 		UE::RigVMEditor::RemoveUnusedMembersFromAsset<URigVMVariableNode>(
diff --git a/Engine/Plugins/Runtime/RigVM/Source/RigVMEditor/Private/Widgets/SRigVMEditorGraphExplorerTreeView.cpp b/Engine/Plugins/Runtime/RigVM/Source/RigVMEditor/Private/Widgets/SRigVMEditorGraphExplorerTreeView.cpp
--- a/Engine/Plugins/Runtime/RigVM/Source/RigVMEditor/Private/Widgets/SRigVMEditorGraphExplorerTreeView.cpp
+++ b/Engine/Plugins/Runtime/RigVM/Source/RigVMEditor/Private/Widgets/SRigVMEditorGraphExplorerTreeView.cpp
@@ -119,25 +119,25 @@ void SRigVMEditorGraphExplorerItem::Construct(const FArguments& InArgs,
 		TArray<TSharedPtr<IPinTypeSelectorFilter>> CustomPinTypeFilters = Delegates.GetCustomPinFilters();
 
 		STableRow<TSharedPtr<FRigVMEditorGraphExplorerTreeElement>>::Construct(
-		   STableRow<TSharedPtr<FRigVMEditorGraphExplorerTreeElement>>::FArguments()
-		   .ShowWires(false)
-		   .OnDragDetected(Delegates.OnDragDetected)
-		   .Content()
-		   [
-			   SNew(SHorizontalBox)
-			   +SHorizontalBox::Slot()
-				   .FillWidth(0.6f)
-				   .VAlign(VAlign_Center)
-				   .Padding(3.0f, 0.0f)
-				   [
-					   CreateTextSlotWidget(InElement->Key, InTreeView->FilterText)
-				   ]
+			STableRow<TSharedPtr<FRigVMEditorGraphExplorerTreeElement>>::FArguments()
+			.ShowWires(false)
+			.OnDragDetected(Delegates.OnDragDetected)
+			.Content()
+			[
+				SNew(SHorizontalBox)
+				+SHorizontalBox::Slot()
+					.FillWidth(0.6f)
+					.VAlign(VAlign_Center)
+					.Padding(3.0f, 0.0f)
+					[
+						CreateTextSlotWidget(InElement->Key, InTreeView->FilterText)
+					]
 
 				+SHorizontalBox::Slot()
-				   .FillWidth(0.4f)
-				   .HAlign(HAlign_Left)
-				   .VAlign(VAlign_Center)
-				   [
+					.FillWidth(0.4f)
+					.HAlign(HAlign_Left)
+					.VAlign(VAlign_Center)
+					[
 						SNew(SPinTypeSelector, FGetPinTypeTree::CreateUObject(GetDefault<UEdGraphSchema_K2>(), &UEdGraphSchema_K2::GetVariableTypeTree))
 						.Schema(GetDefault<URigVMEdGraphSchema>())
 						.TargetPinType_Lambda([this, InElement]() { return Delegates.GetVariablePinType(InElement->Key); })
@@ -145,47 +145,47 @@ void SRigVMEditorGraphExplorerItem::Construct(const FArguments& InArgs,
 						.TypeTreeFilter(ETypeTreeFilter::None)
 						.SelectorType(SPinTypeSelector::ESelectorType::Partial)
 						.CustomFilters(CustomPinTypeFilters)
-				   ]
+					]
 
 				+SHorizontalBox::Slot()
-				   .AutoWidth()
-				   .Padding(FMargin(6.0f, 0.0f, 3.0f, 0.0f))
-				   .HAlign(HAlign_Right)
-				   .VAlign(VAlign_Center)
-				   [
-					   SNew(SBorder)
-					   .Padding(0.0f)
-					   .Visibility(InElement->Key.Type == ERigVMExplorerElementType::Variable ? EVisibility::Visible : EVisibility::Collapsed)
-					   .BorderImage(FStyleDefaults::GetNoBrush())
-					   [
-						   SNew(SCheckBox)
-						   .ToolTipText_Lambda([this, InElement]()
-						   {
-						   	FText ToolTipText = FText::GetEmpty();
-						   	if (Delegates.IsVariablePublic(InElement->Key.Name))
-							{
-								ToolTipText = LOCTEXT("VariablePrivacy_is_public_Tooltip", "Variable is public and is editable on each instance of this Blueprint.");
-							}
-							else
+					.AutoWidth()
+					.Padding(FMargin(6.0f, 0.0f, 3.0f, 0.0f))
+					.HAlign(HAlign_Right)
+					.VAlign(VAlign_Center)
+					[
+						SNew(SBorder)
+						.Padding(0.0f)
+						.Visibility(InElement->Key.Type == ERigVMExplorerElementType::Variable ? EVisibility::Visible : EVisibility::Collapsed)
+						.BorderImage(FStyleDefaults::GetNoBrush())
+						[
+							SNew(SCheckBox)
+							.ToolTipText_Lambda([this, InElement]()
 							{
-								ToolTipText = LOCTEXT("VariablePrivacy_not_public_Tooltip", "Variable is not public and will not be editable on an instance of this Blueprint.");
-							}
-							return ToolTipText;
-						   })
-						   .OnCheckStateChanged_Lambda([this, InElement](ECheckBoxState) { Delegates.ToggleVariablePublic(InElement->Key.Name); })
-						   .IsChecked_Lambda([this, InElement]() { return Delegates.IsVariablePublic(InElement->Key.Name) ? ECheckBoxState::Checked : ECheckBoxState::Unchecked; })
-						   .Style(FAppStyle::Get(), "TransparentCheckBox")
-						   [
-							   SNew(SImage)
-							   .Image_Lambda([this, InElement]()
-							   {
-								   if (Delegates.IsVariablePublic(InElement->Key.Name))
-								   {
-									   return FAppStyle::GetBrush( "Kismet.VariableList.HideForInstance" );
-								   }
-							   		return FAppStyle::GetBrush( "Kismet.VariableList.ExposeForInstance" );
-							   })
-							   .ColorAndOpacity(FSlateColor::UseForeground())
+								FText ToolTipText = FText::GetEmpty();
+								if (Delegates.IsVariablePublic(InElement->Key.Name))
+								{
+									ToolTipText = LOCTEXT("VariablePrivacy_is_public_Tooltip", "Variable is public and is editable on each instance of this Blueprint.");
+								}
+								else
+								{
+									ToolTipText = LOCTEXT("VariablePrivacy_not_public_Tooltip", "Variable is not public and will not be editable on an instance of this Blueprint.");
+								}
+								return ToolTipText;
+							})
+							.OnCheckStateChanged_Lambda([this, InElement](ECheckBoxState) { Delegates.ToggleVariablePublic(InElement->Key.Name); })
+							.IsChecked_Lambda([this, InElement]() { return Delegates.IsVariablePublic(InElement->Key.Name) ? ECheckBoxState::Checked : ECheckBoxState::Unchecked; })
+							.Style(FAppStyle::Get(), "TransparentCheckBox")
+							[
+								SNew(SImage)
+								.Image_Lambda([this, InElement]()
+								{
+									if (Delegates.IsVariablePublic(InElement->Key.Name))
+									{
+										return FAppStyle::GetBrush( "Kismet.VariableList.ExposeForInstance" );
+									}
+									return FAppStyle::GetBrush( "Kismet.VariableList.HideForInstance" );
+								})
+								.ColorAndOpacity(FSlateColor::UseForeground())
 						   ]
 					   ]
 				   ]
@@ -321,29 +321,26 @@ FText SRigVMEditorGraphExplorerItem::GetDisplayText() const
 FText SRigVMEditorGraphExplorerItem::GetItemTooltip() const
 {
 	FText Tooltip;
-	if (!WeakExplorerElement.IsValid())
+	if (TSharedPtr<FRigVMEditorGraphExplorerTreeElement> Element = WeakExplorerElement.Pin())
 	{
-		return Tooltip;
-	}
-	TSharedPtr<FRigVMEditorGraphExplorerTreeElement> Element = WeakExplorerElement.Pin();
-
-	switch (Element->Key.Type)
-	{
-		case ERigVMExplorerElementType::Graph:
+		switch (Element->Key.Type)
 		{
-			return Delegates.GetGraphTooltip(Element->Key.Name);
-		}
+		case ERigVMExplorerElementType::Graph:
+			{
+				return Delegates.GetGraphTooltip(Element->Key.Name);
+			}
 		case ERigVMExplorerElementType::Function:
-		{
-			return Delegates.GetFunctionTooltip(Element->Key.Name);
-		}
+			{
+				return Delegates.GetFunctionTooltip(Element->Key.Name);
+			}
 		case ERigVMExplorerElementType::Variable:
-		{
-			return Delegates.GetVariableTooltip(Element->Key.Name);
-		}
+			{
+				return Delegates.GetVariableTooltip(Element->Key.Name);
+			}
 		case ERigVMExplorerElementType::Event:
-		{
-			return Delegates.GetEventDisplayName(Element->Key.Name);
+			{
+				return Delegates.GetEventDisplayName(Element->Key.Name);
+			}
 		}
 	}
 	
diff --git a/Engine/Plugins/Runtime/RigVM/Source/RigVMEditor/Private/Widgets/SRigVMEditorSelectedDebugObjectWidget.cpp b/Engine/Plugins/Runtime/RigVM/Source/RigVMEditor/Private/Widgets/SRigVMEditorSelectedDebugObjectWidget.cpp
--- a/Engine/Plugins/Runtime/RigVM/Source/RigVMEditor/Private/Widgets/SRigVMEditorSelectedDebugObjectWidget.cpp
+++ b/Engine/Plugins/Runtime/RigVM/Source/RigVMEditor/Private/Widgets/SRigVMEditorSelectedDebugObjectWidget.cpp
@@ -106,7 +106,7 @@ void SRigVMEditorSelectedDebugObjectWidget::Construct(const FArguments& InArgs,
 
 void SRigVMEditorSelectedDebugObjectWidget::Tick( const FGeometry& AllottedGeometry, const double InCurrentTime, const float InDeltaTime )
 {
-	if (FRigVMAssetInterfacePtr Blueprint = GetRigVMAssetInterface())
+	if (FRigVMEditorAssetInterfacePtr Blueprint = GetRigVMAssetInterface())
 	{
 		if (UObject* Object = Blueprint->GetObjectBeingDebugged())
 		{
@@ -745,7 +745,7 @@ void SRigVMEditorSelectedDebugObjectWidget::AddDebugObject(UObject* TestObject,
 	DebugObjects.Add(NewInstance);
 }
 
-FRigVMAssetInterfacePtr SRigVMEditorSelectedDebugObjectWidget::GetRigVMAssetInterface() const
+FRigVMEditorAssetInterfacePtr SRigVMEditorSelectedDebugObjectWidget::GetRigVMAssetInterface() const
 {
 	return Editor.Pin()->GetRigVMAssetInterface();
 }
@@ -775,7 +775,7 @@ FText SRigVMEditorSelectedDebugObjectWidget::GetSelectedDebugObjectTextLabel() c
 	{
 		Label = DebugInstance->ObjectLabel;
 
-		FRigVMAssetInterfacePtr Blueprint = GetRigVMAssetInterface();
+		FRigVMEditorAssetInterfacePtr Blueprint = GetRigVMAssetInterface();
 		if (Blueprint != nullptr)
 		{
 			UObject* DebugObj = Blueprint->GetObjectBeingDebugged();
diff --git a/Engine/Plugins/Runtime/RigVM/Source/RigVMEditor/Private/Widgets/SRigVMExecutionStackView.cpp b/Engine/Plugins/Runtime/RigVM/Source/RigVMEditor/Private/Widgets/SRigVMExecutionStackView.cpp
--- a/Engine/Plugins/Runtime/RigVM/Source/RigVMEditor/Private/Widgets/SRigVMExecutionStackView.cpp
+++ b/Engine/Plugins/Runtime/RigVM/Source/RigVMEditor/Private/Widgets/SRigVMExecutionStackView.cpp
@@ -45,7 +45,7 @@ TSharedRef<ITableRow> FRigStackEntry::MakeTreeRowWidget(const TSharedRef<STableV
 	return SNew(SRigStackItem, InOwnerTable, InEntry, InCommandList, InBlueprint.Get());
 }
 
-TSharedRef<ITableRow> FRigStackEntry::MakeTreeRowWidget(const TSharedRef<STableViewBase>& InOwnerTable, TSharedRef<FRigStackEntry> InEntry, TSharedRef<FUICommandList> InCommandList, TSharedPtr<SRigVMExecutionStackView> InStackView, TWeakInterfacePtr<IRigVMAssetInterface> InBlueprint)
+TSharedRef<ITableRow> FRigStackEntry::MakeTreeRowWidget(const TSharedRef<STableViewBase>& InOwnerTable, TSharedRef<FRigStackEntry> InEntry, TSharedRef<FUICommandList> InCommandList, TSharedPtr<SRigVMExecutionStackView> InStackView, TWeakInterfacePtr<IRigVMEditorAssetInterface> InBlueprint)
 {
 	return SNew(SRigStackItem, InOwnerTable, InEntry, InCommandList, InBlueprint);
 }
@@ -53,7 +53,7 @@ TSharedRef<ITableRow> FRigStackEntry::MakeTreeRowWidget(const TSharedRef<STableV
 //////////////////////////////////////////////////////////////
 /// SRigStackItem
 ///////////////////////////////////////////////////////////
-void SRigStackItem::Construct(const FArguments& InArgs, const TSharedRef<STableViewBase>& OwnerTable, TSharedRef<FRigStackEntry> InStackEntry, TSharedRef<FUICommandList> InCommandList, TWeakInterfacePtr<IRigVMAssetInterface> InBlueprint)
+void SRigStackItem::Construct(const FArguments& InArgs, const TSharedRef<STableViewBase>& OwnerTable, TSharedRef<FRigStackEntry> InStackEntry, TSharedRef<FUICommandList> InCommandList, TWeakInterfacePtr<IRigVMEditorAssetInterface> InBlueprint)
 {
 	WeakStackEntry = InStackEntry;
 	WeakBlueprint = InBlueprint;
@@ -296,7 +296,7 @@ void SRigVMExecutionStackView::Construct( const FArguments& InArgs, TSharedRef<I
 
 	if (TStrongObjectPtr<UObject> StrongPtr = WeakRigVMBlueprint.GetWeakObjectPtr().Pin())
 	{
-		IRigVMAssetInterface* RigVMBlueprint = WeakRigVMBlueprint.Get();
+		IRigVMEditorAssetInterface* RigVMBlueprint = WeakRigVMBlueprint.Get();
 		if (OnVMCompiledHandle.IsValid())
 		{
 			RigVMBlueprint->OnVMCompiled().Remove(OnVMCompiledHandle);
@@ -335,7 +335,7 @@ void SRigVMExecutionStackView::OnSelectionChanged(TSharedPtr<FRigStackEntry> Sel
 	{
 		if (TStrongObjectPtr<UObject> StrongPtr = WeakRigVMBlueprint.GetWeakObjectPtr().Pin())
 		{
-			IRigVMAssetInterface* RigVMBlueprint = WeakRigVMBlueprint.Get();
+			IRigVMEditorAssetInterface* RigVMBlueprint = WeakRigVMBlueprint.Get();
 			URigVMHost* Host = PreviouslyDebuggedHost.Get();
 			if (Host == nullptr || Host->GetVM() == nullptr)
 			{
@@ -392,7 +392,7 @@ void SRigVMExecutionStackView::BindCommands()
 	CommandList->MapAction(Commands.SelectTargetInstructions, FExecuteAction::CreateSP(this, &SRigVMExecutionStackView::HandleSelectTargetInstructions));
 }
 
-TSharedRef<ITableRow> SRigVMExecutionStackView::MakeTableRowWidget(TSharedPtr<FRigStackEntry> InItem, const TSharedRef<STableViewBase>& OwnerTable, TWeakInterfacePtr<IRigVMAssetInterface> InBlueprint)
+TSharedRef<ITableRow> SRigVMExecutionStackView::MakeTableRowWidget(TSharedPtr<FRigStackEntry> InItem, const TSharedRef<STableViewBase>& OwnerTable, TWeakInterfacePtr<IRigVMEditorAssetInterface> InBlueprint)
 {
 	return InItem->MakeTreeRowWidget(OwnerTable, InItem.ToSharedRef(), CommandList.ToSharedRef(), SharedThis(this), InBlueprint);
 }
@@ -842,7 +842,7 @@ void SRigVMExecutionStackView::HandleFocusOnSelectedGraphNode()
 		{
 			if (TStrongObjectPtr<UObject> StrongPtr = WeakRigVMBlueprint.GetWeakObjectPtr().Pin())
 			{
-				IRigVMAssetInterface* RigVMBlueprint = WeakRigVMBlueprint.Get();
+				IRigVMEditorAssetInterface* RigVMBlueprint = WeakRigVMBlueprint.Get();
 				URigVMHost* Host = PreviouslyDebuggedHost.Get();
 				if (Host == nullptr || Host->GetVM() == nullptr)
 				{
@@ -1160,7 +1160,7 @@ void SRigVMExecutionStackView::HandleHostInitializedEvent(URigVMHost* InHost, co
 			InHost->OnInitialized_AnyThread().Remove(OnHostInitializedHandle);
 			OnHostInitializedHandle.Reset();
 
-			if (FRigVMAssetInterfacePtr RigBlueprint = Editor->GetRigVMAssetInterface())
+			if (FRigVMEditorAssetInterfacePtr RigBlueprint = Editor->GetRigVMAssetInterface())
 			{
 				TArray<URigVMGraph*> Models = RigBlueprint->GetAllModels();
 				for (URigVMGraph* Model : Models)
@@ -1200,7 +1200,7 @@ void SRigVMExecutionStackView::HandleHostExecutedEvent(URigVMHost* InHost, const
 
 		if (WeakRigVMBlueprint.IsValid())
 		{
-			IRigVMAssetInterface* Blueprint = WeakRigVMBlueprint.Get();
+			IRigVMEditorAssetInterface* Blueprint = WeakRigVMBlueprint.Get();
 			check(Blueprint);
 			
 			if(Blueprint->GetRigGraphDisplaySettings().bShowNodeRunCounts)
@@ -1275,7 +1275,7 @@ void SRigVMExecutionStackView::HandleItemMouseDoubleClick(TSharedPtr<FRigStackEn
 	{
 		if (TStrongObjectPtr<UObject> StrongPtr = WeakRigVMBlueprint.GetWeakObjectPtr().Pin())
 		{
-			IRigVMAssetInterface* RigVMBlueprint = WeakRigVMBlueprint.Get();
+			IRigVMEditorAssetInterface* RigVMBlueprint = WeakRigVMBlueprint.Get();
 			URigVMHost* Host = PreviouslyDebuggedHost.Get();
 			if (!Host || !Host->GetVM())
 			{
diff --git a/Engine/Plugins/Runtime/RigVM/Source/RigVMEditor/Private/Widgets/SRigVMGraphFunctionBulkEditWidget.cpp b/Engine/Plugins/Runtime/RigVM/Source/RigVMEditor/Private/Widgets/SRigVMGraphFunctionBulkEditWidget.cpp
--- a/Engine/Plugins/Runtime/RigVM/Source/RigVMEditor/Private/Widgets/SRigVMGraphFunctionBulkEditWidget.cpp
+++ b/Engine/Plugins/Runtime/RigVM/Source/RigVMEditor/Private/Widgets/SRigVMGraphFunctionBulkEditWidget.cpp
@@ -23,7 +23,7 @@
 /// SRigVMGraphFunctionBulkEditWidget
 ///////////////////////////////////////////////////////////
 
-void SRigVMGraphFunctionBulkEditWidget::Construct(const FArguments& InArgs, FRigVMAssetInterfacePtr InBlueprint, URigVMController* InController, URigVMLibraryNode* InFunction, ERigVMControllerBulkEditType InEditType)
+void SRigVMGraphFunctionBulkEditWidget::Construct(const FArguments& InArgs, FRigVMEditorAssetInterfacePtr InBlueprint, URigVMController* InController, URigVMLibraryNode* InFunction, ERigVMControllerBulkEditType InEditType)
 {
 	Blueprint = InBlueprint;
 	Controller = InController;
diff --git a/Engine/Plugins/Runtime/RigVM/Source/RigVMEditor/Private/Widgets/SRigVMGraphNode.cpp b/Engine/Plugins/Runtime/RigVM/Source/RigVMEditor/Private/Widgets/SRigVMGraphNode.cpp
--- a/Engine/Plugins/Runtime/RigVM/Source/RigVMEditor/Private/Widgets/SRigVMGraphNode.cpp
+++ b/Engine/Plugins/Runtime/RigVM/Source/RigVMEditor/Private/Widgets/SRigVMGraphNode.cpp
@@ -1104,7 +1104,7 @@ FSlateColor SRigVMGraphNode::GetVariableLabelTextColor(
 	return FLinearColor::White;
 }
 
-FText SRigVMGraphNode::GetVariableLabelTooltipText(TWeakInterfacePtr<IRigVMAssetInterface> InBlueprint,
+FText SRigVMGraphNode::GetVariableLabelTooltipText(TWeakInterfacePtr<IRigVMEditorAssetInterface> InBlueprint,
 	FName InVariableName) const
 {
 	if(InBlueprint.IsValid())
@@ -1147,13 +1147,23 @@ void SRigVMGraphNode::GetNodeInfoPopups(FNodeInfoContext* Context, TArray<FGraph
 
 	URigVMHost* ActiveObject = Cast<URigVMHost>(K2Context->ActiveObjectBeingDebugged);
 	URigVMEdGraphNode* RigGraphNode = Cast<URigVMEdGraphNode>(GraphNode);
-	FRigVMAssetInterfacePtr RigBlueprint = K2Context->SourceBlueprint;
+	FRigVMEditorAssetInterfacePtr RigAsset = K2Context->SourceBlueprint;
+
+	if (!RigAsset && RigGraphNode)
+	{
+		RigAsset = RigGraphNode->GetAsset();
+	}
+
+	if (!ActiveObject && RigAsset)
+	{
+		ActiveObject = RigAsset->GetDebuggedRigVMHost();
+	}
 
 	// Display any pending latent actions
-	if (ActiveObject && RigBlueprint && RigGraphNode)
+	if (ActiveObject && RigAsset && RigGraphNode)
 	{
 		// Display pinned watches
-		if (K2Context->WatchedNodeSet.Contains(GraphNode))
+		if (RigAsset->NodeContainsWatchedPins(GraphNode))
 		{
 			const UEdGraphSchema* Schema = GraphNode->GetSchema();
 
@@ -1163,7 +1173,7 @@ void SRigVMGraphNode::GetNodeInfoPopups(FNodeInfoContext* Context, TArray<FGraph
 			for (int32 PinIndex = 0; PinIndex < GraphNode->Pins.Num(); ++PinIndex)
 			{
 				UEdGraphPin* WatchPin = GraphNode->Pins[PinIndex];
-				if (K2Context->WatchedPinSet.Contains(WatchPin))
+				if (RigAsset->IsPinBeingWatched(WatchPin))
 				{
 					if (URigVMPin* ModelPin = RigGraphNode->GetModel()->FindPin(WatchPin->GetName()))
 					{
@@ -1199,7 +1209,7 @@ void SRigVMGraphNode::GetNodeInfoPopups(FNodeInfoContext* Context, TArray<FGraph
 						if(DefaultValues.IsEmpty())
 						{
 							FString PinHash = URigVMCompiler::GetPinHash(ModelPin, nullptr, true);
-							if (const FRigVMOperand* WatchOperand = RigBlueprint->PinToOperandMap.Find(PinHash))
+							if (const FRigVMOperand* WatchOperand = RigAsset->PinToOperandMap.Find(PinHash))
 							{
 								FRigVMMemoryStorageStruct* Memory = ActiveObject->GetDebugMemory();
 								// We mark PPF_ExternalEditor so that default values are also printed
@@ -2888,7 +2898,7 @@ void SRigVMGraphNode::UpdatePinTreeView()
 	if(URigVMFunctionReferenceNode* FunctionReferenceNode = Cast<URigVMFunctionReferenceNode>(ModelNode))
 	{
 		TWeakObjectPtr<URigVMFunctionReferenceNode> WeakFunctionReferenceNode = FunctionReferenceNode;
-		TWeakInterfacePtr<IRigVMAssetInterface> WeakRigVMBlueprint = Blueprint;
+		TWeakInterfacePtr<IRigVMEditorAssetInterface> WeakRigVMBlueprint = Blueprint;
 
 		// add the entries for the variable remapping
 		for(const TSharedPtr<FRigVMExternalVariable>& ExternalVariable : RigGraphNode->ExternalVariables)
diff --git a/Engine/Plugins/Runtime/RigVM/Source/RigVMEditor/Private/Widgets/SRigVMGraphNodeComment.cpp b/Engine/Plugins/Runtime/RigVM/Source/RigVMEditor/Private/Widgets/SRigVMGraphNodeComment.cpp
--- a/Engine/Plugins/Runtime/RigVM/Source/RigVMEditor/Private/Widgets/SRigVMGraphNodeComment.cpp
+++ b/Engine/Plugins/Runtime/RigVM/Source/RigVMEditor/Private/Widgets/SRigVMGraphNodeComment.cpp
@@ -40,7 +40,7 @@ FReply SRigVMGraphNodeComment::OnMouseButtonUp(const FGeometry& MyGeometry, cons
 			UEdGraphNode_Comment* CommentNode = CastChecked<UEdGraphNode_Comment>(GraphNode);
 			if (URigVMEdGraph* Graph = Cast<URigVMEdGraph>(CommentNode->GetOuter()))
 			{
-				if (FRigVMAssetInterfacePtr Blueprint = Graph->GetOuter())
+				if (FRigVMEditorAssetInterfacePtr Blueprint = Graph->GetOuter())
 				{
 					FVector2D Position(CommentNode->NodePosX, CommentNode->NodePosY);
 					FVector2D Size(CommentNode->NodeWidth, CommentNode->NodeHeight);
@@ -104,7 +104,7 @@ void SRigVMGraphNodeComment::Tick(const FGeometry& AllottedGeometry, const doubl
 		{
 			if (URigVMEdGraph* Graph = Cast<URigVMEdGraph>(CommentNode->GetOuter()))
 			{
-				if (FRigVMAssetInterfacePtr Blueprint = Graph->GetOuter())
+				if (FRigVMEditorAssetInterfacePtr Blueprint = Graph->GetOuter())
 				{
 					if (URigVMController* Controller = Blueprint->GetController(Graph))
 					{
@@ -126,7 +126,7 @@ void SRigVMGraphNodeComment::Tick(const FGeometry& AllottedGeometry, const doubl
 			{
 				if (URigVMEdGraph* Graph = Cast<URigVMEdGraph>(CommentNode->GetOuter()))
 				{
-					if (FRigVMAssetInterfacePtr Blueprint = Graph->GetOuter())
+					if (FRigVMEditorAssetInterfacePtr Blueprint = Graph->GetOuter())
 					{
 						if (URigVMController* Controller = Blueprint->GetController(Graph))
 						{
diff --git a/Engine/Plugins/Runtime/RigVM/Source/RigVMEditor/Private/Widgets/SRigVMGraphPinEnumPicker.cpp b/Engine/Plugins/Runtime/RigVM/Source/RigVMEditor/Private/Widgets/SRigVMGraphPinEnumPicker.cpp
--- a/Engine/Plugins/Runtime/RigVM/Source/RigVMEditor/Private/Widgets/SRigVMGraphPinEnumPicker.cpp
+++ b/Engine/Plugins/Runtime/RigVM/Source/RigVMEditor/Private/Widgets/SRigVMGraphPinEnumPicker.cpp
@@ -118,7 +118,7 @@ void SRigVMGraphPinEnumPicker::HandleControlEnumChanged(TSharedPtr<FString> InIt
 {
 	if(ModelPin)
 	{
-		if(FRigVMAssetInterfacePtr Blueprint = IRigVMAssetInterface::GetInterfaceOuter(ModelPin))
+		if(FRigVMEditorAssetInterfacePtr Blueprint = IRigVMEditorAssetInterface::GetInterfaceOuter(ModelPin))
 		{
 			if(URigVMController* Controller = Blueprint->GetOrCreateController(ModelPin->GetGraph()))
 			{
diff --git a/Engine/Plugins/Runtime/RigVM/Source/RigVMEditor/Private/Widgets/SRigVMGraphPinQuat.cpp b/Engine/Plugins/Runtime/RigVM/Source/RigVMEditor/Private/Widgets/SRigVMGraphPinQuat.cpp
--- a/Engine/Plugins/Runtime/RigVM/Source/RigVMEditor/Private/Widgets/SRigVMGraphPinQuat.cpp
+++ b/Engine/Plugins/Runtime/RigVM/Source/RigVMEditor/Private/Widgets/SRigVMGraphPinQuat.cpp
@@ -80,7 +80,7 @@ void SRigVMGraphPinQuat::OnRotatorCommitted(FRotator InRotator, ETextCommit::Typ
 		FString NewDefaultValue;	
 		TBaseStructure<FRotator>::Get()->ExportText(NewDefaultValue, &InRotator, &InRotator, nullptr, PPF_None, nullptr);
 
-		if(FRigVMAssetInterfacePtr Blueprint = IRigVMAssetInterface::GetInterfaceOuter(ModelPin))
+		if(FRigVMEditorAssetInterfacePtr Blueprint = IRigVMEditorAssetInterface::GetInterfaceOuter(ModelPin))
 		{
 			if(URigVMController* Controller = Blueprint->GetOrCreateController(ModelPin->GetGraph()))
 			{
diff --git a/Engine/Plugins/Runtime/RigVM/Source/RigVMEditor/Private/Widgets/SRigVMGraphPinUserDataPath.cpp b/Engine/Plugins/Runtime/RigVM/Source/RigVMEditor/Private/Widgets/SRigVMGraphPinUserDataPath.cpp
--- a/Engine/Plugins/Runtime/RigVM/Source/RigVMEditor/Private/Widgets/SRigVMGraphPinUserDataPath.cpp
+++ b/Engine/Plugins/Runtime/RigVM/Source/RigVMEditor/Private/Widgets/SRigVMGraphPinUserDataPath.cpp
@@ -14,6 +14,7 @@
 #include "EdGraph/RigVMEdGraphSchema.h"
 #include "Kismet2/BlueprintEditorUtils.h"
 #include "EdGraphSchema_K2.h"
+#include "RigVMDeveloperTypeUtils.h"
 
 #define LOCTEXT_NAMESPACE "SRigVMGraphPinUserDataPath"
 
@@ -270,7 +271,7 @@ void SRigVMUserDataPath::FillUserDataPathMenu(FMenuBuilder& InMenuBuilder, FStri
 
 void SRigVMUserDataPath::HandleSetUserDataPath(FString InUserDataPath)
 {
-	FRigVMAssetInterfacePtr Blueprint = GetRigVMAssetInterface();
+	FRigVMEditorAssetInterfacePtr Blueprint = GetRigVMAssetInterface();
 	if(ModelPins.IsEmpty() || (Blueprint == nullptr))
 	{
 		return;
@@ -321,11 +322,11 @@ void SRigVMUserDataPath::HandleSetUserDataPath(FString InUserDataPath)
 	}
 }
 
-FRigVMAssetInterfacePtr SRigVMUserDataPath::GetRigVMAssetInterface() const
+FRigVMEditorAssetInterfacePtr SRigVMUserDataPath::GetRigVMAssetInterface() const
 {
 	if(!ModelPins.IsEmpty())
 	{
-		return IRigVMAssetInterface::GetInterfaceOuter(ModelPins[0]);
+		return IRigVMEditorAssetInterface::GetInterfaceOuter(ModelPins[0]);
 	}
 	return nullptr;
 }
@@ -368,7 +369,7 @@ const UNameSpacedUserData* SRigVMUserDataPath::GetUserDataObject() const
 		return nullptr;
 	}
 	
-	if(const FRigVMAssetInterfacePtr Blueprint = GetRigVMAssetInterface())
+	if(const FRigVMEditorAssetInterfacePtr Blueprint = GetRigVMAssetInterface())
 	{
 		if(const UObject* DebuggedObject = Blueprint->GetObjectBeingDebugged())
 		{
diff --git a/Engine/Plugins/Runtime/RigVM/Source/RigVMEditor/Private/Widgets/SRigVMGraphPinVariableBinding.cpp b/Engine/Plugins/Runtime/RigVM/Source/RigVMEditor/Private/Widgets/SRigVMGraphPinVariableBinding.cpp
--- a/Engine/Plugins/Runtime/RigVM/Source/RigVMEditor/Private/Widgets/SRigVMGraphPinVariableBinding.cpp
+++ b/Engine/Plugins/Runtime/RigVM/Source/RigVMEditor/Private/Widgets/SRigVMGraphPinVariableBinding.cpp
@@ -12,6 +12,7 @@
 #include "DetailLayoutBuilder.h"
 #include "RigVMModel/RigVMController.h"
 #include "RigVMBlueprintGeneratedClass.h"
+#include "RigVMDeveloperTypeUtils.h"
 #include "EdGraph/RigVMEdGraphSchema.h"
 #include "Kismet2/BlueprintEditorUtils.h"
 
@@ -142,7 +143,8 @@ FLinearColor SRigVMGraphVariableBinding::GetBindingColor() const
 		{
 			if (VariableDescription.Name == BoundVariable)
 			{
-				return Schema->GetPinTypeColor(VariableDescription.ToPinType());
+				const FEdGraphPinType EdPinType = RigVMTypeUtils::PinTypeFromRigVMVariableDescription(VariableDescription);
+				return Schema->GetPinTypeColor(EdPinType);
 			}
 		}
 
diff --git a/Engine/Plugins/Runtime/RigVM/Source/RigVMEditor/Private/Widgets/SRigVMGraphPinVariableName.cpp b/Engine/Plugins/Runtime/RigVM/Source/RigVMEditor/Private/Widgets/SRigVMGraphPinVariableName.cpp
--- a/Engine/Plugins/Runtime/RigVM/Source/RigVMEditor/Private/Widgets/SRigVMGraphPinVariableName.cpp
+++ b/Engine/Plugins/Runtime/RigVM/Source/RigVMEditor/Private/Widgets/SRigVMGraphPinVariableName.cpp
@@ -115,7 +115,7 @@ TArray<TSharedPtr<FString>>& SRigVMGraphPinVariableName::GetVariableNames()
 				}
 			}
 
-			if (FRigVMAssetInterfacePtr Blueprint = ModelNode->GetGraph()->GetOuter())
+			if (FRigVMEditorAssetInterfacePtr Blueprint = ModelNode->GetGraph()->GetOuter())
 			{
 				for (FRigVMGraphVariableDescription& BPVariable : Blueprint->GetAssetVariables())
 				{
diff --git a/Engine/Plugins/Runtime/RigVM/Source/RigVMEditor/Private/Widgets/SRigVMLogWidget.cpp b/Engine/Plugins/Runtime/RigVM/Source/RigVMEditor/Private/Widgets/SRigVMLogWidget.cpp
--- a/Engine/Plugins/Runtime/RigVM/Source/RigVMEditor/Private/Widgets/SRigVMLogWidget.cpp
+++ b/Engine/Plugins/Runtime/RigVM/Source/RigVMEditor/Private/Widgets/SRigVMLogWidget.cpp
@@ -87,7 +87,7 @@ void SRigVMLogWidget::UnbindLog(TSharedPtr<IMessageLogListing> InListing)
 	}
 }
 
-void SRigVMLogWidget::BindLog(const FRigVMAssetInterfacePtr InRigVMBlueprint)
+void SRigVMLogWidget::BindLog(const FRigVMEditorAssetInterfacePtr InRigVMBlueprint)
 {
 	TSharedPtr<IMessageLogListing> Listing = nullptr;
 	FMessageLogModule& MessageLogModule = FModuleManager::LoadModuleChecked<FMessageLogModule>("MessageLog");
@@ -112,7 +112,7 @@ void SRigVMLogWidget::BindLog(const FRigVMAssetInterfacePtr InRigVMBlueprint)
 	}
 }
 
-void SRigVMLogWidget::UnbindLog(const FRigVMAssetInterfacePtr InRigVMBlueprint)
+void SRigVMLogWidget::UnbindLog(const FRigVMEditorAssetInterfacePtr InRigVMBlueprint)
 {
 	TSharedPtr<IMessageLogListing> Listing = nullptr;
 	FMessageLogModule& MessageLogModule = FModuleManager::LoadModuleChecked<FMessageLogModule>("MessageLog");
diff --git a/Engine/Plugins/Runtime/RigVM/Source/RigVMEditor/Private/Widgets/SRigVMSwapFunctionsWidget.cpp b/Engine/Plugins/Runtime/RigVM/Source/RigVMEditor/Private/Widgets/SRigVMSwapFunctionsWidget.cpp
--- a/Engine/Plugins/Runtime/RigVM/Source/RigVMEditor/Private/Widgets/SRigVMSwapFunctionsWidget.cpp
+++ b/Engine/Plugins/Runtime/RigVM/Source/RigVMEditor/Private/Widgets/SRigVMSwapFunctionsWidget.cpp
@@ -463,14 +463,14 @@ URigVMFunctionReferenceNode* FRigVMSwapFunctionTask::GetReferenceNode(const TSha
 	return ReferenceNode;
 }
 
-FRigVMAssetInterfacePtr FRigVMSwapFunctionTask::GetRigVMAsset(const TSharedRef<FRigVMTreePhase>& InPhase) const
+FRigVMEditorAssetInterfacePtr FRigVMSwapFunctionTask::GetRigVMAsset(const TSharedRef<FRigVMTreePhase>& InPhase) const
 {
 	URigVMFunctionReferenceNode* ReferenceNode = GetReferenceNode(InPhase);
 	if(ReferenceNode == nullptr)
 	{
 		return nullptr;
 	}
-	return IRigVMAssetInterface::GetInterfaceOuter(ReferenceNode);
+	return IRigVMEditorAssetInterface::GetInterfaceOuter(ReferenceNode);
 }
 
 void SRigVMSwapFunctionsWidget::Construct(const FArguments& InArgs)
@@ -828,7 +828,7 @@ FReply SRigVMSwapFunctionsWidget::OnNodeDoubleClicked(TSharedRef<FRigVMTreeNode>
 			// force load
 			if(UObject* TopLevelObject = AssetData.GetSoftObjectPath().GetWithoutSubPath().TryLoad())
 			{
-				if(TopLevelObject->GetClass()->ImplementsInterface(URigVMAssetInterface::StaticClass()))
+				if(TopLevelObject->GetClass()->ImplementsInterface(URigVMEditorAssetInterface::StaticClass()))
 				{
 					if(const UObject* ObjectReference = StaticLoadObject(UObject::StaticClass(), nullptr, *Node->GetPath(), nullptr))
 					{
diff --git a/Engine/Plugins/Runtime/RigVM/Source/RigVMEditor/Public/EdGraph/NodeSpawners/RigVMEdGraphEnumNodeSpawner.h b/Engine/Plugins/Runtime/RigVM/Source/RigVMEditor/Public/EdGraph/NodeSpawners/RigVMEdGraphEnumNodeSpawner.h
--- a/Engine/Plugins/Runtime/RigVM/Source/RigVMEditor/Public/EdGraph/NodeSpawners/RigVMEdGraphEnumNodeSpawner.h
+++ b/Engine/Plugins/Runtime/RigVM/Source/RigVMEditor/Public/EdGraph/NodeSpawners/RigVMEdGraphEnumNodeSpawner.h
@@ -13,7 +13,7 @@ class URigVMEdGraphEnumNodeSpawner : public URigVMEdGraphNodeSpawner
 {
 public:
 
-	virtual ~URigVMEdGraphEnumNodeSpawner() = default;
+	virtual ~URigVMEdGraphEnumNodeSpawner() {}
 
 	/**
 	 * Creates a new URigVMEdGraphEnumNodeSpawner
diff --git a/Engine/Plugins/Runtime/RigVM/Source/RigVMEditor/Public/EdGraph/NodeSpawners/RigVMEdGraphFunctionRefNodeSpawner.h b/Engine/Plugins/Runtime/RigVM/Source/RigVMEditor/Public/EdGraph/NodeSpawners/RigVMEdGraphFunctionRefNodeSpawner.h
--- a/Engine/Plugins/Runtime/RigVM/Source/RigVMEditor/Public/EdGraph/NodeSpawners/RigVMEdGraphFunctionRefNodeSpawner.h
+++ b/Engine/Plugins/Runtime/RigVM/Source/RigVMEditor/Public/EdGraph/NodeSpawners/RigVMEdGraphFunctionRefNodeSpawner.h
@@ -24,7 +24,7 @@ class URigVMEdGraphFunctionRefNodeSpawner : public URigVMEdGraphNodeSpawner
 
 public:
 
-	virtual ~URigVMEdGraphFunctionRefNodeSpawner() = default;
+	virtual ~URigVMEdGraphFunctionRefNodeSpawner() {}
 	
 	/**
 	 * Creates a new URigVMEdGraphFunctionRefNodeSpawner, charged with spawning a function reference
@@ -32,7 +32,7 @@ public:
 	 * @return A newly allocated instance of this class.
 	 */
 	UE_DEPRECATED(5.7, "Plase use URigVMEdGraphFunctionRefNodeSpawner(URigVMLibraryNode* InFunction)")
-	static UE_API URigVMEdGraphFunctionRefNodeSpawner* CreateFromFunction(URigVMLibraryNode* InFunction) { return nullptr; }
+	static URigVMEdGraphFunctionRefNodeSpawner* CreateFromFunction(URigVMLibraryNode* InFunction) { return nullptr; }
 	UE_API URigVMEdGraphFunctionRefNodeSpawner(URigVMLibraryNode* InFunction);
 
 	/**
@@ -41,9 +41,9 @@ public:
 	 * @return A newly allocated instance of this class.
 	 */
 	UE_DEPRECATED(5.7, "Please use URigVMEdGraphFunctionRefNodeSpawner(const FAssetData& InAssetData, const FRigVMGraphFunctionHeader& InPublicFunction)")
-	static UE_API URigVMEdGraphFunctionRefNodeSpawner* CreateFromAssetData(const FAssetData& InAssetData, const FRigVMGraphFunctionHeader& InPublicFunction) { return nullptr; }
+	static URigVMEdGraphFunctionRefNodeSpawner* CreateFromAssetData(const FAssetData& InAssetData, const FRigVMGraphFunctionHeader& InPublicFunction) { return nullptr; }
 	UE_DEPRECATED(5.7, "Please use URigVMEdGraphFunctionRefNodeSpawner(const FAssetData& InAssetData, const FRigVMOldPublicFunctionData& InPublicFunction)")
-	static UE_API URigVMEdGraphFunctionRefNodeSpawner* CreateFromAssetData(const FAssetData& InAssetData, const FRigVMOldPublicFunctionData& InPublicFunction) { return nullptr; }
+	static URigVMEdGraphFunctionRefNodeSpawner* CreateFromAssetData(const FAssetData& InAssetData, const FRigVMOldPublicFunctionData& InPublicFunction) { return nullptr; }
 	UE_API URigVMEdGraphFunctionRefNodeSpawner(const FAssetData& InAssetData, const FRigVMGraphFunctionHeader& InPublicFunction);
 	UE_API URigVMEdGraphFunctionRefNodeSpawner(const FAssetData& InAssetData, const FRigVMOldPublicFunctionData& InPublicFunction);
 
@@ -68,7 +68,7 @@ private:
 	mutable FSoftObjectPath AssetPath;
 	mutable FString FunctionPath;
 
-	static UE_API URigVMEdGraphNode* SpawnNode(UEdGraph* ParentGraph, FRigVMAssetInterfacePtr Blueprint, FRigVMGraphFunctionHeader& InFunction, FVector2D const Location);
+	static UE_API URigVMEdGraphNode* SpawnNode(UEdGraph* ParentGraph, FRigVMEditorAssetInterfacePtr Blueprint, FRigVMGraphFunctionHeader& InFunction, FVector2D const Location);
 
 	friend class UEngineTestControlRig;
 	friend class FRigVMEditorBase;
diff --git a/Engine/Plugins/Runtime/RigVM/Source/RigVMEditor/Public/EdGraph/NodeSpawners/RigVMEdGraphInvokeEntryNodeSpawner.h b/Engine/Plugins/Runtime/RigVM/Source/RigVMEditor/Public/EdGraph/NodeSpawners/RigVMEdGraphInvokeEntryNodeSpawner.h
--- a/Engine/Plugins/Runtime/RigVM/Source/RigVMEditor/Public/EdGraph/NodeSpawners/RigVMEdGraphInvokeEntryNodeSpawner.h
+++ b/Engine/Plugins/Runtime/RigVM/Source/RigVMEditor/Public/EdGraph/NodeSpawners/RigVMEdGraphInvokeEntryNodeSpawner.h
@@ -13,11 +13,11 @@ class URigVMEdGraphNode;
 class URigVMEdGraphInvokeEntryNodeSpawner : public URigVMEdGraphNodeSpawner
 {
 public:
-	virtual ~URigVMEdGraphInvokeEntryNodeSpawner() = default;
+	virtual ~URigVMEdGraphInvokeEntryNodeSpawner() {}
 
-	UE_DEPRECATED(5.7, "Plase use URigVMEdGraphInvokeEntryNodeSpawner(FRigVMAssetInterfacePtr InBlueprint, const FName& InEntryName, const FText& InMenuDesc, const FText& InCategory, const FText& InTooltip)")
-	static UE_API URigVMEdGraphInvokeEntryNodeSpawner* CreateForEntry(URigVMBlueprint* InBlueprint, const FName& InEntryName, const FText& InMenuDesc, const FText& InCategory, const FText& InTooltip) { return nullptr; }
-	UE_API URigVMEdGraphInvokeEntryNodeSpawner(FRigVMAssetInterfacePtr InBlueprint, const FName& InEntryName, const FText& InMenuDesc, const FText& InCategory, const FText& InTooltip);
+	UE_DEPRECATED(5.7, "Plase use URigVMEdGraphInvokeEntryNodeSpawner(FRigVMEditorAssetInterfacePtr InBlueprint, const FName& InEntryName, const FText& InMenuDesc, const FText& InCategory, const FText& InTooltip)")
+	static URigVMEdGraphInvokeEntryNodeSpawner* CreateForEntry(URigVMBlueprint* InBlueprint, const FName& InEntryName, const FText& InMenuDesc, const FText& InCategory, const FText& InTooltip) { return nullptr; }
+	UE_API URigVMEdGraphInvokeEntryNodeSpawner(FRigVMEditorAssetInterfacePtr InBlueprint, const FName& InEntryName, const FText& InMenuDesc, const FText& InCategory, const FText& InTooltip);
 
 	// URigVMEdGraphNodeSpawner interface
 	UE_API virtual FString GetSpawnerSignature() const override;
@@ -28,7 +28,7 @@ public:
 private:
 
 	/** The pin type we will spawn */
-	TWeakInterfacePtr<IRigVMAssetInterface> Blueprint;
+	TWeakInterfacePtr<IRigVMEditorAssetInterface> Blueprint;
 	TWeakObjectPtr<URigVMGraph> GraphOwner;
 	FName EntryName;
 
diff --git a/Engine/Plugins/Runtime/RigVM/Source/RigVMEditor/Public/EdGraph/NodeSpawners/RigVMEdGraphNodeSpawner.h b/Engine/Plugins/Runtime/RigVM/Source/RigVMEditor/Public/EdGraph/NodeSpawners/RigVMEdGraphNodeSpawner.h
--- a/Engine/Plugins/Runtime/RigVM/Source/RigVMEditor/Public/EdGraph/NodeSpawners/RigVMEdGraphNodeSpawner.h
+++ b/Engine/Plugins/Runtime/RigVM/Source/RigVMEditor/Public/EdGraph/NodeSpawners/RigVMEdGraphNodeSpawner.h
@@ -15,13 +15,11 @@ class URigVMEdGraphNode;
 
 class URigVMEdGraphNodeSpawner : public TSharedFromThis<URigVMEdGraphNodeSpawner>
 {
-protected:
-	virtual ~URigVMEdGraphNodeSpawner() = default;
-	
 public:
+	virtual ~URigVMEdGraphNodeSpawner() {}
 
 	UE_DEPRECATED(5.7, "Plase use SetRelatedBlueprintClass(UClass* InClass)")
-	UE_API void SetRelatedBlueprintClass(TSubclassOf<URigVMBlueprint> InClass){}
+	void SetRelatedBlueprintClass(TSubclassOf<URigVMBlueprint> InClass){}
 	UE_API void SetRelatedBlueprintClass(UClass* InClass);
 
 	struct FPinInfo
@@ -47,9 +45,16 @@ public:
 	};
 
 	UE_API virtual bool IsTemplateNodeFilteredOut(TArray<UObject*>& InAssets, TArray<UEdGraph*> InGraphs, TArray<UEdGraphPin*> InPins) const;
-	UE_API virtual FString GetSpawnerSignature() const { return FString(); }
-	UE_API virtual URigVMEdGraphNode* Invoke(URigVMEdGraph* ParentGraph, FVector2D const Location) const { return nullptr; }
+	virtual FString GetSpawnerSignature() const { return FString(); }
+	virtual URigVMEdGraphNode* Invoke(URigVMEdGraph* ParentGraph, FVector2D const Location) const { return nullptr; }
 	static UE_API URigVMEdGraphNode* SpawnTemplateNode(URigVMEdGraph* InParentGraph, const TArray<FPinInfo>& InPins, const FName& InNodeName = NAME_None);
+
+	FText GetMenuName() const { return MenuName;}
+	FText GetMenuTooltip() const { return MenuTooltip;}
+	FText GetMenuCategory() const { return MenuCategory;}
+	FText GetMenuKeywords() const { return MenuKeywords;}
+	FSlateIcon GetMenuIcon() const { return MenuIcon;}
+	FLinearColor GetMenuColor() const { return MenuIconTint;}
 	
 protected:
 
diff --git a/Engine/Plugins/Runtime/RigVM/Source/RigVMEditor/Public/EdGraph/NodeSpawners/RigVMEdGraphTemplateNodeSpawner.h b/Engine/Plugins/Runtime/RigVM/Source/RigVMEditor/Public/EdGraph/NodeSpawners/RigVMEdGraphTemplateNodeSpawner.h
--- a/Engine/Plugins/Runtime/RigVM/Source/RigVMEditor/Public/EdGraph/NodeSpawners/RigVMEdGraphTemplateNodeSpawner.h
+++ b/Engine/Plugins/Runtime/RigVM/Source/RigVMEditor/Public/EdGraph/NodeSpawners/RigVMEdGraphTemplateNodeSpawner.h
@@ -22,7 +22,7 @@ class URigVMEdGraphTemplateNodeSpawner : public URigVMEdGraphNodeSpawner
 {
 public:
 	
-	virtual ~URigVMEdGraphTemplateNodeSpawner() = default;
+	virtual ~URigVMEdGraphTemplateNodeSpawner() {}
 	/**
 	 * Creates a new URigVMEdGraphVariableNodeSpawner, charged with spawning 
 	 * a new member-variable node
@@ -30,7 +30,7 @@ public:
 	 * @return A newly allocated instance of this class.
 	 */
 	UE_DEPRECATED(5.7, "Please use URigVMEdGraphTemplateNodeSpawner(const FName& InNotation, const FText& InMenuDesc, const FText& InCategory, const FText& InTooltip)")
-	static UE_API URigVMEdGraphTemplateNodeSpawner* CreateFromNotation(const FName& InNotation, const FText& InMenuDesc, const FText& InCategory, const FText& InTooltip) { return nullptr; }
+	static URigVMEdGraphTemplateNodeSpawner* CreateFromNotation(const FName& InNotation, const FText& InMenuDesc, const FText& InCategory, const FText& InTooltip) { return nullptr; }
 	UE_API URigVMEdGraphTemplateNodeSpawner(const FName& InNotation, const FText& InMenuDesc, const FText& InCategory, const FText& InTooltip);
 
 	// URigVMEdGraphNodeSpawner interface
@@ -46,7 +46,7 @@ private:
 
 	const FRigVMTemplate* Template = nullptr;
 
-	static UE_API URigVMEdGraphNode* SpawnNode(URigVMEdGraph* ParentGraph, FRigVMAssetInterfacePtr Blueprint, const FRigVMTemplate* Template, FVector2D const Location);
+	static UE_API URigVMEdGraphNode* SpawnNode(URigVMEdGraph* ParentGraph, FRigVMEditorAssetInterfacePtr Blueprint, const FRigVMTemplate* Template, FVector2D const Location);
 
 	friend class UEngineTestControlRig;
 	friend class FRigVMEditorBase;
diff --git a/Engine/Plugins/Runtime/RigVM/Source/RigVMEditor/Public/EdGraph/NodeSpawners/RigVMEdGraphUnitNodeSpawner.h b/Engine/Plugins/Runtime/RigVM/Source/RigVMEditor/Public/EdGraph/NodeSpawners/RigVMEdGraphUnitNodeSpawner.h
--- a/Engine/Plugins/Runtime/RigVM/Source/RigVMEditor/Public/EdGraph/NodeSpawners/RigVMEdGraphUnitNodeSpawner.h
+++ b/Engine/Plugins/Runtime/RigVM/Source/RigVMEditor/Public/EdGraph/NodeSpawners/RigVMEdGraphUnitNodeSpawner.h
@@ -23,7 +23,7 @@ class URigVMEdGraphUnitNodeSpawner : public URigVMEdGraphNodeSpawner
 {
 public:
 	
-	virtual ~URigVMEdGraphUnitNodeSpawner() = default;
+	virtual ~URigVMEdGraphUnitNodeSpawner() {}
 	
 	/**
 	 * Creates a new URigVMEdGraphVariableNodeSpawner, charged with spawning 
@@ -32,7 +32,7 @@ public:
 	 * @return A newly allocated instance of this class.
 	 */
 	UE_DEPRECATED(5.7, "Please use URigVMEdGraphUnitNodeSpawner(UScriptStruct* InStruct, const FName& InMethodName, const FText& InMenuDesc, const FText& InCategory, const FText& InTooltip)")
-	static UE_API URigVMEdGraphUnitNodeSpawner* CreateFromStruct(UScriptStruct* InStruct, const FName& InMethodName, const FText& InMenuDesc, const FText& InCategory, const FText& InTooltip) { return nullptr; }
+	static URigVMEdGraphUnitNodeSpawner* CreateFromStruct(UScriptStruct* InStruct, const FName& InMethodName, const FText& InMenuDesc, const FText& InCategory, const FText& InTooltip) { return nullptr; }
 	UE_API URigVMEdGraphUnitNodeSpawner(UScriptStruct* InStruct, const FName& InMethodName, const FText& InMenuDesc, const FText& InCategory, const FText& InTooltip);
 
 	// URigVMEdGraphNodeSpawner interface
@@ -40,7 +40,7 @@ public:
 	UE_API virtual URigVMEdGraphNode* Invoke(URigVMEdGraph* ParentGraph, FVector2D const Location) const override;
 	// URigVMEdGraphNodeSpawner interface
 	
-	static UE_API void HookupMutableNode(URigVMNode* InModelNode, FRigVMAssetInterfacePtr InRigBlueprint);
+	static UE_API void HookupMutableNode(URigVMNode* InModelNode, FRigVMEditorAssetInterfacePtr InRigBlueprint);
 
 private:
 	/** The unit type we will spawn */
@@ -48,7 +48,7 @@ private:
 
 	FName MethodName;
 
-	static UE_API URigVMEdGraphNode* SpawnNode(URigVMEdGraph* ParentGraph, FRigVMAssetInterfacePtr Blueprint, UScriptStruct* StructTemplate, const
+	static UE_API URigVMEdGraphNode* SpawnNode(URigVMEdGraph* ParentGraph, FRigVMEditorAssetInterfacePtr Blueprint, UScriptStruct* StructTemplate, const
 	                                           FName& InMethodName, FVector2D Location);
 
 	friend class UEngineTestControlRig;
