Control Rig: Split IRigVMAssetInterface into IRigVMEditorAssetInterface and IRigVMRuntimeAssetInterface.

diff --git a/Engine/Plugins/Runtime/RigVM/Source/RigVMEditor/Public/EdGraph/NodeSpawners/RigVMEdGraphVariableNodeSpawner.h b/Engine/Plugins/Runtime/RigVM/Source/RigVMEditor/Public/EdGraph/NodeSpawners/RigVMEdGraphVariableNodeSpawner.h
--- a/Engine/Plugins/Runtime/RigVM/Source/RigVMEditor/Public/EdGraph/NodeSpawners/RigVMEdGraphVariableNodeSpawner.h
+++ b/Engine/Plugins/Runtime/RigVM/Source/RigVMEditor/Public/EdGraph/NodeSpawners/RigVMEdGraphVariableNodeSpawner.h
@@ -11,7 +11,7 @@
 #include "Engine/Blueprint.h"
 #include "EdGraph/NodeSpawners/RigVMEdGraphNodeSpawner.h"
 #include "RigVMCore/RigVMExternalVariable.h"
-#include "RigVMModel/RigVMVariableDescription.h"
+#include "RigVMCore/RigVMVariableDescription.h"
 
 #define UE_API RIGVMEDITOR_API
 
@@ -22,7 +22,7 @@ class URigVMEdGraphVariableNodeSpawner : public URigVMEdGraphNodeSpawner
 
 public:
 	
-	virtual ~URigVMEdGraphVariableNodeSpawner() = default;
+	virtual ~URigVMEdGraphVariableNodeSpawner() {}
 
 	/**
 	 * Creates a new URigVMEdGraphVariableNodeSpawner, charged with spawning 
@@ -30,13 +30,13 @@ public:
 	 * 
 	 * @return A newly allocated instance of this class.
 	 */
-	UE_DEPRECATED(5.7, "Please use URigVMEdGraphVariableNodeSpawner(FRigVMAssetInterfacePtr InBlueprint, const FRigVMExternalVariable& InExternalVariable, bool bInIsGetter, const FText& InMenuDesc, const FText& InCategory, const FText& InTooltip)")
-	static UE_API URigVMEdGraphVariableNodeSpawner* CreateFromExternalVariable(URigVMBlueprint* InBlueprint, const FRigVMExternalVariable& InExternalVariable, bool bInIsGetter, const FText& InMenuDesc, const FText& InCategory, const FText& InTooltip) { return nullptr; }
-	UE_DEPRECATED(5.7, "Please use URigVMEdGraphVariableNodeSpawner(FRigVMAssetInterfacePtr InBlueprint, URigVMGraph* InGraphOwner, const FRigVMGraphVariableDescription& InLocalVariable, bool bInIsGetter, const FText& InMenuDesc, const FText& InCategory, const FText& InTooltip)")
-	static UE_API URigVMEdGraphVariableNodeSpawner* CreateFromLocalVariable(URigVMBlueprint* InBlueprint, URigVMGraph* InGraphOwner, const FRigVMGraphVariableDescription& InLocalVariable, bool bInIsGetter, const FText& InMenuDesc, const FText& InCategory, const FText& InTooltip) { return nullptr; }
+	UE_DEPRECATED(5.7, "Please use URigVMEdGraphVariableNodeSpawner(FRigVMEditorAssetInterfacePtr InBlueprint, const FRigVMExternalVariable& InExternalVariable, bool bInIsGetter, const FText& InMenuDesc, const FText& InCategory, const FText& InTooltip)")
+	static URigVMEdGraphVariableNodeSpawner* CreateFromExternalVariable(URigVMBlueprint* InBlueprint, const FRigVMExternalVariable& InExternalVariable, bool bInIsGetter, const FText& InMenuDesc, const FText& InCategory, const FText& InTooltip) { return nullptr; }
+	UE_DEPRECATED(5.7, "Please use URigVMEdGraphVariableNodeSpawner(FRigVMEditorAssetInterfacePtr InBlueprint, URigVMGraph* InGraphOwner, const FRigVMGraphVariableDescription& InLocalVariable, bool bInIsGetter, const FText& InMenuDesc, const FText& InCategory, const FText& InTooltip)")
+	static URigVMEdGraphVariableNodeSpawner* CreateFromLocalVariable(URigVMBlueprint* InBlueprint, URigVMGraph* InGraphOwner, const FRigVMGraphVariableDescription& InLocalVariable, bool bInIsGetter, const FText& InMenuDesc, const FText& InCategory, const FText& InTooltip) { return nullptr; }
 
-	UE_API URigVMEdGraphVariableNodeSpawner(FRigVMAssetInterfacePtr InBlueprint, const FRigVMExternalVariable& InExternalVariable, bool bInIsGetter, const FText& InMenuDesc, const FText& InCategory, const FText& InTooltip);
-	UE_API URigVMEdGraphVariableNodeSpawner(FRigVMAssetInterfacePtr InBlueprint, URigVMGraph* InGraphOwner, const FRigVMGraphVariableDescription& InLocalVariable, bool bInIsGetter, const FText& InMenuDesc, const FText& InCategory, const FText& InTooltip);
+	UE_API URigVMEdGraphVariableNodeSpawner(FRigVMEditorAssetInterfacePtr InBlueprint, const FRigVMExternalVariable& InExternalVariable, bool bInIsGetter, const FText& InMenuDesc, const FText& InCategory, const FText& InTooltip);
+	UE_API URigVMEdGraphVariableNodeSpawner(FRigVMEditorAssetInterfacePtr InBlueprint, URigVMGraph* InGraphOwner, const FRigVMGraphVariableDescription& InLocalVariable, bool bInIsGetter, const FText& InMenuDesc, const FText& InCategory, const FText& InTooltip);
 
 	// URigVMEdGraphNodeSpawner interface
 	UE_API virtual FString GetSpawnerSignature() const override;
@@ -47,7 +47,7 @@ public:
 private:
 
 	/** The pin type we will spawn */
-	TWeakInterfacePtr<IRigVMAssetInterface> Blueprint;
+	TWeakInterfacePtr<IRigVMEditorAssetInterface> Blueprint;
 	TWeakObjectPtr<URigVMGraph> GraphOwner;
 	FRigVMExternalVariable ExternalVariable;
 	bool bIsGetter;
diff --git a/Engine/Plugins/Runtime/RigVM/Source/RigVMEditor/Public/Editor/RigVMCommentNodeDetailCustomization.h b/Engine/Plugins/Runtime/RigVM/Source/RigVMEditor/Public/Editor/RigVMCommentNodeDetailCustomization.h
--- a/Engine/Plugins/Runtime/RigVM/Source/RigVMEditor/Public/Editor/RigVMCommentNodeDetailCustomization.h
+++ b/Engine/Plugins/Runtime/RigVM/Source/RigVMEditor/Public/Editor/RigVMCommentNodeDetailCustomization.h
@@ -6,7 +6,7 @@
 #include "IDetailCustomization.h"
 #include "IPropertyTypeCustomization.h"
 #include "RigVMHost.h"
-#include "RigVMAsset.h"
+#include "RigVMEditorAsset.h"
 #include "Styling/SlateTypes.h"
 #include "IPropertyUtilities.h"
 
@@ -35,7 +35,7 @@ public:
 
 private:
 
-	IRigVMAssetInterface* BlueprintBeingCustomized;
+	FRigVMEditorAssetInterfacePtr BlueprintBeingCustomized;
 	TArray<TWeakObjectPtr<URigVMCommentNode>> ObjectsBeingCustomized;
 
 	UE_API void GetValuesFromNode(TWeakObjectPtr<URigVMCommentNode> Node);
diff --git a/Engine/Plugins/Runtime/RigVM/Source/RigVMEditor/Public/Editor/RigVMEditor.h b/Engine/Plugins/Runtime/RigVM/Source/RigVMEditor/Public/Editor/RigVMEditor.h
--- a/Engine/Plugins/Runtime/RigVM/Source/RigVMEditor/Public/Editor/RigVMEditor.h
+++ b/Engine/Plugins/Runtime/RigVM/Source/RigVMEditor/Public/Editor/RigVMEditor.h
@@ -25,7 +25,7 @@
 
 #define LOCTEXT_NAMESPACE "RigVMEditor"
 
-class IRigVMAssetInterface;
+class IRigVMEditorAssetInterface;
 class FDocumentTracker;
 class SRigVMEditorGraphExplorer;
 class URigVM;
@@ -37,7 +37,7 @@ class FTransaction;
 class SGraphEditor;
 class FPreviewScene;
 
-DECLARE_MULTICAST_DELEGATE_TwoParams(FRigVMEditorClosed, const IRigVMEditor*, FRigVMAssetInterfacePtr);
+DECLARE_MULTICAST_DELEGATE_TwoParams(FRigVMEditorClosed, const IRigVMEditor*, FRigVMEditorAssetInterfacePtr);
 
 
 /* Enums to use when grouping the blueprint members in the list panel. The order here will determine the order in the list */
@@ -106,7 +106,7 @@ public:
 	virtual TSharedPtr<FAssetEditorToolkit> GetHostingApp() = 0;
 	UE_DEPRECATED(5.7, "Plase use GetRigVMAsset")
 	virtual URigVMBlueprint* GetRigVMBlueprint() const { return Cast<URigVMBlueprint>(GetRigVMAssetInterface().GetObject()); }
-	virtual FRigVMAssetInterfacePtr GetRigVMAssetInterface() const = 0;
+	virtual FRigVMEditorAssetInterfacePtr GetRigVMAssetInterface() const = 0;
 	virtual const TArray< UObject* >* GetObjectsCurrentlyBeingEdited() const = 0;
 	virtual URigVMHost* GetRigVMHost() const = 0;
 
@@ -145,6 +145,7 @@ public:
 	virtual TSharedRef<FUICommandList> GetToolkitCommands() = 0;
 
 	virtual bool SelectLocalVariable(const UEdGraph* Graph, const FName& VariableName) = 0;
+	virtual bool SelectVariable(const FName& VariableName) = 0;
 	virtual void JumpToHyperlink(const UObject* ObjectReference, bool bRequestRename = false) = 0;
 	virtual void OnAddNewLocalVariable() = 0;
 	virtual bool CanAddNewLocalVariable() const = 0;
@@ -180,6 +181,9 @@ public:
 	virtual bool IsRunGraphOnceOn() const = 0;
 	virtual bool CanRunGraphOnce() const = 0;
 
+	virtual bool IsDetailsPanelEditingGlobalOptions() const = 0;
+	virtual bool IsDetailsPanelEditingClassDefaults() const = 0;
+
 protected:
 	
 	virtual TSharedRef<IRigVMEditor> SharedRef() = 0;
@@ -189,7 +193,7 @@ protected:
 	virtual const TSharedPtr<FAssetEditorToolkit> GetHostingApp() const = 0;
 	virtual void InitAssetEditor(const EToolkitMode::Type Mode, const TSharedPtr<IToolkitHost>& InitToolkitHost, const FName AppIdentifier, const TSharedRef<FTabManager::FLayout>& StandaloneDefaultLayout, const bool bCreateDefaultStandaloneMenu, const bool bCreateDefaultToolbar, const TArray<UObject*>& ObjectsToEdit, const bool bInIsToolbarFocusable = false, const bool bInUseSmallToolbarIcons = false, const TOptional<EAssetOpenMethod>& InOpenMethod = TOptional<EAssetOpenMethod>()) = 0;
 	virtual void CreateEditorToolbar() = 0;
-	virtual void CommonInitialization(const TArray<FRigVMAssetInterfacePtr>& InitBlueprints, bool bShouldOpenInDefaultsMode) = 0;
+	virtual void CommonInitialization(const TArray<FRigVMEditorAssetInterfacePtr>& InitBlueprints, bool bShouldOpenInDefaultsMode) = 0;
 	virtual TSharedPtr<FDocumentTracker> GetDocumentManager() const = 0;
 	virtual void AddApplicationMode(FName ModeName, TSharedRef<FApplicationMode> Mode) = 0;
 	virtual void RegenerateMenusAndToolbars() = 0;
@@ -212,7 +216,7 @@ protected:
 	virtual void RenameNewlyAddedAction(FName InActionName) = 0;
 	virtual FGraphPanelSelectionSet GetSelectedNodes() const = 0;
 	virtual void SetUISelectionState(FName SelectionOwner) = 0;
-	virtual void AnalyticsTrackNodeEvent(IRigVMAssetInterface* Blueprint, UEdGraphNode *GraphNode, bool bNodeDelete = false) const = 0;
+	virtual void AnalyticsTrackNodeEvent(IRigVMEditorAssetInterface* Blueprint, UEdGraphNode *GraphNode, bool bNodeDelete = false) const = 0;
 	virtual void PostUndo(bool bSuccess) = 0;
 	virtual void PostRedo(bool bSuccess) = 0;
 	virtual UEdGraphPin* GetCurrentlySelectedPin() const = 0;
@@ -230,7 +234,7 @@ protected:
 	virtual void AddMenuExtender(TSharedPtr<FExtender> Extender) = 0;
 	virtual void RemoveMenuExtender(TSharedPtr<FExtender> Extender) = 0;
 	virtual TSharedPtr<class IMessageLogListing> GetCompilerResultsListing() = 0;
-	virtual void OnBlueprintChangedInnerImpl(IRigVMAssetInterface* InBlueprint, bool bIsJustBeingCompiled = false) = 0;
+	virtual void OnBlueprintChangedInnerImpl(IRigVMEditorAssetInterface* InBlueprint, bool bIsJustBeingCompiled = false) = 0;
 	virtual void RefreshEditorsImpl(ERefreshRigVMEditorReason::Type Reason = ERefreshRigVMEditorReason::UnknownReason) = 0;
 	virtual void SetupGraphEditorEventsImpl(UEdGraph* InGraph, SGraphEditor::FGraphEditorEvents& InEvents) = 0;
 	virtual FActionMenuContent OnCreateGraphActionMenu(UEdGraph* InGraph, const FVector2f& InNodePosition, const TArray<UEdGraphPin*>& InDraggedPins, bool bAutoExpand, SGraphEditor::FActionMenuClosed InOnMenuClosed) = 0;
@@ -266,13 +270,13 @@ public:
 	 * @param	InitToolkitHost			When Mode is WorldCentric, this is the level editor instance to spawn this editor within
 	 * @param	InRigVMBlueprint	The blueprint object to start editing.
 	 */
-	UE_API virtual void InitRigVMEditor(const EToolkitMode::Type Mode, const TSharedPtr< class IToolkitHost >& InitToolkitHost, FRigVMAssetInterfacePtr InRigVMBlueprint);
+	UE_API virtual void InitRigVMEditor(const EToolkitMode::Type Mode, const TSharedPtr< class IToolkitHost >& InitToolkitHost, FRigVMEditorAssetInterfacePtr InRigVMBlueprint);
 
 	static UE_API FRigVMEditorBase* GetFromAssetEditorInstance(IAssetEditorInstance* Instance);
 	// returns the blueprint being edited
 	UE_DEPRECATED(5.7, "Please use GetRigVMAsset")
 	UE_API virtual URigVMBlueprint* GetRigVMBlueprint() const override;
-	UE_API virtual FRigVMAssetInterfacePtr GetRigVMAssetInterface() const override;
+	UE_API virtual FRigVMEditorAssetInterfacePtr GetRigVMAssetInterface() const override;
 	
 	UE_API void HandleJumpToHyperlink(const UObject* InSubject);
 
@@ -377,11 +381,6 @@ protected:
 
 	UE_API virtual void PostTransaction(bool bSuccess, const FTransaction* Transaction, bool bIsRedo);
 
-
-	// IToolkitHost Interface
-	UE_API void OnToolkitHostingStarted(const TSharedRef<class IToolkit>& Toolkit);
-	UE_API void OnToolkitHostingFinished(const TSharedRef<class IToolkit>& Toolkit);
-
 	//  FTickableEditorObject Interface
 	UE_API TStatId GetStatId() const;
 
@@ -400,11 +399,8 @@ protected:
 
 	virtual FRigVMEditorClosed& OnEditorClosed() override { return RigVMEditorClosedDelegate; }
 
-
-	/** Get the toolbox hosting widget */
-	TSharedRef<SBorder> GetToolbox() { return Toolbox.ToSharedRef(); }
-
 	UE_API virtual bool SelectLocalVariable(const UEdGraph* Graph, const FName& VariableName) override;
+	UE_API virtual bool SelectVariable(const FName& VariableName) override;
 
 	// FBlueprintEditor Interface
 	UE_API void CreateDefaultCommandsImpl();
@@ -419,7 +415,7 @@ protected:
 	UE_API FGraphAppearanceInfo GetGraphAppearance(class UEdGraph* InGraph) const;
 	UE_API void OnActiveTabChanged( TSharedPtr<SDockTab> PreviouslyActive, TSharedPtr<SDockTab> NewlyActivated );
 	UE_API void OnSelectedNodesChangedImpl(const TSet<class UObject*>& NewSelection);
-	UE_API void OnBlueprintChangedImpl(IRigVMAssetInterface* InBlueprint, bool bIsJustBeingCompiled);
+	UE_API void OnBlueprintChangedImpl(IRigVMEditorAssetInterface* InBlueprint, bool bIsJustBeingCompiled);
 	UE_API virtual void ForceEditorRefresh(ERefreshRigVMEditorReason::Type Reason = ERefreshRigVMEditorReason::UnknownReason) override;
 	UE_API void SetupGraphEditorEvents(UEdGraph* InGraph, SGraphEditor::FGraphEditorEvents& InEvents);
 	UE_API void FocusInspectorOnGraphSelection(const TSet<class UObject*>& NewSelection, bool bForceRefresh = false);
@@ -441,7 +437,7 @@ protected:
 	UE_API void OnPropertyChanged(UObject* InObject, FPropertyChangedEvent& InEvent);
 	UE_API virtual void OnWrappedPropertyChangedChainEvent(URigVMDetailsViewWrapperObject* InWrapperObject, const FString& InPropertyPath, FPropertyChangedChainEvent& InPropertyChangedChainEvent);
 	UE_API void OnRequestLocalizeFunctionDialog(FRigVMGraphFunctionIdentifier& InFunction, URigVMController* InTargetController, IRigVMGraphFunctionHost* InTargetFunctionHost, bool bForce);
-	UE_API FRigVMController_BulkEditResult OnRequestBulkEditDialog(FRigVMAssetInterfacePtr InBlueprint, URigVMController* InController, URigVMLibraryNode* InFunction, ERigVMControllerBulkEditType InEditType);
+	UE_API FRigVMController_BulkEditResult OnRequestBulkEditDialog(FRigVMEditorAssetInterfacePtr InBlueprint, URigVMController* InController, URigVMLibraryNode* InFunction, ERigVMControllerBulkEditType InEditType);
 	UE_API bool OnRequestBreakLinksDialog(TArray<URigVMLink*> InLinks);
 	UE_API TRigVMTypeIndex OnRequestPinTypeSelectionDialog(const TArray<TRigVMTypeIndex>& InTypes);
 
@@ -498,9 +494,11 @@ protected:
 	UE_API void SetDetailViewForGraph(URigVMGraph* InGraph);
 	UE_API void SetDetailViewForFocusedGraph();
 	UE_API void SetDetailViewForLocalVariable();
+	UE_API void SetDetailViewForVariable();
 	UE_API virtual void RefreshDetailView();
 	UE_API bool DetailViewShowsAnyRigUnit() const;
 	UE_API bool DetailViewShowsLocalVariable() const;
+	UE_API bool DetailViewShowsVariable() const;
 	UE_API bool DetailViewShowsStruct(UScriptStruct* InStruct) const;
 	UE_API void ClearDetailObject(bool bChangeUISelectionState = true);
 	UE_API void ClearDetailsViewWrapperObjects();
@@ -538,10 +536,10 @@ protected:
 	/** Update stale watch pins */
 	UE_API void UpdateStaleWatchedPins();
 	
-	UE_API virtual void HandleRefreshEditorFromBlueprint(FRigVMAssetInterfacePtr InBlueprint);
+	UE_API virtual void HandleRefreshEditorFromBlueprint(FRigVMEditorAssetInterfacePtr InBlueprint);
 	UE_API void HandleVariableDroppedFromBlueprint(UObject* InSubject, FProperty* InVariableToDrop, const FVector2D& InDropPosition, const FVector2D& InScreenPosition);
 	UE_API void OnGraphNodeClicked(URigVMEdGraphNode* InNode, const FGeometry& InNodeGeometry, const FPointerEvent& InMouseEvent);
-	UE_API void OnNodeDoubleClicked(FRigVMAssetInterfacePtr InBlueprint, URigVMNode* InNode);
+	UE_API void OnNodeDoubleClicked(FRigVMEditorAssetInterfacePtr InBlueprint, URigVMNode* InNode);
 	UE_API void OnGraphImported(UEdGraph* InEdGraph);
 	UE_API bool OnActionMatchesName(FEdGraphSchemaAction* InAction, const FName& InName) const;
 	UE_API void FrameSelection();
@@ -579,14 +577,11 @@ protected:
 
 	FPreviewHostUpdated PreviewHostUpdated;
 
-	/** Toolbox hosting widget */
-	TSharedPtr<SBorder> Toolbox;
-
 	TSharedPtr<SRigVMEditorGraphExplorer> GraphExplorerWidget;
 
 	FRigVMEditorClosed RigVMEditorClosedDelegate;
 
-	virtual void SetGraphExplorerWidget(TSharedPtr<SRigVMEditorGraphExplorer> InWidget) override { GraphExplorerWidget = InWidget; }
+	UE_API virtual void SetGraphExplorerWidget(TSharedPtr<SRigVMEditorGraphExplorer> InWidget) override;
 	virtual TSharedPtr<SRigVMEditorGraphExplorer> GetGraphExplorerWidget() override { return GraphExplorerWidget; }
 
 	UE_API bool IsEditingSingleBlueprint() const;
diff --git a/Engine/Plugins/Runtime/RigVM/Source/RigVMEditor/Public/Editor/RigVMEditorMenuContext.h b/Engine/Plugins/Runtime/RigVM/Source/RigVMEditor/Public/Editor/RigVMEditorMenuContext.h
--- a/Engine/Plugins/Runtime/RigVM/Source/RigVMEditor/Public/Editor/RigVMEditorMenuContext.h
+++ b/Engine/Plugins/Runtime/RigVM/Source/RigVMEditor/Public/Editor/RigVMEditorMenuContext.h
@@ -3,7 +3,7 @@
 #pragma once
 
 #include "CoreMinimal.h"
-#include "RigVMAsset.h"
+#include "RigVMEditorAsset.h"
 #include "EdGraph/EdGraph.h"
 
 #include "RigVMEditorMenuContext.generated.h"
@@ -66,7 +66,7 @@ public:
     UE_API URigVMBlueprint* GetRigVMBlueprint() const;
 
 	UFUNCTION(BlueprintCallable, Category = RigVMEditor)
-	UE_API TScriptInterface<IRigVMAssetInterface> GetRigVMAssetInterface() const;
+	UE_API TScriptInterface<IRigVMEditorAssetInterface> GetRigVMAssetInterface() const;
 	
 	/** Get the active rigvm host instance in the viewport */
 	UFUNCTION(BlueprintCallable, Category = RigVMEditor)
diff --git a/Engine/Plugins/Runtime/RigVM/Source/RigVMEditor/Public/Editor/RigVMFindReferences.h b/Engine/Plugins/Runtime/RigVM/Source/RigVMEditor/Public/Editor/RigVMFindReferences.h
--- a/Engine/Plugins/Runtime/RigVM/Source/RigVMEditor/Public/Editor/RigVMFindReferences.h
+++ b/Engine/Plugins/Runtime/RigVM/Source/RigVMEditor/Public/Editor/RigVMFindReferences.h
@@ -3,7 +3,7 @@
 #pragma once
 
 #include "CoreMinimal.h"
-#include "RigVMAsset.h"
+#include "RigVMEditorAsset.h"
 #include "Input/Reply.h"
 #include "Widgets/Views/STreeView.h"
 #include "EdGraph/EdGraphPin.h"
@@ -12,7 +12,7 @@
 
 #define UE_API RIGVMEDITOR_API
 
-class IRigVMAssetInterface;
+class IRigVMEditorAssetInterface;
 class URigVMBlueprint;
 class SWidget;
 class SDockTab;
@@ -27,12 +27,12 @@ public:
 	virtual ~FRigVMFindResult() = default;
 
 	/* Create a root */
-	UE_DEPRECATED(5.7, "Please use FRigVMFindResult(FRigVMAssetInterfacePtr InBlueprint)")
+	UE_DEPRECATED(5.7, "Please use FRigVMFindResult(FRigVMEditorAssetInterfacePtr InBlueprint)")
 	UE_API explicit FRigVMFindResult(TWeakObjectPtr<URigVMBlueprint> InBlueprint);
-	UE_DEPRECATED(5.7, "Please use FRigVMFindResult(FRigVMAssetInterfacePtr InBlueprint, const FText& InDisplayText)")
+	UE_DEPRECATED(5.7, "Please use FRigVMFindResult(FRigVMEditorAssetInterfacePtr InBlueprint, const FText& InDisplayText)")
 	UE_API explicit FRigVMFindResult(TWeakObjectPtr<URigVMBlueprint> InBlueprint, const FText& InDisplayText);
-	UE_API explicit FRigVMFindResult(TWeakInterfacePtr<IRigVMAssetInterface> InBlueprint);
-	UE_API explicit FRigVMFindResult(TWeakInterfacePtr<IRigVMAssetInterface> InBlueprint, const FText& InDisplayText);
+	UE_API explicit FRigVMFindResult(TWeakInterfacePtr<IRigVMEditorAssetInterface> InBlueprint);
+	UE_API explicit FRigVMFindResult(TWeakInterfacePtr<IRigVMEditorAssetInterface> InBlueprint, const FText& InDisplayText);
 
 	/* Called when user clicks on the search item */
 	UE_API virtual FReply OnClick();
@@ -47,10 +47,10 @@ public:
 	virtual void FinalizeSearchData() {};
 
 	/** gets the blueprint housing all these search results */
-	UE_DEPRECATED(5.7, "Please use FRigVMAssetInterfacePtr GetRigVMAssetInterface()")
+	UE_DEPRECATED(5.7, "Please use FRigVMEditorAssetInterfacePtr GetRigVMAssetInterface()")
 	UE_API URigVMBlueprint* GetBlueprint() const;
-	UE_API FRigVMAssetInterfacePtr GetRigVMAssetInterface();
-	UE_API const FRigVMAssetInterfacePtr GetRigVMAssetInterface() const;
+	UE_API FRigVMEditorAssetInterfacePtr GetRigVMAssetInterface();
+	UE_API const FRigVMEditorAssetInterfacePtr GetRigVMAssetInterface() const;
 
 	/**
 	* Parses search info for specific data important for displaying the search result in an easy to understand format
@@ -61,7 +61,7 @@ public:
 	virtual void ParseSearchInfo(FText InKey, FText InValue) {};
 
 	/** Returns the Object represented by this search information give the Blueprint it can be found in */
-	UE_API virtual UObject* GetObject(FRigVMAssetInterfacePtr InBlueprint) const;
+	UE_API virtual UObject* GetObject(FRigVMEditorAssetInterfacePtr InBlueprint) const;
 
 	/** Returns the display string for the row */
 	UE_API FText GetDisplayString() const;
@@ -74,7 +74,7 @@ public:
 	TWeakPtr<FRigVMFindResult> Parent;
 
 	/*If it exists it is the blueprint*/
-	TWeakInterfacePtr<IRigVMAssetInterface> WeakBlueprint;
+	TWeakInterfacePtr<IRigVMEditorAssetInterface> WeakBlueprint;
 
 	/*The display text for this item */
 	FText DisplayText;
@@ -115,9 +115,9 @@ namespace RigVMFindReferencesHelpers
 class FRigVMFindReferencesGraphNode : public FRigVMFindResult
 {
 public:
-	UE_DEPRECATED(5.7, "Please use FRigVMFindReferencesGraphNode(FRigVMAssetInterfacePtr InBlueprint)")
+	UE_DEPRECATED(5.7, "Please use FRigVMFindReferencesGraphNode(FRigVMEditorAssetInterfacePtr InBlueprint)")
 	FRigVMFindReferencesGraphNode(TWeakObjectPtr<URigVMBlueprint> InBlueprint);
-	FRigVMFindReferencesGraphNode(FRigVMAssetInterfacePtr InBlueprint);
+	FRigVMFindReferencesGraphNode(FRigVMEditorAssetInterfacePtr InBlueprint);
 	virtual ~FRigVMFindReferencesGraphNode() = default;
 
 	/** FRigVMFindResult Interface */
@@ -125,7 +125,7 @@ public:
 	virtual TSharedRef<SWidget>	CreateIcon() const override;
 	virtual void ParseSearchInfo(FText InKey, FText InValue) override;
 	virtual void FinalizeSearchData() override;
-	virtual UObject* GetObject(FRigVMAssetInterfacePtr InBlueprint) const override;
+	virtual UObject* GetObject(FRigVMEditorAssetInterfacePtr InBlueprint) const override;
 	virtual FText GetCategory() const override;
 	/** End FRigVMFindResult Interface */
 
@@ -150,9 +150,9 @@ private:
 class FRigVMFindReferencesPin : public FRigVMFindResult
 {
 public:
-	UE_DEPRECATED(5.7, "Plase use FRigVMFindReferencesPin(FRigVMAssetInterfacePtr InBlueprint, FString InSchemaName)")
+	UE_DEPRECATED(5.7, "Plase use FRigVMFindReferencesPin(FRigVMEditorAssetInterfacePtr InBlueprint, FString InSchemaName)")
 	FRigVMFindReferencesPin(TWeakObjectPtr<URigVMBlueprint> InBlueprint, FString InSchemaName);
-	FRigVMFindReferencesPin(FRigVMAssetInterfacePtr InBlueprint, FString InSchemaName);
+	FRigVMFindReferencesPin(FRigVMEditorAssetInterfacePtr InBlueprint, FString InSchemaName);
 	virtual ~FRigVMFindReferencesPin() = default;
 
 	/** FRigVMFindResult Interface */
@@ -177,9 +177,9 @@ private:
 class FRigVMFindReferencesVariable : public FRigVMFindResult
 {
 public:
-	UE_DEPRECATED(5.7, "Plase use FRigVMFindReferencesVariable(FRigVMAssetInterfacePtr InBlueprint)")
+	UE_DEPRECATED(5.7, "Plase use FRigVMFindReferencesVariable(FRigVMEditorAssetInterfacePtr InBlueprint)")
 	FRigVMFindReferencesVariable(TWeakObjectPtr<URigVMBlueprint> InBlueprint);
-	FRigVMFindReferencesVariable(FRigVMAssetInterfacePtr InBlueprint);
+	FRigVMFindReferencesVariable(FRigVMEditorAssetInterfacePtr InBlueprint);
 	virtual ~FRigVMFindReferencesVariable() = default;
 
 	/** FRigVMFindResult Interface */
@@ -201,9 +201,9 @@ private:
 class FRigVMFindReferencesGraph : public FRigVMFindResult
 {
 public:
-	UE_DEPRECATED(5.7, "Please use FRigVMFindReferencesGraph(FRigVMAssetInterfacePtr InBlueprint, EGraphType InGraphType)")
+	UE_DEPRECATED(5.7, "Please use FRigVMFindReferencesGraph(FRigVMEditorAssetInterfacePtr InBlueprint, EGraphType InGraphType)")
 	FRigVMFindReferencesGraph(TWeakObjectPtr<URigVMBlueprint> InBlueprint, EGraphType InGraphType);
-	FRigVMFindReferencesGraph(FRigVMAssetInterfacePtr InBlueprint, EGraphType InGraphType);
+	FRigVMFindReferencesGraph(FRigVMEditorAssetInterfacePtr InBlueprint, EGraphType InGraphType);
 	virtual ~FRigVMFindReferencesGraph() = default;
 
 	/** FRigVMFindResult Interface */
diff --git a/Engine/Plugins/Runtime/RigVM/Source/RigVMEditor/Public/Editor/RigVMGraphDetailCustomization.h b/Engine/Plugins/Runtime/RigVM/Source/RigVMEditor/Public/Editor/RigVMGraphDetailCustomization.h
--- a/Engine/Plugins/Runtime/RigVM/Source/RigVMEditor/Public/Editor/RigVMGraphDetailCustomization.h
+++ b/Engine/Plugins/Runtime/RigVM/Source/RigVMEditor/Public/Editor/RigVMGraphDetailCustomization.h
@@ -118,7 +118,7 @@ public:
 		, GraphPtr(InGraph)
 		, WeakRigVMClientHost(InRigVMClientHost)
 		, RigVMEditorPtr(InEditor)
-		, NameValidator((const FRigVMAssetInterfacePtr)nullptr, InGraph.Get(), InPin->GetFName()) // TODO check if missing Blueprint affects name validation
+		, NameValidator((const FRigVMEditorAssetInterfacePtr)nullptr, InGraph.Get(), InPin->GetFName()) // TODO check if missing Blueprint affects name validation
 	{}
 
 private:
@@ -210,10 +210,10 @@ class FRigVMGraphDetailCustomization : public IDetailCustomization
 public:
 	/** Makes a new instance of this detail layout class for a specific detail view requesting it */
 	static UE_API TSharedRef<IDetailCustomization> MakeInstance(TSharedPtr<IRigVMEditor> InBlueprintEditor, const UClass* InExpectedBlueprintClass);
-	UE_API FRigVMGraphDetailCustomization(TSharedPtr<IRigVMEditor> RigVMigEditor, FRigVMAssetInterfacePtr RigVMBlueprint);
+	UE_API FRigVMGraphDetailCustomization(TSharedPtr<IRigVMEditor> RigVMigEditor, FRigVMEditorAssetInterfacePtr RigVMBlueprint);
 #if WITH_RIGVMLEGACYEDITOR
 	static UE_API TSharedPtr<IDetailCustomization> MakeLegacyInstance(TSharedPtr<IBlueprintEditor> InBlueprintEditor, const UClass* InExpectedBlueprintClass);
-	UE_API FRigVMGraphDetailCustomization(TSharedPtr<IBlueprintEditor> RigVMigEditor, FRigVMAssetInterfacePtr RigVMBlueprint);
+	UE_API FRigVMGraphDetailCustomization(TSharedPtr<IBlueprintEditor> RigVMigEditor, FRigVMEditorAssetInterfacePtr RigVMBlueprint);
 #endif
 
 
@@ -228,7 +228,7 @@ private:
 	TWeakObjectPtr<URigVMEdGraph> GraphPtr;
 
 	/** The blueprint we are editing */
-	TWeakInterfacePtr<IRigVMAssetInterface> RigVMBlueprintPtr;
+	TWeakInterfacePtr<IRigVMEditorAssetInterface> RigVMBlueprintPtr;
 
 	TSharedPtr<FRigVMGraphDetailCustomizationImpl> RigVMGraphDetailCustomizationImpl;
 };
@@ -341,7 +341,7 @@ public:
 	UE_API void OnNameListComboBox(const FNameProperty* InProperty, const TArray<TSharedPtr<FRigVMStringWithTag>>* InNameList);
 	UE_API void CustomizeLiveValues(IDetailLayoutBuilder& DetailLayout);
 
-	FRigVMAssetInterfacePtr BlueprintBeingCustomized;
+	FRigVMEditorAssetInterfacePtr BlueprintBeingCustomized;
 	TArray<TWeakObjectPtr<URigVMDetailsViewWrapperObject>> ObjectsBeingCustomized;
 	TArray<TWeakObjectPtr<URigVMNode>> NodesBeingCustomized;
 	TMap<FName, TSharedPtr<SRigVMGraphPinNameListValueWidget>> NameListWidgets;
@@ -470,7 +470,7 @@ protected:
 
 	UE_API void HandleControlEnumChanged(TSharedPtr<FString> InEnumPath, ESelectInfo::Type InSelectType, TSharedRef<IPropertyHandle> InPropertyHandle);
 
-	FRigVMAssetInterfacePtr BlueprintBeingCustomized;
+	FRigVMEditorAssetInterfacePtr BlueprintBeingCustomized;
 	URigVMGraph* GraphBeingCustomized;
 	TArray<TWeakObjectPtr<UObject>> ObjectsBeingCustomized;
 	TArray<TSharedPtr<FStructOnScope>> StructsBeingCustomized;
@@ -731,7 +731,7 @@ protected:
 		return TEXT("Scale3D");
 	}
 
-	TWeakInterfacePtr<IRigVMAssetInterface> BlueprintBeingCustomized;
+	TWeakInterfacePtr<IRigVMEditorAssetInterface> BlueprintBeingCustomized;
 	TWeakObjectPtr<URigVMGraph> GraphBeingCustomized;
 };
 
diff --git a/Engine/Plugins/Runtime/RigVM/Source/RigVMEditor/Public/Editor/RigVMLegacyEditor.h b/Engine/Plugins/Runtime/RigVM/Source/RigVMEditor/Public/Editor/RigVMLegacyEditor.h
--- a/Engine/Plugins/Runtime/RigVM/Source/RigVMEditor/Public/Editor/RigVMLegacyEditor.h
+++ b/Engine/Plugins/Runtime/RigVM/Source/RigVMEditor/Public/Editor/RigVMLegacyEditor.h
@@ -34,7 +34,7 @@ protected:
 	virtual const TSharedPtr<FAssetEditorToolkit> GetHostingApp() const override { return ConstCastSharedRef<FAssetEditorToolkit>(AsShared()); }
 	RIGVMEDITOR_API virtual void InitAssetEditor(const EToolkitMode::Type Mode, const TSharedPtr<IToolkitHost>& InitToolkitHost, const FName AppIdentifier, const TSharedRef<FTabManager::FLayout>& StandaloneDefaultLayout, const bool bCreateDefaultStandaloneMenu, const bool bCreateDefaultToolbar, const TArray<UObject*>& ObjectsToEdit, const bool bInIsToolbarFocusable, const bool bInUseSmallToolbarIcons, const TOptional<EAssetOpenMethod>& InOpenMethod) override;
 	RIGVMEDITOR_API virtual void CreateEditorToolbar() override;
-	RIGVMEDITOR_API virtual void CommonInitialization(const TArray<FRigVMAssetInterfacePtr>& InitBlueprints, bool bShouldOpenInDefaultsMode) override;
+	RIGVMEDITOR_API virtual void CommonInitialization(const TArray<FRigVMEditorAssetInterfacePtr>& InitBlueprints, bool bShouldOpenInDefaultsMode) override;
 	virtual TSharedPtr<FDocumentTracker> GetDocumentManager() const override { return DocumentManager; }
 	virtual void AddApplicationMode(FName ModeName, TSharedRef<FApplicationMode> Mode) override { FBlueprintEditor::AddApplicationMode(ModeName, Mode); }
 	virtual void RegenerateMenusAndToolbars() override { FBlueprintEditor::RegenerateMenusAndToolbars(); }
@@ -67,7 +67,7 @@ protected:
 	virtual FGraphPanelSelectionSet GetSelectedNodes() const override { return FBlueprintEditor::GetSelectedNodes(); }
 	virtual void SetUISelectionState(FName SelectionOwner) override { FBlueprintEditor::SetUISelectionState(SelectionOwner); }
 	virtual void AnalyticsTrackNodeEvent(UBlueprint* Blueprint, UEdGraphNode* GraphNode, bool bNodeDelete) const override { FBlueprintEditor::AnalyticsTrackNodeEvent(Blueprint, GraphNode, bNodeDelete); }
-	virtual void AnalyticsTrackNodeEvent(IRigVMAssetInterface* Blueprint, UEdGraphNode* GraphNode, bool bNodeDelete) const override { FBlueprintEditor::AnalyticsTrackNodeEvent(Cast<UBlueprint>(Blueprint->GetObject()), GraphNode, bNodeDelete); }
+	virtual void AnalyticsTrackNodeEvent(IRigVMEditorAssetInterface* Blueprint, UEdGraphNode* GraphNode, bool bNodeDelete) const override { FBlueprintEditor::AnalyticsTrackNodeEvent(Cast<UBlueprint>(Blueprint->GetObject()), GraphNode, bNodeDelete); }
 	RIGVMEDITOR_API virtual void JumpToHyperlink(const UObject* ObjectReference, bool bRequestRename) override;
 	RIGVMEDITOR_API virtual void PostUndo(bool bSuccess) override;
 	RIGVMEDITOR_API virtual void PostRedo(bool bSuccess) override;
@@ -87,7 +87,7 @@ protected:
 	virtual void AddMenuExtender(TSharedPtr<FExtender> Extender) override { FBlueprintEditor::AddMenuExtender(Extender); }
 	virtual void RemoveMenuExtender(TSharedPtr<FExtender> Extender) override { FBlueprintEditor::RemoveMenuExtender(Extender); }
 	virtual TSharedPtr<IMessageLogListing> GetCompilerResultsListing() override { return CompilerResultsListing; }
-	virtual void OnBlueprintChangedInnerImpl(IRigVMAssetInterface* InBlueprint, bool bIsJustBeingCompiled) override { FBlueprintEditor::OnBlueprintChangedImpl(Cast<UBlueprint>(InBlueprint->GetObject()), bIsJustBeingCompiled); }
+	virtual void OnBlueprintChangedInnerImpl(IRigVMEditorAssetInterface* InBlueprint, bool bIsJustBeingCompiled) override { FBlueprintEditor::OnBlueprintChangedImpl(Cast<UBlueprint>(InBlueprint->GetObject()), bIsJustBeingCompiled); }
 	RIGVMEDITOR_API virtual void RefreshEditorsImpl(ERefreshRigVMEditorReason::Type Reason) override;
 	virtual void SetupGraphEditorEvents(UEdGraph* InGraph, SGraphEditor::FGraphEditorEvents& InEvents) override { FRigVMEditorBase::SetupGraphEditorEvents(InGraph, InEvents); }
 	virtual void SetupGraphEditorEventsImpl(UEdGraph* InGraph, SGraphEditor::FGraphEditorEvents& InEvents) override { FBlueprintEditor::SetupGraphEditorEvents(InGraph, InEvents); }
@@ -103,6 +103,8 @@ protected:
 	virtual bool IsNodeTitleVisible(const UEdGraphNode* Node, bool bRequestRename) override { return FBlueprintEditor::IsNodeTitleVisible(Node, bRequestRename); }
 	virtual void EditClassDefaults_Clicked() override { FBlueprintEditor::EditClassDefaults_Clicked(); }
 	virtual void EditGlobalOptions_Clicked() override { FBlueprintEditor::EditGlobalOptions_Clicked(); }
+	virtual bool IsDetailsPanelEditingClassDefaults() const override { return FBlueprintEditor::IsDetailsPanelEditingClassDefaults(); }
+	virtual bool IsDetailsPanelEditingGlobalOptions() const override { return FBlueprintEditor::IsDetailsPanelEditingGlobalOptions(); }
 	virtual void TryInvokingDetailsTab(bool bFlash) override { FBlueprintEditor::TryInvokingDetailsTab(bFlash); }
 	virtual FName GetGraphExplorerWidgetID() override { return FBlueprintEditorTabs::MyBlueprintID; }
 	virtual void RefreshInspector() override { FBlueprintEditor::RefreshInspector(); }
@@ -127,12 +129,12 @@ public:
 	RIGVMEDITOR_API virtual void OnClose() override;
 	virtual FPreviewScene* GetPreviewScene() override { return &PreviewScene; }
 
-	RIGVMEDITOR_API virtual void ToggleHideUnrelatedNodes() override
+	virtual void ToggleHideUnrelatedNodes() override
 	{
 		ToggleFadeOutUnrelateNodes();
 	}
 
-	RIGVMEDITOR_API virtual bool IsToggleHideUnrelatedNodesChecked() const override
+	virtual bool IsToggleHideUnrelatedNodesChecked() const override
 	{
 		return IsToggleFadeOutUnrelatedNodesChecked();
 	}
@@ -202,8 +204,6 @@ protected:
 
 	// IToolkit Interface
 	virtual void BringToolkitToFront() override { FRigVMEditorBase::BringToolkitToFront(); }
-	virtual void OnToolkitHostingStarted(const TSharedRef<class IToolkit>& Toolkit) override { FRigVMEditorBase::OnToolkitHostingStarted(Toolkit); }
-	virtual void OnToolkitHostingFinished(const TSharedRef<class IToolkit>& Toolkit) override { FRigVMEditorBase::OnToolkitHostingFinished(Toolkit); }
 
 	virtual FNotifyHook* GetNotifyHook() override { return this; }
 
diff --git a/Engine/Plugins/Runtime/RigVM/Source/RigVMEditor/Public/Editor/RigVMLocalVariableDetailCustomization.h b/Engine/Plugins/Runtime/RigVM/Source/RigVMEditor/Public/Editor/RigVMLocalVariableDetailCustomization.h
--- a/Engine/Plugins/Runtime/RigVM/Source/RigVMEditor/Public/Editor/RigVMLocalVariableDetailCustomization.h
+++ b/Engine/Plugins/Runtime/RigVM/Source/RigVMEditor/Public/Editor/RigVMLocalVariableDetailCustomization.h
@@ -24,7 +24,7 @@ class FRigVMLocalVariableDetailCustomization : public IDetailCustomization
 	FRigVMLocalVariableDetailCustomization()
 	: GraphBeingCustomized(nullptr)
 	, BlueprintBeingCustomized(nullptr)
-	, NameValidator((const FRigVMAssetInterfacePtr)nullptr, nullptr, NAME_None)
+	, NameValidator((const FRigVMEditorAssetInterfacePtr)nullptr, nullptr, NAME_None)
 	{}
 
 	
@@ -42,7 +42,7 @@ public:
 private:
 
 	URigVMGraph* GraphBeingCustomized;
-	FRigVMAssetInterfacePtr BlueprintBeingCustomized;
+	FRigVMEditorAssetInterfacePtr BlueprintBeingCustomized;
 	FRigVMGraphVariableDescription VariableDescription;
 	TArray<TWeakObjectPtr<URigVMDetailsViewWrapperObject>> ObjectsBeingCustomized;
 
diff --git a/Engine/Plugins/Runtime/RigVM/Source/RigVMEditor/Public/Editor/RigVMNewEditor.h b/Engine/Plugins/Runtime/RigVM/Source/RigVMEditor/Public/Editor/RigVMNewEditor.h
--- a/Engine/Plugins/Runtime/RigVM/Source/RigVMEditor/Public/Editor/RigVMNewEditor.h
+++ b/Engine/Plugins/Runtime/RigVM/Source/RigVMEditor/Public/Editor/RigVMNewEditor.h
@@ -32,7 +32,7 @@ protected:
 	
 	UE_API virtual void InitAssetEditor(const EToolkitMode::Type Mode, const TSharedPtr<IToolkitHost>& InitToolkitHost, const FName AppIdentifier, const TSharedRef<FTabManager::FLayout>& StandaloneDefaultLayout, const bool bCreateDefaultStandaloneMenu, const bool bCreateDefaultToolbar, const TArray<UObject*>& ObjectsToEdit, const bool bInIsToolbarFocusable, const bool bInUseSmallToolbarIcons, const TOptional<EAssetOpenMethod>& InOpenMethod) override;
 	virtual void CreateEditorToolbar() override {}
-	UE_API virtual void CommonInitialization(const TArray<FRigVMAssetInterfacePtr>& InitBlueprints, bool bShouldOpenInDefaultsMode) override;
+	UE_API virtual void CommonInitialization(const TArray<FRigVMEditorAssetInterfacePtr>& InitBlueprints, bool bShouldOpenInDefaultsMode) override;
 	UE_API void OnBlueprintChanged(UObject* InBlueprint);
 	UE_API void SaveEditedObjectState();
 	virtual TSharedPtr<FDocumentTracker> GetDocumentManager() const override { return DocumentManager; }
@@ -50,7 +50,7 @@ public:
 	virtual TSharedPtr<SRigVMDetailsInspector> GetRigVMInspector() const override { return Inspector.IsValid() ? Inspector.ToSharedRef().ToSharedPtr() : nullptr; }
 	virtual void SetInspector(TSharedPtr<SRigVMDetailsInspector> InWidget) { Inspector = InWidget; };
 	
-	TSharedRef<SWidget> GetCompilerResults() const { return CompilerResults.ToSharedRef(); }
+	TSharedRef<SWidget> GetCompilerResults() const;
 	TSharedRef<SRigVMFindReferences> GetFindResults() const { return FindResults.ToSharedRef(); }
 	UE_API virtual void JumpToHyperlink(const UObject* ObjectReference, bool bRequestRename) override;
 protected:
@@ -71,7 +71,7 @@ protected:
 	UE_API virtual void RenameNewlyAddedAction(FName InActionName) override;
 	UE_API virtual FGraphPanelSelectionSet GetSelectedNodes() const override;
 	UE_API virtual void SetUISelectionState(FName SelectionOwner) override;
-	UE_API virtual void AnalyticsTrackNodeEvent(IRigVMAssetInterface* Blueprint, UEdGraphNode* GraphNode, bool bNodeDelete) const override;
+	UE_API virtual void AnalyticsTrackNodeEvent(IRigVMEditorAssetInterface* Blueprint, UEdGraphNode* GraphNode, bool bNodeDelete) const override;
 	UE_API virtual void PostUndo(bool bSuccess) override;
 	UE_API virtual void PostRedo(bool bSuccess) override;
 	UE_API virtual UEdGraphPin* GetCurrentlySelectedPin() const override;
@@ -89,9 +89,10 @@ protected:
 	virtual void RemoveToolbarExtender(TSharedPtr<FExtender> Extender) override { FWorkflowCentricApplication::RemoveToolbarExtender(Extender); };
 	virtual void AddMenuExtender(TSharedPtr<FExtender> Extender) override { FWorkflowCentricApplication::AddMenuExtender(Extender); }
 	virtual void RemoveMenuExtender(TSharedPtr<FExtender> Extender) override { FWorkflowCentricApplication::RemoveMenuExtender(Extender); }
-	UE_API virtual void OnBlueprintChangedInnerImpl(IRigVMAssetInterface* InBlueprint, bool bIsJustBeingCompiled) override;
+	UE_API virtual void OnBlueprintChangedInnerImpl(IRigVMEditorAssetInterface* InBlueprint, bool bIsJustBeingCompiled) override;
 	UE_API virtual void RefreshEditorsImpl(ERefreshRigVMEditorReason::Type Reason) override;
 	UE_API virtual void SetupGraphEditorEventsImpl(UEdGraph* InGraph, SGraphEditor::FGraphEditorEvents& InEvents) override;
+	UE_API void OnClosedActionMenu(UEdGraph* InEdGraph) const;
 	UE_API virtual FActionMenuContent OnCreateGraphActionMenu(UEdGraph* InGraph, const FVector2f& InNodePosition, const TArray<UEdGraphPin*>& InDraggedPins, bool bAutoExpand, SGraphEditor::FActionMenuClosed InOnMenuClosed) override;
 	UE_API virtual void AddReferencedObjectsImpl(FReferenceCollector& Collector) override;
 	UE_API virtual bool IsSectionVisible(RigVMNodeSectionID::Type InSectionID) const; 
@@ -103,8 +104,8 @@ protected:
 	UE_API virtual bool IsNodeTitleVisible(const UEdGraphNode* Node, bool bRequestRename) override;
 	UE_API virtual void EditClassDefaults_Clicked() override;
 	UE_API virtual void EditGlobalOptions_Clicked() override;
-	UE_API bool IsDetailsPanelEditingGlobalOptions() const;
-	UE_API bool IsDetailsPanelEditingClassDefaults() const;
+	UE_API virtual bool IsDetailsPanelEditingGlobalOptions() const override;
+	UE_API virtual bool IsDetailsPanelEditingClassDefaults() const override;
 	UE_API virtual void TryInvokingDetailsTab(bool bFlash = true) override;
 	virtual FName GetGraphExplorerWidgetID() override { return FRigVMEditorGraphExplorerTabSummoner::TabID(); }
 	UE_API virtual void RefreshInspector() override;
@@ -154,9 +155,9 @@ protected:
 	/** Dumps messages to the compiler log, with an option to force it to display/come to front */
 	UE_API void DumpMessagesToCompilerLog(const TArray<TSharedRef<class FTokenizedMessage>>& Messages, bool bForceMessageDisplay);
 public:
-	UE_DEPRECATED(5.7, "Please use void CreateDefaultTabContents(const TArray<FRigVMAssetInterfacePtr> InBlueprints)")
-	UE_API void CreateDefaultTabContents(const TArray<UBlueprint*> InBlueprints){}
-	UE_API void CreateDefaultTabContents(const TArray<FRigVMAssetInterfacePtr> InBlueprints);
+	UE_DEPRECATED(5.7, "Please use void CreateDefaultTabContents(const TArray<FRigVMEditorAssetInterfacePtr> InBlueprints)")
+	void CreateDefaultTabContents(const TArray<UBlueprint*> InBlueprints){}
+	UE_API void CreateDefaultTabContents(const TArray<FRigVMEditorAssetInterfacePtr> InBlueprints);
 
 	UE_API TSharedRef<SWidget> GenerateCompileOptionsMenu();
 	UE_API void MakeSaveOnCompileSubMenu(FMenuBuilder& InMenu);
@@ -237,6 +238,10 @@ protected:
 
 	/** Whether the graph action menu should be sensitive to the pins dragged off of */
 	bool bIsActionMenuContextSensitive;
+
+	FName ElementInDetailsPanel = NAME_None;
+	UE_API void HandlePreVariablesChanged(FName InName);
+	UE_API void HandlePostVariablesChanged(FName InName);
 };
 
 
diff --git a/Engine/Plugins/Runtime/RigVM/Source/RigVMEditor/Public/Editor/RigVMVariableDetailCustomization.h b/Engine/Plugins/Runtime/RigVM/Source/RigVMEditor/Public/Editor/RigVMVariableDetailCustomization.h
--- a/Engine/Plugins/Runtime/RigVM/Source/RigVMEditor/Public/Editor/RigVMVariableDetailCustomization.h
+++ b/Engine/Plugins/Runtime/RigVM/Source/RigVMEditor/Public/Editor/RigVMVariableDetailCustomization.h
@@ -5,7 +5,7 @@
 #include "CoreMinimal.h"
 #include "Engine/Blueprint.h"
 #include "IDetailCustomization.h"
-#include "RigVMAsset.h"
+#include "RigVMEditorAsset.h"
 #include "Styling/SlateTypes.h"
 #include "UObject/WeakFieldPtr.h"
 #include "Layout/Visibility.h"
@@ -31,7 +31,7 @@ class FRigVMVariableDetailCustomization : public IDetailCustomization
 public:
 	/** Makes a new instance of this detail layout class for a specific detail view requesting it */
 	static UE_API TSharedRef<IDetailCustomization> MakeInstance(TSharedPtr<IRigVMEditor> InEditor);
-	UE_API FRigVMVariableDetailCustomization(TSharedPtr<IRigVMEditor> InEditor, FRigVMAssetInterfacePtr Blueprint);
+	UE_API FRigVMVariableDetailCustomization(TSharedPtr<IRigVMEditor> InEditor, FRigVMEditorAssetInterfacePtr Blueprint);
 
 #if WITH_RIGVMLEGACYEDITOR
 	static UE_API TSharedPtr<IDetailCustomization> MakeLegacyInstance(TSharedPtr<IBlueprintEditor> InBlueprintEditor);
@@ -39,10 +39,11 @@ public:
 #endif
 
 	// IDetailCustomization interface
-	UE_API virtual void CustomizeDetails(IDetailLayoutBuilder& DetailLayout) override;
+	virtual void CustomizeDetails(IDetailLayoutBuilder& DetailLayout) override;
 
 private:
 	UE_API void PopulateCategories();
+	UE_API void OnFinishedChangingVariable(const FPropertyChangedEvent& PropertyChangedEvent);
 
 private:
 	/** Accessors passed to parent */
@@ -92,7 +93,7 @@ private:
 	TWeakPtr<IRigVMEditor> EditorPtr;
 
 	/** The blueprint we are editing */
-	TWeakInterfacePtr<IRigVMAssetInterface> BlueprintPtr;
+	TWeakInterfacePtr<IRigVMEditorAssetInterface> BlueprintPtr;
 
 	/** The widget used when in variable name editing mode */ 
 	TSharedPtr<SEditableTextBox> VarNameEditableTextBox;
@@ -104,7 +105,7 @@ private:
 	TWeakPtr<SListView<TSharedPtr<FText>>> CategoryListView;
 
 	/** Cached property for the variable we are affecting */
-	TWeakFieldPtr<FProperty> CachedVariableProperty;
+	TWeakFieldPtr<const FProperty> CachedVariableProperty;
 
 	/** Cached name for the variable we are affecting */
 	mutable FName CachedVariableName;
diff --git a/Engine/Plugins/Runtime/RigVM/Source/RigVMEditor/Public/Editor/RigVMVariantDetailCustomization.h b/Engine/Plugins/Runtime/RigVM/Source/RigVMEditor/Public/Editor/RigVMVariantDetailCustomization.h
--- a/Engine/Plugins/Runtime/RigVM/Source/RigVMEditor/Public/Editor/RigVMVariantDetailCustomization.h
+++ b/Engine/Plugins/Runtime/RigVM/Source/RigVMEditor/Public/Editor/RigVMVariantDetailCustomization.h
@@ -35,7 +35,7 @@ protected:
 	void OnAddTag(const FName& InTagName);
 	void OnRemoveTag(const FName& InTagName);
 
-	FRigVMAssetInterfacePtr BlueprintBeingCustomized = nullptr;
+	FRigVMEditorAssetInterfacePtr BlueprintBeingCustomized = nullptr;
 
 	/** The log widget used for function variants */
 	TSharedPtr<SRigVMLogWidget> VariantLog;
diff --git a/Engine/Plugins/Runtime/RigVM/Source/RigVMEditor/Public/Editor/SRigVMActionMenu.cpp b/Engine/Plugins/Runtime/RigVM/Source/RigVMEditor/Public/Editor/SRigVMActionMenu.cpp
--- a/Engine/Plugins/Runtime/RigVM/Source/RigVMEditor/Public/Editor/SRigVMActionMenu.cpp
+++ b/Engine/Plugins/Runtime/RigVM/Source/RigVMEditor/Public/Editor/SRigVMActionMenu.cpp
@@ -13,13 +13,16 @@
 #include "Editor/RigVMActionMenuUtils.h"
 #include "BlueprintEditorSettings.h"
 #include "BlueprintNodeSpawner.h"
+#include "BlueprintNodeTemplateCache.h"
 #include "K2Node_Variable.h"
 #include "RigVMActionMenuItem.h"
 #include "RigVMBlueprintLegacy.h"
 #include "RigVMBlueprintUtils.h"
+#include "EdGraph/NodeSpawners/RigVMEdGraphNodeSpawner.h"
 #include "Widgets/Input/SCheckBox.h"
 #include "Framework/Application/SlateApplication.h"
 #include "Widgets/SToolTip.h"
+#include "RigVMDeveloperTypeUtils.h"
 
 #define LOCTEXT_NAMESPACE "SRigVMGraphContextMenu"
 
@@ -41,7 +44,7 @@ struct FRigVMAction_PromoteVariable : public FEdGraphSchemaAction
 	{
 		if( ( ParentGraph != NULL ) && ( FromPin != NULL ) )
 		{
-			if (FRigVMAssetInterfacePtr RigVMBlueprint = FRigVMBlueprintUtils::FindAssetForGraph(ParentGraph))
+			if (FRigVMEditorAssetInterfacePtr RigVMBlueprint = FRigVMBlueprintUtils::FindAssetForGraph(ParentGraph))
 			{
 				if (URigVMController* Controller = RigVMBlueprint->GetController(ParentGraph))
 				{
@@ -73,9 +76,71 @@ struct FRigVMAction_PromoteVariable : public FEdGraphSchemaAction
 
 SRigVMActionMenu::~SRigVMActionMenu()
 {
+	OnClosedCallback.ExecuteIfBound();
 	
 }
 
+UEdGraphNode* FRigVMGraphSchemaAction_NewNode::PerformAction(class UEdGraph* ParentGraph, UEdGraphPin* FromPin, const FVector2f& Location, bool bSelectNewNode)
+{
+	if (URigVMEdGraph* RigVMEdGraph = Cast<URigVMEdGraph>(ParentGraph))
+	{
+		if (URigVMEdGraphNode* NewNode = NodeSpawner->Invoke(RigVMEdGraph, FVector2D(Location)))
+		{
+			if (FBlueprintNodeTemplateCache::IsTemplateOuter(ParentGraph))
+			{
+				return NewNode;
+			}
+			
+			FRigVMEditorAssetInterfacePtr Blueprint = FRigVMBlueprintUtils::FindAssetForGraph(ParentGraph);
+			URigVMController* Controller = Blueprint->GetController(ParentGraph);
+			if (FromPin)
+			{
+				if (URigVMEdGraphNode* SourceNode = Cast<URigVMEdGraphNode>(FromPin->GetOwningNode()))
+				{
+					if (URigVMPin* SourcePin = SourceNode->FindModelPinFromGraphPin(FromPin))
+					{
+						if (SourcePin->IsExecuteContext())
+						{
+							if (!NewNode->GetExecutePins().IsEmpty())
+							{
+								Controller->AddLink(SourcePin, NewNode->GetExecutePins()[0]);
+							}
+							else
+							{
+								if (URigVMPin* const* Pin = NewNode->GetInputPins().FindByPredicate([](URigVMPin* Pin)
+									{
+										return Pin->IsExecuteContext();
+									}))
+								{
+									Controller->AddLink(SourcePin, *Pin);
+								}
+							}
+						}
+						else
+						{
+							for (URigVMPin* TargetPin : NewNode->GetInputPins())
+							{
+								if (RigVMTypeUtils::AreCompatible(*SourcePin->GetCPPType(), SourcePin->GetCPPTypeObject(), *TargetPin->GetCPPType(), TargetPin->GetCPPTypeObject()))
+								{
+									Controller->AddLink(SourcePin, TargetPin);
+									break;
+								}
+							}
+						}
+					}
+				}
+			}
+
+			if (bSelectNewNode)
+			{
+				Controller->SelectNode(NewNode->GetModelNode(), true);
+			}
+		}
+	}
+
+	return nullptr;
+}
+
 void SRigVMActionMenu::Construct( const FArguments& InArgs, TSharedPtr<IRigVMEditor> InEditor )
 {
 	bActionExecuted = false;
@@ -83,6 +148,7 @@ void SRigVMActionMenu::Construct( const FArguments& InArgs, TSharedPtr<IRigVMEdi
 	this->GraphObj = InArgs._GraphObj;
 	this->DraggedFromPins = InArgs._DraggedFromPins;
 	this->NewNodePosition = InArgs._NewNodePosition;
+	this->OnClosedCallback = InArgs._OnClosedCallback;
 	this->EditorPtr = InEditor;
 
 	// Generate the context display; showing the user what they're picking something for
@@ -171,6 +237,7 @@ void SRigVMActionMenu::Construct( const FArguments& InArgs, TSharedPtr<IRigVMEdi
 						.OnGetActionList(this, &SRigVMActionMenu::OnGetActionList)
 						.DraggedFromPins(DraggedFromPins)
 						.GraphObj(GraphObj)
+						
 				]
 
 				// PROGRESS BAR
@@ -272,6 +339,26 @@ TSharedRef<FGraphActionListBuilderBase> SRigVMActionMenu::OnGetActionList()
 	}
 
 	ContextMenuBuilder = MakeShared<FRigVMActionMenuBuilder>(ConfigFlags);
+	FRigVMEditorAssetInterfacePtr Asset = BlueprintEditor->GetRigVMAssetInterface();
+	if (IRigVMEditorModule* RigVMEditorModule = Asset->GetEditorModule())
+	{
+		TArray<TSharedRef<URigVMEdGraphNodeSpawner>> NodeSpawners;
+		RigVMEditorModule->GetTypeActions(Asset, NodeSpawners);
+		RigVMEditorModule->GetInstanceActions(Asset, NodeSpawners);
+		
+		for (TSharedRef<URigVMEdGraphNodeSpawner>& NodeSpawner : NodeSpawners)
+		{
+			TSharedPtr<FRigVMGraphSchemaAction_NewNode> NewNodeAction(new FRigVMGraphSchemaAction_NewNode(
+					NodeSpawner->GetMenuCategory(),
+					NodeSpawner->GetMenuName(),
+					NodeSpawner->GetMenuTooltip(),
+					0
+				));
+			NewNodeAction->NodeSpawner = NodeSpawner;
+
+			ContextMenuBuilder->AddAction(NewNodeAction);
+		}
+	}
 
 	// NOTE: cannot call GetGraphContextActions() during serialization and GC due to its use of FindObject()
 	if (!GIsSavingPackage && !IsGarbageCollecting() && FilterContext.Blueprints.Num() > 0)
@@ -280,7 +367,13 @@ TSharedRef<FGraphActionListBuilderBase> SRigVMActionMenu::OnGetActionList()
 	}
 
 	// also try adding promote to variable if we can do so.
-	TryInsertPromoteToVariable(FilterContext, *ContextMenuBuilder); 
+	//TryInsertPromoteToVariable(FilterContext, *ContextMenuBuilder);
+
+	// give the schema the opportunity to add another action
+	if (const URigVMEdGraphSchema* Schema = Cast<const URigVMEdGraphSchema>(GraphObj->GetSchema()))
+	{
+		Schema->InsertAdditionalActions({Asset}, FilterContext.Graphs, FilterContext.Pins, *ContextMenuBuilder);
+	}
 
 	return ContextMenuBuilder.ToSharedRef();
 }
@@ -356,6 +449,11 @@ TSharedRef<SWidget> SRigVMActionMenu::OnCreateWidgetForAction(FCreateWidgetForAc
 				IconBrush = FRigVMEditorBase::GetVarIconAndColorFromProperty(Property, IconColor, SecondaryBrush, SecondaryIconColor);
 			}
 		}
+		else if (GraphAction->GetTypeId() == FRigVMGraphSchemaAction_NewNode::StaticGetTypeId())
+		{
+			FRigVMGraphSchemaAction_NewNode* NodeSpawnerAction = (FRigVMGraphSchemaAction_NewNode*)GraphAction.Get();
+			IconBrush = NodeSpawnerAction->NodeSpawner->GetMenuIcon().GetOptionalIcon();
+		}
 	}
 	TSharedRef<SWidget> IconWidget = SPinTypeSelector::ConstructPinTypeImage(
 		IconBrush, 
diff --git a/Engine/Plugins/Runtime/RigVM/Source/RigVMEditor/Public/Editor/SRigVMActionMenu.h b/Engine/Plugins/Runtime/RigVM/Source/RigVMEditor/Public/Editor/SRigVMActionMenu.h
--- a/Engine/Plugins/Runtime/RigVM/Source/RigVMEditor/Public/Editor/SRigVMActionMenu.h
+++ b/Engine/Plugins/Runtime/RigVM/Source/RigVMEditor/Public/Editor/SRigVMActionMenu.h
@@ -2,13 +2,34 @@
 
 
 #pragma once
+#include "RigVMActionMenuItem.h"
 #include "Editor/RigVMActionMenuBuilder.h"
 #include "SGraphActionMenu.h"
+#include "EdGraph/NodeSpawners/RigVMEdGraphNodeSpawner.h"
+#include "GraphEditor.h"
 
+#include "SRigVMActionMenu.generated.h"
 
 class SGraphActionMenu;
 class IRigVMEditor;
 
+/// Action to add a new RigVM node to the graph
+USTRUCT()
+struct FRigVMGraphSchemaAction_NewNode : public FEdGraphSchemaAction
+{
+	GENERATED_BODY()
+
+	// Inherit the base class's constructors
+	using FEdGraphSchemaAction::FEdGraphSchemaAction;
+
+	TSharedPtr<URigVMEdGraphNodeSpawner> NodeSpawner;
+
+	static FName StaticGetTypeId() { static FName Type("FRigVMGraphSchemaAction_NewNode"); return Type; }
+	FName GetTypeId() const override { return StaticGetTypeId(); }
+
+	// FEdGraphSchemaAction overrides
+	virtual UEdGraphNode* PerformAction(class UEdGraph* ParentGraph, UEdGraphPin* FromPin, const FVector2f& Location, bool bSelectNewNode = true) override;
+};
 
 /*******************************************************************************
 * SRigVMActionMenu
@@ -26,6 +47,7 @@ public:
 		SLATE_ARGUMENT( UEdGraph*, GraphObj )
 		SLATE_ARGUMENT( FDeprecateSlateVector2D, NewNodePosition )
 		SLATE_ARGUMENT( TArray<UEdGraphPin*>, DraggedFromPins )
+		SLATE_ARGUMENT( SGraphEditor::FActionMenuClosed, OnClosedCallback )
 	SLATE_END_ARGS()
 
 	void Construct( const FArguments& InArgs, TSharedPtr<IRigVMEditor> InEditor );
@@ -62,6 +84,7 @@ private:
 	TObjectPtr<UEdGraph> GraphObj = nullptr;
 	TArray<UEdGraphPin*> DraggedFromPins;
 	FDeprecateSlateVector2D NewNodePosition;
+	SGraphEditor::FActionMenuClosed OnClosedCallback;
 
 	TSharedPtr<SGraphActionMenu> GraphActionMenu;
 	TWeakPtr<IRigVMEditor> EditorPtr;
diff --git a/Engine/Plugins/Runtime/RigVM/Source/RigVMEditor/Public/Editor/SRigVMDetailsInspector.h b/Engine/Plugins/Runtime/RigVM/Source/RigVMEditor/Public/Editor/SRigVMDetailsInspector.h
--- a/Engine/Plugins/Runtime/RigVM/Source/RigVMEditor/Public/Editor/SRigVMDetailsInspector.h
+++ b/Engine/Plugins/Runtime/RigVM/Source/RigVMEditor/Public/Editor/SRigVMDetailsInspector.h
@@ -11,7 +11,7 @@
 #include "Internationalization/Text.h"
 #include "Layout/Visibility.h"
 #include "PropertyEditorDelegates.h"
-#include "RigVMAsset.h"
+#include "RigVMEditorAsset.h"
 #include "SlateFwd.h"
 #include "Styling/SlateTypes.h"
 #include "Templates/SharedPointer.h"
@@ -123,7 +123,7 @@ public:
 	/** returns the list of selected objects */
 	UE_API const TArray< TWeakObjectPtr<UObject> >& GetSelectedObjects() const;
 
-	UE_API void OnEditorClose(const IRigVMEditor* RigVMEditorBase, FRigVMAssetInterfacePtr RigVMBlueprint);
+	UE_API void OnEditorClose(const IRigVMEditor* RigVMEditorBase, FRigVMEditorAssetInterfacePtr RigVMBlueprint);
 
 	UE_API virtual void AddReferencedObjects(FReferenceCollector& Collector) override;
 	UE_API virtual FString GetReferencerName() const override;
diff --git a/Engine/Plugins/Runtime/RigVM/Source/RigVMEditor/Public/RigVMEditorModule.h b/Engine/Plugins/Runtime/RigVM/Source/RigVMEditor/Public/RigVMEditorModule.h
--- a/Engine/Plugins/Runtime/RigVM/Source/RigVMEditor/Public/RigVMEditorModule.h
+++ b/Engine/Plugins/Runtime/RigVM/Source/RigVMEditor/Public/RigVMEditorModule.h
@@ -7,7 +7,7 @@
 #pragma once
 
 #include "CoreMinimal.h"
-#include "RigVMAsset.h"
+#include "RigVMEditorAsset.h"
 #include "Modules/ModuleInterface.h"
 #include "Logging/LogMacros.h"
 #include "EdGraph/RigVMEdGraphPanelNodeFactory.h"
@@ -23,8 +23,9 @@
 
 #define UE_API RIGVMEDITOR_API
 
-class URigVMAssetInterface;
-class IRigVMAssetInterface;
+class URigVMEdGraphNodeSpawner;
+class URigVMEditorAssetInterface;
+class IRigVMEditorAssetInterface;
 class IRigVMEditorModule;
 class URigVMUserWorkflowOptions;
 class URigVMPin;
@@ -38,7 +39,7 @@ class FBlueprintActionDatabaseRegistrar;
 DECLARE_LOG_CATEGORY_EXTERN(LogRigVMEditor, Log, All);
 
 extern RIGVMEDITOR_API TAutoConsoleVariable<bool> CVarRigVMUseNewEditor;
-extern RIGVMEDITOR_API TAutoConsoleVariable<bool> CVarRigVMUseDualEditor;
+extern RIGVMEDITOR_API TAutoConsoleVariable<bool> CVarRigVMBlueprintIndependentAssets;
 
 /** Describes the reason for Refreshing the editor */
 namespace ERefreshRigVMEditorReason
@@ -57,7 +58,7 @@ namespace UE::RigVMEditor
 	struct FRigVMEditorFindNodeReferencesParams
 	{
 		FRigVMEditorFindNodeReferencesParams(
-			const TWeakInterfacePtr<IRigVMAssetInterface> InRigVMAsset,
+			const TWeakInterfacePtr<IRigVMEditorAssetInterface> InRigVMAsset,
 			const TWeakObjectPtr<const URigVMEdGraphNode> InEdGraphNode,
 			const bool bInSearchInAllBlueprints)
 			: WeakRigVMAsset(InRigVMAsset)
@@ -66,7 +67,7 @@ namespace UE::RigVMEditor
 		{}
 
 		/** The asset that contains the node */
-		const TWeakInterfacePtr<IRigVMAssetInterface> WeakRigVMAsset;
+		const TWeakInterfacePtr<IRigVMEditorAssetInterface> WeakRigVMAsset;
 
 		/** The ed graph node */
 		const TWeakObjectPtr<const URigVMEdGraphNode> WeakEdGraphNode;
@@ -89,8 +90,10 @@ public:
 	}
 
 	virtual void GetContextMenuActions(const URigVMEdGraphSchema* Schema, class UToolMenu* Menu, class UGraphNodeContextMenuContext* Context) const = 0;
-	virtual void GetTypeActions(FRigVMAssetInterfacePtr RigVMBlueprint, FBlueprintActionDatabaseRegistrar& ActionRegistrar) = 0;
-	virtual void GetInstanceActions(FRigVMAssetInterfacePtr RigVMBlueprint, FBlueprintActionDatabaseRegistrar& ActionRegistrar) = 0;
+	virtual void GetTypeActions(FRigVMEditorAssetInterfacePtr RigVMBlueprint, FBlueprintActionDatabaseRegistrar& ActionRegistrar) = 0;
+	virtual void GetTypeActions(FRigVMEditorAssetInterfacePtr RigVMBlueprint, TArray<TSharedRef<URigVMEdGraphNodeSpawner>>& OutGraphNodeSpawners) = 0;
+	virtual void GetInstanceActions(FRigVMEditorAssetInterfacePtr RigVMBlueprint, FBlueprintActionDatabaseRegistrar& ActionRegistrar) = 0;
+	virtual void GetInstanceActions(FRigVMEditorAssetInterfacePtr RigVMBlueprint, TArray<TSharedRef<URigVMEdGraphNodeSpawner>>& OutGraphNodeSpawners) = 0;
 	virtual void GetNodeContextMenuActions(TScriptInterface<IRigVMClientHost> RigVMClientHost, const URigVMEdGraphNode* EdGraphNode, URigVMNode* ModelNode, class UToolMenu* Menu) const = 0;
 	virtual void GetPinContextMenuActions(TScriptInterface<IRigVMClientHost> RigVMClientHost, const UEdGraphPin* EdGraphPin, URigVMPin* ModelPin, class UToolMenu* Menu) const = 0;
 	virtual FConnectionDrawingPolicy* CreateConnectionDrawingPolicy(int32 InBackLayerID, int32 InFrontLayerID, float InZoomFactor, const FSlateRect& InClippingRect, class FSlateWindowElementList& InDrawElements, class UEdGraph* InGraphObj) const = 0;
@@ -112,11 +115,13 @@ public:
 
 	/** RigVMEditorModule interface */
 	UE_DEPRECATED(5.7, "Please update your code")
-	UE_API virtual UClass* GetRigVMBlueprintClass() const { return nullptr; }
+	virtual UClass* GetRigVMBlueprintClass() const { return nullptr; }
 	UE_DEPRECATED(5.7, "Please update your code")
-	UE_API const URigVMBlueprint* GetRigVMBlueprintCDO() const { return nullptr; }
-	UE_API virtual void GetTypeActions(FRigVMAssetInterfacePtr RigVMBlueprint, FBlueprintActionDatabaseRegistrar& ActionRegistrar) override;
-	UE_API virtual void GetInstanceActions(FRigVMAssetInterfacePtr RigVMBlueprint, FBlueprintActionDatabaseRegistrar& ActionRegistrar) override;
+	const URigVMBlueprint* GetRigVMBlueprintCDO() const { return nullptr; }
+	UE_API virtual void GetTypeActions(FRigVMEditorAssetInterfacePtr RigVMBlueprint, FBlueprintActionDatabaseRegistrar& ActionRegistrar) override;
+	UE_API virtual void GetTypeActions(FRigVMEditorAssetInterfacePtr RigVMBlueprint, TArray<TSharedRef<URigVMEdGraphNodeSpawner>>& OutGraphNodeSpawners) override;
+	UE_API virtual void GetInstanceActions(FRigVMEditorAssetInterfacePtr RigVMBlueprint, FBlueprintActionDatabaseRegistrar& ActionRegistrar) override;
+	UE_API virtual void GetInstanceActions(FRigVMEditorAssetInterfacePtr RigVMBlueprint, TArray<TSharedRef<URigVMEdGraphNodeSpawner>>& OutGraphNodeSpawners) override;
 	UE_API virtual void GetNodeContextMenuActions(TScriptInterface<IRigVMClientHost> RigVMClientHost, const URigVMEdGraphNode* EdGraphNode, URigVMNode* ModelNode, class UToolMenu* Menu) const override;
 	UE_API virtual void GetPinContextMenuActions(TScriptInterface<IRigVMClientHost> RigVMClientHost, const UEdGraphPin* EdGraphPin, URigVMPin* ModelPin, class UToolMenu* Menu) const override;
 	virtual bool AssetsPublicFunctionsAllowed(const FAssetData& InAssetData) const override { return true; }
@@ -138,7 +143,7 @@ public:
 	UE_API virtual void GetContextMenuActions(const URigVMEdGraphSchema* Schema, class UToolMenu* Menu, class UGraphNodeContextMenuContext* Context) const override;
 
 	/** Make sure to create the root graph for a given blueprint */
-	UE_API void CreateRootGraphIfRequired(FRigVMAssetInterfacePtr InBlueprint) const;
+	UE_API void CreateRootGraphIfRequired(FRigVMEditorAssetInterfacePtr InBlueprint) const;
 
 	UE_API virtual FConnectionDrawingPolicy* CreateConnectionDrawingPolicy(int32 InBackLayerID, int32 InFrontLayerID, float InZoomFactor, const FSlateRect& InClippingRect, FSlateWindowElementList& InDrawElements, UEdGraph* InGraphObj) const override;
 
@@ -148,9 +153,10 @@ public:
 protected:
 
 	UE_API bool IsRigVMEditorModuleBase() const;
-	UE_API virtual UClass* GetRigVMAssetClass() const { return URigVMAssetInterface::StaticClass(); }
-	UE_API virtual UClass* GetRigVMEdGraphSchemaClass() const { return URigVMEdGraphSchema::StaticClass(); }
-	UE_API virtual UScriptStruct* GetRigVMExecuteContextStruct() const { return FRigVMExecuteContext::StaticStruct(); }
+	virtual UClass* GetRigVMLegacyBlueprintClass() const { return URigVMEditorAssetInterface::StaticClass(); }
+	virtual UClass* GetRigVMAssetClass() const { return URigVMEditorAssetInterface::StaticClass(); }
+	virtual UClass* GetRigVMEdGraphSchemaClass() const { return URigVMEdGraphSchema::StaticClass(); }
+	virtual UScriptStruct* GetRigVMExecuteContextStruct() const { return FRigVMExecuteContext::StaticStruct(); }
 
 	/** Specific section callbacks for the context menu */
 	UE_API virtual void GetNodeActionsContextMenuActions(TScriptInterface<IRigVMClientHost> RigVMClientHost, const URigVMEdGraphNode* EdGraphNode, URigVMNode* ModelNode, UToolMenu* Menu) const;
diff --git a/Engine/Plugins/Runtime/RigVM/Source/RigVMEditor/Public/Widgets/RigVMTreeToolkitNode.h b/Engine/Plugins/Runtime/RigVM/Source/RigVMEditor/Public/Widgets/RigVMTreeToolkitNode.h
--- a/Engine/Plugins/Runtime/RigVM/Source/RigVMEditor/Public/Widgets/RigVMTreeToolkitNode.h
+++ b/Engine/Plugins/Runtime/RigVM/Source/RigVMEditor/Public/Widgets/RigVMTreeToolkitNode.h
@@ -30,9 +30,9 @@ public:
 	}
 
 	virtual FAssetData GetAssetData() const;
-	UE_DEPRECATED(5.7, "Please use IRigVMAssetInterface* GetAsset() const")
+	UE_DEPRECATED(5.7, "Please use IRigVMEditorAssetInterface* GetAsset() const")
 	virtual URigVMBlueprint* GetBlueprint() const { return Cast<URigVMBlueprint>(GetAsset()->GetObject()); }
-	virtual IRigVMAssetInterface* GetAsset() const;
+	virtual IRigVMEditorAssetInterface* GetAsset() const;
 
 	virtual FText GetLabel() const;
 
diff --git a/Engine/Plugins/Runtime/RigVM/Source/RigVMEditor/Public/Widgets/RigVMTreeToolkitTask.h b/Engine/Plugins/Runtime/RigVM/Source/RigVMEditor/Public/Widgets/RigVMTreeToolkitTask.h
--- a/Engine/Plugins/Runtime/RigVM/Source/RigVMEditor/Public/Widgets/RigVMTreeToolkitTask.h
+++ b/Engine/Plugins/Runtime/RigVM/Source/RigVMEditor/Public/Widgets/RigVMTreeToolkitTask.h
@@ -23,13 +23,13 @@ public:
 		return false;
 	};
 
-	UE_DEPRECATED(5.7, "Please use FRigVMAssetInterfacePtr GetRigVMAsset(const TSharedRef<FRigVMTreePhase>& InPhase) const")
+	UE_DEPRECATED(5.7, "Please use FRigVMEditorAssetInterfacePtr GetRigVMAsset(const TSharedRef<FRigVMTreePhase>& InPhase) const")
 	virtual URigVMBlueprint* GetBlueprint(const TSharedRef<FRigVMTreePhase>& InPhase) const
 	{
 		return nullptr;
 	}
 
-	virtual FRigVMAssetInterfacePtr GetRigVMAsset(const TSharedRef<FRigVMTreePhase>& InPhase) const
+	virtual FRigVMEditorAssetInterfacePtr GetRigVMAsset(const TSharedRef<FRigVMTreePhase>& InPhase) const
 	{
 		return nullptr;
 	}
diff --git a/Engine/Plugins/Runtime/RigVM/Source/RigVMEditor/Public/Widgets/SRigVMEditorGraphExplorer.h b/Engine/Plugins/Runtime/RigVM/Source/RigVMEditor/Public/Widgets/SRigVMEditorGraphExplorer.h
--- a/Engine/Plugins/Runtime/RigVM/Source/RigVMEditor/Public/Widgets/SRigVMEditorGraphExplorer.h
+++ b/Engine/Plugins/Runtime/RigVM/Source/RigVMEditor/Public/Widgets/SRigVMEditorGraphExplorer.h
@@ -4,7 +4,7 @@
 
 #include "CoreMinimal.h"
 #include "GraphEditorDragDropAction.h"
-#include "RigVMAsset.h"
+#include "RigVMEditorAsset.h"
 #include "RigVMBlueprintLegacy.h"
 #include "DragAndDrop/GraphNodeDragDropOp.h"
 #include "Widgets/DeclarativeSyntaxSupport.h"
@@ -20,7 +20,7 @@ namespace UE::RigVMEditor { class FRigVMEdGraphNodeRegistry; }
 //
 class URigVMGraph;
 class URigVMEdGraph;
-class IRigVMAssetInterface;
+class IRigVMEditorAssetInterface;
 class SGraphActionMenu;
 struct FCreateWidgetForActionData;
 class FRigVMEditorGraphExplorerTreeElement;
@@ -34,22 +34,22 @@ class FRigVMGraphExplorerDragDropOp : public FGraphEditorDragDropAction
 public:
 	DRAG_DROP_OPERATOR_TYPE(FRigVMExplorerDragDropOp, FGraphEditorDragDropAction)
 
-	UE_DEPRECATED(5.7, "Please use New(const FRigVMExplorerElementKey& InElement, FRigVMAssetInterfacePtr InBlueprint);")
+	UE_DEPRECATED(5.7, "Please use New(const FRigVMExplorerElementKey& InElement, FRigVMEditorAssetInterfacePtr InBlueprint);")
 	static TSharedRef<FRigVMGraphExplorerDragDropOp> New(const FRigVMExplorerElementKey& InElement, TObjectPtr<URigVMBlueprint> InBlueprint);
-	static TSharedRef<FRigVMGraphExplorerDragDropOp> New(const FRigVMExplorerElementKey& InElement, FRigVMAssetInterfacePtr InBlueprint);
+	static TSharedRef<FRigVMGraphExplorerDragDropOp> New(const FRigVMExplorerElementKey& InElement, FRigVMEditorAssetInterfacePtr InBlueprint);
 
 	const FRigVMExplorerElementKey& GetElement() const
 	{
 		return Element;
 	}
 
-	UE_DEPRECATED(5.7, "Please use FRigVMAssetInterfacePtr GetRigVMAssetInterface() const")
+	UE_DEPRECATED(5.7, "Please use FRigVMEditorAssetInterfacePtr GetRigVMAssetInterface() const")
 	TObjectPtr<URigVMBlueprint> GetBlueprint() const
 	{
 		return Cast<URigVMBlueprint>(GetRigVMAssetInterface().GetObject());
 	}
 
-	FRigVMAssetInterfacePtr GetRigVMAssetInterface() const
+	FRigVMEditorAssetInterfacePtr GetRigVMAssetInterface() const
 	{
 		return SourceBlueprint;
 	}
@@ -62,7 +62,7 @@ public:
 private:
 
 	FRigVMExplorerElementKey Element;
-	FRigVMAssetInterfacePtr SourceBlueprint;
+	FRigVMEditorAssetInterfacePtr SourceBlueprint;
 };
 
 class FRigVMEditorGraphExplorerCommands : public TCommands<FRigVMEditorGraphExplorerCommands>
@@ -117,6 +117,8 @@ public:
 	/** Returns a registry for Ed Graph Nodes that hold RigVMVariableNodes */
 	const TSharedPtr<FRigVMEdGraphNodeRegistry>& GetEdGraphNodeVariableRegistry() const { return EdGraphNodeVariableRegistry; }
 
+	void HandleGraphModifiedEvent(ERigVMGraphNotifType InNotifType, URigVMGraph* InGraph, UObject* InSubject);
+
 private:
 	bool bNeedsRefresh;
 
diff --git a/Engine/Plugins/Runtime/RigVM/Source/RigVMEditor/Public/Widgets/SRigVMEditorSelectedDebugObjectWidget.h b/Engine/Plugins/Runtime/RigVM/Source/RigVMEditor/Public/Widgets/SRigVMEditorSelectedDebugObjectWidget.h
--- a/Engine/Plugins/Runtime/RigVM/Source/RigVMEditor/Public/Widgets/SRigVMEditorSelectedDebugObjectWidget.h
+++ b/Engine/Plugins/Runtime/RigVM/Source/RigVMEditor/Public/Widgets/SRigVMEditorSelectedDebugObjectWidget.h
@@ -62,7 +62,7 @@ public:
 	UE_API void AddDebugObject(UObject* TestObject, const FString& TestObjectName = FString());
 
 private:
-	UE_API FRigVMAssetInterfacePtr GetRigVMAssetInterface() const;
+	UE_API FRigVMEditorAssetInterfacePtr GetRigVMAssetInterface() const;
 
 	/** Creates a list of all debug objects **/
 	UE_API void GenerateDebugObjectInstances(bool bRestoreSelection);
diff --git a/Engine/Plugins/Runtime/RigVM/Source/RigVMEditor/Public/Widgets/SRigVMExecutionStackView.h b/Engine/Plugins/Runtime/RigVM/Source/RigVMEditor/Public/Widgets/SRigVMExecutionStackView.h
--- a/Engine/Plugins/Runtime/RigVM/Source/RigVMEditor/Public/Widgets/SRigVMExecutionStackView.h
+++ b/Engine/Plugins/Runtime/RigVM/Source/RigVMEditor/Public/Widgets/SRigVMExecutionStackView.h
@@ -34,10 +34,10 @@ class FRigStackEntry : public TSharedFromThis<FRigStackEntry>
 public:
 	UE_API FRigStackEntry(int32 InEntryIndex, ERigStackEntry::Type InEntryType, int32 InInstructionIndex, ERigVMOpCode InOpCode, const FString& InLabel, const FRigVMASTProxy& InProxy);
 
-	UE_DEPRECATED(5.7, "Use TSharedRef<ITableRow> MakeTreeRowWidget(const TSharedRef<STableViewBase>& InOwnerTable, TSharedRef<FRigStackEntry> InEntry, TSharedRef<FUICommandList> InCommandList, TSharedPtr<SRigVMExecutionStackView> InStackView, TWeakInterfacePtr<IRigVMAssetInterface> InBlueprint); instead")
+	UE_DEPRECATED(5.7, "Use TSharedRef<ITableRow> MakeTreeRowWidget(const TSharedRef<STableViewBase>& InOwnerTable, TSharedRef<FRigStackEntry> InEntry, TSharedRef<FUICommandList> InCommandList, TSharedPtr<SRigVMExecutionStackView> InStackView, TWeakInterfacePtr<IRigVMEditorAssetInterface> InBlueprint); instead")
 	UE_API TSharedRef<ITableRow> MakeTreeRowWidget(const TSharedRef<STableViewBase>& InOwnerTable, TSharedRef<FRigStackEntry> InEntry, TSharedRef<FUICommandList> InCommandList, TSharedPtr<SRigVMExecutionStackView> InStackView, TWeakObjectPtr<UObject> InBlueprint);
 	UE_API TSharedRef<ITableRow> MakeTreeRowWidget(const TSharedRef<STableViewBase>& InOwnerTable, TSharedRef<FRigStackEntry> InEntry, TSharedRef<FUICommandList>
-	                                               InCommandList, TSharedPtr<SRigVMExecutionStackView> InStackView, TWeakInterfacePtr<IRigVMAssetInterface> InBlueprint);
+	                                               InCommandList, TSharedPtr<SRigVMExecutionStackView> InStackView, TWeakInterfacePtr<IRigVMEditorAssetInterface> InBlueprint);
 
 	int32 EntryIndex;
 	ERigStackEntry::Type EntryType;
@@ -59,11 +59,11 @@ class SRigStackItem : public STableRow<TSharedPtr<FRigStackEntry>>
 	SLATE_BEGIN_ARGS(SRigStackItem) {}
 	SLATE_END_ARGS()
 
-	void Construct(const FArguments& InArgs, const TSharedRef<STableViewBase>& OwnerTable, TSharedRef<FRigStackEntry> InStackEntry, TSharedRef<FUICommandList> InCommandList, TWeakInterfacePtr<IRigVMAssetInterface> InBlueprint);
+	void Construct(const FArguments& InArgs, const TSharedRef<STableViewBase>& OwnerTable, TSharedRef<FRigStackEntry> InStackEntry, TSharedRef<FUICommandList> InCommandList, TWeakInterfacePtr<IRigVMEditorAssetInterface> InBlueprint);
 
 private:
 	TWeakPtr<FRigStackEntry> WeakStackEntry;
-	TWeakInterfacePtr<IRigVMAssetInterface> WeakBlueprint;
+	TWeakInterfacePtr<IRigVMEditorAssetInterface> WeakBlueprint;
 	TWeakPtr<FUICommandList> WeakCommandList;
 
 	FText GetIndexText() const;
@@ -100,7 +100,7 @@ private:
 	UE_API void BindCommands();
 
 	/** Make a row widget for the table */
-	UE_API TSharedRef<ITableRow> MakeTableRowWidget(TSharedPtr<FRigStackEntry> InItem, const TSharedRef<STableViewBase>& OwnerTable, TWeakInterfacePtr<IRigVMAssetInterface> InBlueprint);
+	UE_API TSharedRef<ITableRow> MakeTableRowWidget(TSharedPtr<FRigStackEntry> InItem, const TSharedRef<STableViewBase>& OwnerTable, TWeakInterfacePtr<IRigVMEditorAssetInterface> InBlueprint);
 
 	/** Get children for the tree */
 	UE_API void HandleGetChildrenForTree(TSharedPtr<FRigStackEntry> InItem, TArray<TSharedPtr<FRigStackEntry>>& OutChildren);
@@ -149,7 +149,7 @@ private:
 	TSharedPtr<FUICommandList> CommandList;
 
 	TWeakPtr<IRigVMEditor> WeakEditor;
-	TWeakInterfacePtr<IRigVMAssetInterface> WeakRigVMBlueprint;
+	TWeakInterfacePtr<IRigVMEditorAssetInterface> WeakRigVMBlueprint;
 
 	TArray<TSharedPtr<FRigStackEntry>> Operators;
 
diff --git a/Engine/Plugins/Runtime/RigVM/Source/RigVMEditor/Public/Widgets/SRigVMGraphFunctionBulkEditWidget.h b/Engine/Plugins/Runtime/RigVM/Source/RigVMEditor/Public/Widgets/SRigVMGraphFunctionBulkEditWidget.h
--- a/Engine/Plugins/Runtime/RigVM/Source/RigVMEditor/Public/Widgets/SRigVMGraphFunctionBulkEditWidget.h
+++ b/Engine/Plugins/Runtime/RigVM/Source/RigVMEditor/Public/Widgets/SRigVMGraphFunctionBulkEditWidget.h
@@ -18,7 +18,7 @@ public:
 	SLATE_BEGIN_ARGS(SRigVMGraphFunctionBulkEditWidget) {}
 	SLATE_END_ARGS()
 
-	UE_API void Construct(const FArguments& InArgs, FRigVMAssetInterfacePtr InBlueprint, URigVMController* InController, URigVMLibraryNode* InFunction, ERigVMControllerBulkEditType InEditType);
+	UE_API void Construct(const FArguments& InArgs, FRigVMEditorAssetInterfacePtr InBlueprint, URigVMController* InController, URigVMLibraryNode* InFunction, ERigVMControllerBulkEditType InEditType);
 
 private:
 
@@ -31,7 +31,7 @@ private:
 	/** Show a progress bar and load all of the assets */
 	UE_API void LoadAffectedAssets();
 
-	FRigVMAssetInterfacePtr Blueprint;
+	FRigVMEditorAssetInterfacePtr Blueprint;
 	URigVMController* Controller;
 	URigVMLibraryNode* Function;
 	ERigVMControllerBulkEditType EditType;
@@ -49,7 +49,7 @@ public:
 	{
 	}
 
-	SLATE_ARGUMENT(FRigVMAssetInterfacePtr, Blueprint)
+	SLATE_ARGUMENT(FRigVMEditorAssetInterfacePtr, Blueprint)
 	SLATE_ARGUMENT(URigVMController*, Controller)
 	SLATE_ARGUMENT(URigVMLibraryNode*, Function)
 	SLATE_ARGUMENT(ERigVMControllerBulkEditType, EditType)
diff --git a/Engine/Plugins/Runtime/RigVM/Source/RigVMEditor/Public/Widgets/SRigVMGraphNode.h b/Engine/Plugins/Runtime/RigVM/Source/RigVMEditor/Public/Widgets/SRigVMGraphNode.h
--- a/Engine/Plugins/Runtime/RigVM/Source/RigVMEditor/Public/Widgets/SRigVMGraphNode.h
+++ b/Engine/Plugins/Runtime/RigVM/Source/RigVMEditor/Public/Widgets/SRigVMGraphNode.h
@@ -130,7 +130,7 @@ protected:
 
 	UE_API virtual TOptional<FSlateColor> GetHighlightColor(const SGraphPin* InGraphPin) const override;
 	UE_API FSlateColor GetVariableLabelTextColor(TWeakObjectPtr<URigVMFunctionReferenceNode> FunctionReferenceNode, FName InVariableName) const;
-	UE_API FText GetVariableLabelTooltipText(TWeakInterfacePtr<IRigVMAssetInterface> InBlueprint, FName InVariableName) const;
+	UE_API FText GetVariableLabelTooltipText(TWeakInterfacePtr<IRigVMEditorAssetInterface> InBlueprint, FName InVariableName) const;
 
 	UE_API FReply HandleAddArrayElement(FString InModelPinPath);
 
@@ -177,7 +177,7 @@ protected:
 	/** Cache the node title so we can invalidate it */
 	TSharedPtr<SNodeTitle> NodeTitle;
 
-	TWeakInterfacePtr<IRigVMAssetInterface> Blueprint;
+	TWeakInterfacePtr<IRigVMEditorAssetInterface> Blueprint;
 
 	FVector2D LastHighDetailSize;
 
diff --git a/Engine/Plugins/Runtime/RigVM/Source/RigVMEditor/Public/Widgets/SRigVMGraphPinUserDataPath.h b/Engine/Plugins/Runtime/RigVM/Source/RigVMEditor/Public/Widgets/SRigVMGraphPinUserDataPath.h
--- a/Engine/Plugins/Runtime/RigVM/Source/RigVMEditor/Public/Widgets/SRigVMGraphPinUserDataPath.h
+++ b/Engine/Plugins/Runtime/RigVM/Source/RigVMEditor/Public/Widgets/SRigVMGraphPinUserDataPath.h
@@ -42,7 +42,7 @@ protected:
 	UE_API void FillUserDataPathMenu( FMenuBuilder& InMenuBuilder, FString InParentPath );
 	UE_API void HandleSetUserDataPath(FString InUserDataPath);
 
-	UE_API FRigVMAssetInterfacePtr GetRigVMAssetInterface() const;
+	UE_API FRigVMEditorAssetInterfacePtr GetRigVMAssetInterface() const;
 	UE_API FString GetUserDataNameSpace(URigVMPin* ModelPin) const;
 	UE_API FString GetUserDataNameSpace() const;
 	UE_API const UNameSpacedUserData* GetUserDataObject() const;
diff --git a/Engine/Plugins/Runtime/RigVM/Source/RigVMEditor/Public/Widgets/SRigVMGraphPinVariableBinding.h b/Engine/Plugins/Runtime/RigVM/Source/RigVMEditor/Public/Widgets/SRigVMGraphPinVariableBinding.h
--- a/Engine/Plugins/Runtime/RigVM/Source/RigVMEditor/Public/Widgets/SRigVMGraphPinVariableBinding.h
+++ b/Engine/Plugins/Runtime/RigVM/Source/RigVMEditor/Public/Widgets/SRigVMGraphPinVariableBinding.h
@@ -29,7 +29,7 @@ public:
 		SLATE_ARGUMENT(TArray<URigVMPin*>, ModelPins)
 		SLATE_ARGUMENT(URigVMFunctionReferenceNode*, FunctionReferenceNode)
 		SLATE_ARGUMENT(FName, InnerVariableName)
-		SLATE_ARGUMENT(FRigVMAssetInterfacePtr, Asset)
+		SLATE_ARGUMENT(FRigVMEditorAssetInterfacePtr, Asset)
 		SLATE_ARGUMENT_DEPRECATED(URigVMBlueprint*, Blueprint, 5.7, "Please use Asset instead")
 		SLATE_ARGUMENT(bool, CanRemoveBinding)
 
@@ -54,7 +54,7 @@ protected:
 	TArray<URigVMPin*> ModelPins;
 	URigVMFunctionReferenceNode* FunctionReferenceNode;
 	FName InnerVariableName;
-	FRigVMAssetInterfacePtr Blueprint;
+	FRigVMEditorAssetInterfacePtr Blueprint;
 	FPropertyBindingWidgetArgs BindingArgs;
 	bool bCanRemoveBinding;
 };
@@ -67,7 +67,7 @@ public:
 
 		SLATE_ARGUMENT(TArray<URigVMPin*>, ModelPins)
 		SLATE_ARGUMENT_DEPRECATED(URigVMBlueprint*, Blueprint, 5.7, "Please use Asset instead")
-		SLATE_ARGUMENT(FRigVMAssetInterfacePtr, Asset)
+		SLATE_ARGUMENT(FRigVMEditorAssetInterfacePtr, Asset)
 
 	SLATE_END_ARGS()
 
@@ -80,7 +80,7 @@ protected:
 	//~ End SGraphPin Interface
 
 	TArray<URigVMPin*> ModelPins;
-	FRigVMAssetInterfacePtr Blueprint;
+	FRigVMEditorAssetInterfacePtr Blueprint;
 };
 
 #undef UE_API
diff --git a/Engine/Plugins/Runtime/RigVM/Source/RigVMEditor/Public/Widgets/SRigVMLogWidget.h b/Engine/Plugins/Runtime/RigVM/Source/RigVMEditor/Public/Widgets/SRigVMLogWidget.h
--- a/Engine/Plugins/Runtime/RigVM/Source/RigVMEditor/Public/Widgets/SRigVMLogWidget.h
+++ b/Engine/Plugins/Runtime/RigVM/Source/RigVMEditor/Public/Widgets/SRigVMLogWidget.h
@@ -40,8 +40,8 @@ public:
 	void UnbindLog(const FName& InLogName);
 	void BindLog(TSharedPtr<IMessageLogListing> InListing);
 	void UnbindLog(TSharedPtr<IMessageLogListing> InListing);
-	void BindLog(const FRigVMAssetInterfacePtr InRigVMBlueprint);
-	void UnbindLog(const FRigVMAssetInterfacePtr InRigVMBlueprint);
+	void BindLog(const FRigVMEditorAssetInterfacePtr InRigVMBlueprint);
+	void UnbindLog(const FRigVMEditorAssetInterfacePtr InRigVMBlueprint);
 
 	TSharedRef<FMessageLogListingModel> GetListing() const
 	{
diff --git a/Engine/Plugins/Runtime/RigVM/Source/RigVMEditor/Public/Widgets/SRigVMSwapFunctionsWidget.h b/Engine/Plugins/Runtime/RigVM/Source/RigVMEditor/Public/Widgets/SRigVMSwapFunctionsWidget.h
--- a/Engine/Plugins/Runtime/RigVM/Source/RigVMEditor/Public/Widgets/SRigVMSwapFunctionsWidget.h
+++ b/Engine/Plugins/Runtime/RigVM/Source/RigVMEditor/Public/Widgets/SRigVMSwapFunctionsWidget.h
@@ -272,7 +272,7 @@ public:
 private:
 
 	URigVMFunctionReferenceNode* GetReferenceNode(const TSharedRef<FRigVMTreePhase>& InPhase) const;
-	virtual FRigVMAssetInterfacePtr GetRigVMAsset(const TSharedRef<FRigVMTreePhase>& InPhase) const override;
+	virtual FRigVMEditorAssetInterfacePtr GetRigVMAsset(const TSharedRef<FRigVMTreePhase>& InPhase) const override;
 	
 	FString ObjectPath;
 	FRigVMGraphFunctionIdentifier Identifier;
