Control Rig: Fix create example asset from blueprint failing to create variables with spaces
#rnx

#Synced-CL 48388747

diff --git a/Engine/Plugins/Runtime/RigVM/Source/RigVM/Private/RigVMRuntimeAsset.cpp b/Engine/Plugins/Runtime/RigVM/Source/RigVM/Private/RigVMRuntimeAsset.cpp
--- a/Engine/Plugins/Runtime/RigVM/Source/RigVM/Private/RigVMRuntimeAsset.cpp
+++ b/Engine/Plugins/Runtime/RigVM/Source/RigVM/Private/RigVMRuntimeAsset.cpp
@@ -394,7 +394,7 @@ FName URigVMRuntimeAsset::AddHostMemberVariableFromExternal(FRigVMExternalVariab
 	UObject* CPPTypeObject = nullptr;
 	RigVMTypeUtils::CPPTypeFromExternalVariable(InVariableToCreate, CPPType, &CPPTypeObject);
 
-	FRigVMPropertyDescription Description(InVariableToCreate.Name, CPPType, CPPTypeObject, InDefaultValue);
+	FRigVMPropertyDescription Description(InVariableToCreate.Name, CPPType, CPPTypeObject, InDefaultValue, true);
 	Description.PropertyFlags = CPF_Edit | CPF_BlueprintVisible | CPF_DisableEditOnInstance | CPF_Interp;
 	if (InVariableToCreate.bIsPublic)
 	{
