[RigVM] Added Any/All nodes to MathBool library

#Synced-CL 48392406

diff --git a/Engine/Plugins/Runtime/RigVM/Source/RigVM/Private/RigVMFunctions/Math/RigVMFunction_MathBool.cpp b/Engine/Plugins/Runtime/RigVM/Source/RigVM/Private/RigVMFunctions/Math/RigVMFunction_MathBool.cpp
--- a/Engine/Plugins/Runtime/RigVM/Source/RigVM/Private/RigVMFunctions/Math/RigVMFunction_MathBool.cpp
+++ b/Engine/Plugins/Runtime/RigVM/Source/RigVM/Private/RigVMFunctions/Math/RigVMFunction_MathBool.cpp
@@ -132,3 +132,27 @@ FRigVMFunction_MathBoolToInteger_Execute()
 {
 	Result = Value ? 1 : 0;
 }
+
+FRigVMFunction_MathBoolAny_Execute()
+{
+	Result = false;
+	for (const bool& Item : Array)
+	{
+		if (Item)
+		{
+			Result = true;
+			break;
+		}
+	}
+	IsEmpty = Array.IsEmpty();
+}
+
+FRigVMFunction_MathBoolAll_Execute()
+{
+	Result = true;
+	for (const bool& Item : Array)
+	{
+		Result &= Item;
+	}
+	IsEmpty = Array.IsEmpty();
+}
\ No newline at end of file
diff --git a/Engine/Plugins/Runtime/RigVM/Source/RigVM/Private/RigVMFunctions/Math/RigVMFunction_MathBoolTests.cpp b/Engine/Plugins/Runtime/RigVM/Source/RigVM/Private/RigVMFunctions/Math/RigVMFunction_MathBoolTests.cpp
--- a/Engine/Plugins/Runtime/RigVM/Source/RigVM/Private/RigVMFunctions/Math/RigVMFunction_MathBoolTests.cpp
+++ b/Engine/Plugins/Runtime/RigVM/Source/RigVM/Private/RigVMFunctions/Math/RigVMFunction_MathBoolTests.cpp
@@ -121,4 +121,38 @@ IMPLEMENT_RIGVMSTRUCT_AUTOMATION_TEST(FRigVMFunction_MathBoolNotEquals)
 	return true;
 }
 
+IMPLEMENT_RIGVMSTRUCT_AUTOMATION_TEST(FRigVMFunction_MathBoolAny)
+{
+	Unit.Array = { true, false, false };
+	Execute();
+	AddErrorIfFalse(Unit.Result == true, TEXT("unexpected result"));
+	AddErrorIfFalse(Unit.IsEmpty == false, TEXT("unexpected result"));
+	Unit.Array = { false, false, false };
+	Execute();
+	AddErrorIfFalse(Unit.Result == false, TEXT("unexpected result"));
+	AddErrorIfFalse(Unit.IsEmpty == false, TEXT("unexpected result"));
+	Unit.Array = {};
+	Execute();
+	AddErrorIfFalse(Unit.Result == false, TEXT("unexpected result"));
+	AddErrorIfFalse(Unit.IsEmpty == true, TEXT("unexpected result"));
+	return true;
+}
+
+IMPLEMENT_RIGVMSTRUCT_AUTOMATION_TEST(FRigVMFunction_MathBoolAll)
+{
+	Unit.Array = { true, false, true };
+	Execute();
+	AddErrorIfFalse(Unit.Result == false, TEXT("unexpected result"));
+	AddErrorIfFalse(Unit.IsEmpty == false, TEXT("unexpected result"));
+	Unit.Array = { true, true, true };
+	Execute();
+	AddErrorIfFalse(Unit.Result == true, TEXT("unexpected result"));
+	AddErrorIfFalse(Unit.IsEmpty == false, TEXT("unexpected result"));
+	Unit.Array = {};
+	Execute();
+	AddErrorIfFalse(Unit.Result == true, TEXT("unexpected result"));
+	AddErrorIfFalse(Unit.IsEmpty == true, TEXT("unexpected result"));
+	return true;
+}
+
 #endif
diff --git a/Engine/Plugins/Runtime/RigVM/Source/RigVM/Public/RigVMFunctions/Math/RigVMFunction_MathBool.h b/Engine/Plugins/Runtime/RigVM/Source/RigVM/Public/RigVMFunctions/Math/RigVMFunction_MathBool.h
--- a/Engine/Plugins/Runtime/RigVM/Source/RigVM/Public/RigVMFunctions/Math/RigVMFunction_MathBool.h
+++ b/Engine/Plugins/Runtime/RigVM/Source/RigVM/Public/RigVMFunctions/Math/RigVMFunction_MathBool.h
@@ -437,3 +437,65 @@ struct FRigVMFunction_MathBoolToInteger : public FRigVMFunction_MathBoolBase
 	UPROPERTY(meta= (Output))
 	int32 Result;
 };
+
+/**
+* Returns true if any item in the array is true
+*/
+USTRUCT(meta = (DisplayName = "Any", Keywords = "Any"))
+struct FRigVMFunction_MathBoolAny : public FRigVMFunction_MathBoolBase
+{
+	GENERATED_BODY()
+
+	FRigVMFunction_MathBoolAny()
+	{
+		Array.Reset();
+		Result = false;
+	}
+
+	RIGVM_METHOD()
+	RIGVM_API virtual void Execute() override;
+
+	//The input array
+	UPROPERTY(meta = (Input))
+	TArray<bool> Array;
+
+	//The resulting value. False if input array is empty
+	UPROPERTY(meta = (Output))
+	bool Result;
+
+	//True if the array is empty
+	UPROPERTY(meta = (Output))
+	bool IsEmpty;
+
+};
+
+/**
+* Returns true if all the items in the array are true
+*/
+USTRUCT(meta = (DisplayName = "All", Keywords = "All"))
+struct FRigVMFunction_MathBoolAll : public FRigVMFunction_MathBoolBase
+{
+	GENERATED_BODY()
+
+	FRigVMFunction_MathBoolAll()
+	{
+		Array.Reset();
+		Result = false;
+	}
+
+	RIGVM_METHOD()
+	RIGVM_API virtual void Execute() override;
+
+	//The input array
+	UPROPERTY(meta = (Input))
+	TArray<bool> Array;
+
+	//The resulting value. True if input array is empty
+	UPROPERTY(meta = (Output))
+	bool Result;
+
+	//True if the array is empty
+	UPROPERTY(meta = (Output))
+	bool IsEmpty;
+
+};
\ No newline at end of file
