RigVM Variable Mapping - Remove unreachable code

#rnx

#Synced-CL 48390870

diff --git a/Engine/Plugins/Runtime/RigVM/Source/RigVMEditor/Private/Widgets/SRigVMVariableMappingWidget.cpp b/Engine/Plugins/Runtime/RigVM/Source/RigVMEditor/Private/Widgets/SRigVMVariableMappingWidget.cpp
--- a/Engine/Plugins/Runtime/RigVM/Source/RigVMEditor/Private/Widgets/SRigVMVariableMappingWidget.cpp
+++ b/Engine/Plugins/Runtime/RigVM/Source/RigVMEditor/Private/Widgets/SRigVMVariableMappingWidget.cpp
@@ -201,15 +201,7 @@ void SRigVMVariableMappingTreeRow::OnVarOptionComboOpening()
 	if (OnGetVariableMapping.IsBound())
 	{
 		FName MappedVar = OnGetVariableMapping.Execute(Item->GetPathName());
-		if (MappedVar!= NAME_None)
-		{
-			TSharedPtr<FString> ComboStringPtr = GetVarOptionString(MappedVar);
-			if (ComboStringPtr.IsValid())
-			{
-				VarOptionComboBox->SetSelectedItem(ComboStringPtr);
-			}
-		}
-		else
+		if (MappedVar.IsNone())
 		{
 			VarOptionComboBox->ClearSelection();
 		}
@@ -243,23 +235,6 @@ FText SRigVMVariableMappingTreeRow::GetVarOptionComboBoxToolTip() const
 	return LOCTEXT("VarOptionComboToolTip", "Map input/output variable to available options.");
 }
 
-TSharedPtr<FString> SRigVMVariableMappingTreeRow::GetVarOptionString(FName VarOptionName) const
-{
-	return TSharedPtr<FString>();
-	FString VarOptionString = VarOptionName.ToString();
-
-	// go through profile and see if it has mine
-	for (int32 Index = 1; Index < VariableOptionList.Num(); ++Index)
-	{
-		if (VarOptionString == *VariableOptionList[Index])
-		{
-			return VariableOptionList[Index];
-		}
-	}
-
-	return TSharedPtr<FString>();
-}
-
 TSharedRef<SWidget> SRigVMVariableMappingTreeRow::MakeVarOptionComboWidget(TSharedPtr<FString> InItem)
 {
 	return SNew(STextBlock).Text(FText::FromString(*InItem));// .Font(FAppStyle::GetFontStyle(TEXT("PropertyWindow.NormalFont")));
diff --git a/Engine/Plugins/Runtime/RigVM/Source/RigVMEditor/Public/Widgets/SRigVMVariableMappingWidget.h b/Engine/Plugins/Runtime/RigVM/Source/RigVMEditor/Public/Widgets/SRigVMVariableMappingWidget.h
--- a/Engine/Plugins/Runtime/RigVM/Source/RigVMEditor/Public/Widgets/SRigVMVariableMappingWidget.h
+++ b/Engine/Plugins/Runtime/RigVM/Source/RigVMEditor/Public/Widgets/SRigVMVariableMappingWidget.h
@@ -119,7 +119,6 @@ private:
 	UE_API FText GetVarOptionComboBoxToolTip() const;
 	UE_API void OnVarOptionComboOpening();
 	UE_API bool IsVarOptionEnabled() const;
-	UE_API TSharedPtr<FString> GetVarOptionString(FName VarOptionName) const;
 
 	UE_API ECheckBoxState IsPinChecked() const;
 	UE_API void OnPinCheckStatusChanged(ECheckBoxState NewState);
