Improve usability of the RigVM Print node - defaults to display for a single frame, and the duration is more clearly exposed

#Synced-CL 48390192

diff --git a/Engine/Plugins/Runtime/RigVM/Source/RigVM/Private/RigVMFunctions/RigVMDispatch_Print.cpp b/Engine/Plugins/Runtime/RigVM/Source/RigVM/Private/RigVMFunctions/RigVMDispatch_Print.cpp
--- a/Engine/Plugins/Runtime/RigVM/Source/RigVM/Private/RigVMFunctions/RigVMDispatch_Print.cpp
+++ b/Engine/Plugins/Runtime/RigVM/Source/RigVM/Private/RigVMFunctions/RigVMDispatch_Print.cpp
@@ -69,14 +69,14 @@ FString FRigVMDispatch_Print::GetArgumentDefaultValue(const FName& InArgumentNam
 	}
 	if(InArgumentName == ScreenDurationName)
 	{
-		return TEXT("0.050000");
+		return TEXT("0.0");
 	}
 	return FRigVMDispatchFactory::GetArgumentDefaultValue(InArgumentName, InTypeIndex);
 }
 
 FString FRigVMDispatch_Print::GetArgumentMetaData(const FName& InArgumentName, const FName& InMetaDataKey) const
 {
-	if(InArgumentName == ScreenDurationName || InArgumentName == ScreenColorName)
+	if(InArgumentName == ScreenColorName)
 	{
 		if(InMetaDataKey == FRigVMStruct::DetailsOnlyMetaName)
 		{
@@ -102,7 +102,7 @@ FText FRigVMDispatch_Print::GetArgumentTooltip(const FName& InArgumentName, TRig
 	}
 	if (InArgumentName == ScreenDurationName)
 	{
-		return LOCTEXT("ScreenDurationToolTip", "The duration of the printed message to appear on the screen");
+		return LOCTEXT("ScreenDurationToolTip", "The duration of the printed message to appear on the screen. Zero will show it for a single frame. Negative will not display the message.");
 	}
 	if (InArgumentName == ScreenColorName)
 	{
@@ -147,7 +147,7 @@ void FRigVMDispatch_Print::Execute(FRigVMExtendedExecuteContext& InContext, FRig
 	InContext.GetPublicData<>().Logf({EMessageSeverity::Info, false}, LogFormat, *ObjectPath, InContext.GetPublicData<>().GetInstructionIndex(), *Prefix, *String);
 	const UObject* WorldObject = (const UObject*)InContext.VM;
 
-	if(ScreenDuration > SMALL_NUMBER && WorldObject)
+	if(ScreenDuration >= 0.0f && WorldObject)
 	{
 		static constexpr TCHAR PrintStringFormat[] = TEXT("[%04d] %s%s");
 		UKismetSystemLibrary::PrintString(WorldObject, FString::Printf(PrintStringFormat, InContext.GetPublicData<>().GetInstructionIndex(), *Prefix, *String), true, false, ScreenColor, ScreenDuration);
