RigVM: Fix crash when creating new function using MyBlueprint widget

#rnx

#Synced-CL 48451780

diff --git a/Engine/Plugins/Runtime/RigVM/Source/RigVMEditor/Private/Editor/RigVMEditor.cpp b/Engine/Plugins/Runtime/RigVM/Source/RigVMEditor/Private/Editor/RigVMEditor.cpp
--- a/Engine/Plugins/Runtime/RigVM/Source/RigVMEditor/Private/Editor/RigVMEditor.cpp
+++ b/Engine/Plugins/Runtime/RigVM/Source/RigVMEditor/Private/Editor/RigVMEditor.cpp
@@ -1696,7 +1696,10 @@ void FRigVMEditorBase::OnNewDocumentClicked(ECreatedDocumentType GraphType)
 					if (const UEdGraph* NewGraph = RigVMBlueprint->GetEdGraph(FunctionNode->GetContainedGraph()))
 					{
 						OpenDocument(NewGraph, FDocumentTracker::OpenNewDocument);
-						GraphExplorerWidget->Refresh();
+						if (GraphExplorerWidget.IsValid())
+						{
+							GraphExplorerWidget->Refresh();
+						}
 						RenameNewlyAddedAction(FunctionNode->GetFName());
 					}
 				}
@@ -1715,7 +1718,10 @@ void FRigVMEditorBase::OnNewDocumentClicked(ECreatedDocumentType GraphType)
 				if (const UEdGraph* NewGraph = RigVMBlueprint->GetEdGraph(Model))
 				{
 					OpenDocument(NewGraph, FDocumentTracker::OpenNewDocument);
-					GraphExplorerWidget->Refresh();
+					if (GraphExplorerWidget.IsValid())
+					{
+						GraphExplorerWidget->Refresh();
+					}
 					RenameNewlyAddedAction(NewGraph->GetFName());
 				}
 			}
