Add package path to some rigvm compile error reports

#Synced-CL 48429033

diff --git a/Engine/Plugins/Runtime/RigVM/Source/RigVMDeveloper/Private/RigVMCompiler/RigVMCompiler.cpp b/Engine/Plugins/Runtime/RigVM/Source/RigVMDeveloper/Private/RigVMCompiler/RigVMCompiler.cpp
--- a/Engine/Plugins/Runtime/RigVM/Source/RigVMDeveloper/Private/RigVMCompiler/RigVMCompiler.cpp
+++ b/Engine/Plugins/Runtime/RigVM/Source/RigVMDeveloper/Private/RigVMCompiler/RigVMCompiler.cpp
@@ -936,14 +936,14 @@ bool URigVMCompiler::Compile(const FRigVMCompileSettings& InSettings, TArray<URi
 							{
 								if (Pin->GetTypeIndex() != ExpectedArgument.GetTypeIndex(0))
 								{
-									FString MissingPinMessage = FString::Printf(TEXT("Could not find pin %s of type %s in Node @@."), *ExpectedArgument.Name.ToString(), *Registry.GetType(ExpectedArgument.GetTypeIndex(0)).CPPType.ToString());
+									FString MissingPinMessage = FString::Printf(TEXT("[%s] Could not find pin %s of type %s in Node @@."), *ModelNode->GetPackage()->GetPathName(), *ExpectedArgument.Name.ToString(), *Registry.GetType(ExpectedArgument.GetTypeIndex(0)).CPPType.ToString());
 									Settings.ASTSettings.Report(EMessageSeverity::Error, ModelNode, MissingPinMessage);
 									bEncounteredGraphError = true;
 								}
 							}
 							else
 							{
-								FString MissingPinMessage = FString::Printf(TEXT("Could not find pin %s of type %s in Node @@."), *ExpectedArgument.Name.ToString(), *Registry.GetType(ExpectedArgument.GetTypeIndex(0)).CPPType.ToString());
+								FString MissingPinMessage = FString::Printf(TEXT("[%s] Could not find pin %s of type %s in Node @@."), *ModelNode->GetPackage()->GetPathName(), *ExpectedArgument.Name.ToString(), *Registry.GetType(ExpectedArgument.GetTypeIndex(0)).CPPType.ToString());
 								Settings.ASTSettings.Report(EMessageSeverity::Error, ModelNode, MissingPinMessage);
 								bEncounteredGraphError = true;
 							}
@@ -953,7 +953,7 @@ bool URigVMCompiler::Compile(const FRigVMCompileSettings& InSettings, TArray<URi
 
 				if((Settings.bWarnAboutDeprecatedNodes || CVarRigVMWarnAboutDeprecatedNodes->GetBool()) && UnitNode->IsOutDated())
 				{
-					static const FString LinkedMessage = TEXT("Node @@ is outdated.");
+					static const FString LinkedMessage = FString::Printf(TEXT("[%s] Node @@ is outdated."), *ModelNode->GetPackage()->GetPathName());
 					Settings.ASTSettings.Report(EMessageSeverity::Warning, ModelNode, LinkedMessage);
 				}
 			}
@@ -1098,7 +1098,7 @@ bool URigVMCompiler::Compile(const FRigVMCompileSettings& InSettings, TArray<URi
 			URigVMFunctionReturnNode* ReturnNode = CurrentCompilationFunction->GetReturnNode();
 			if (!ReturnNode)
 			{
-				const FString FunctionCompilationErrorMessage = FString::Printf(TEXT("Mutable function graph %s does not contain a return node."), *CurrentCompilationFunction->GetName());
+				const FString FunctionCompilationErrorMessage = FString::Printf(TEXT("[%s] Mutable function graph %s does not contain a return node."), *EntryNode->GetPackage()->GetPathName(), *CurrentCompilationFunction->GetName());
 				Settings.ASTSettings.Report(EMessageSeverity::Error, nullptr, FunctionCompilationErrorMessage);
 				bEncounteredGraphError = true;
 			}
@@ -1154,7 +1154,7 @@ bool URigVMCompiler::Compile(const FRigVMCompileSettings& InSettings, TArray<URi
 
 			if (!bReturnNodeFound)
 			{
-				const FString FunctionCompilationErrorMessage = FString::Printf(TEXT("Not all paths of the mutable function graph %s lead to a return node."), *CurrentCompilationFunction->GetName());
+				const FString FunctionCompilationErrorMessage = FString::Printf(TEXT("[%s] Not all paths of the mutable function graph %s lead to a return node."), *CurrentCompilationFunction->GetPackage()->GetPathName(), *CurrentCompilationFunction->GetName());
 				Settings.ASTSettings.Report(EMessageSeverity::Error, nullptr, FunctionCompilationErrorMessage);
 				bEncounteredGraphError = true;
 			}
@@ -1234,7 +1234,7 @@ bool URigVMCompiler::Compile(const FRigVMCompileSettings& InSettings, TArray<URi
 			{
 				if(EntryNode == nullptr)
 				{
-					WorkData.ReportErrorf(TEXT("Corrupt library node '%s' - Missing entry node."), *CurrentCompilationFunction->GetPathName());
+					WorkData.ReportErrorf(TEXT("[%s] Corrupt library node '%s' - Missing entry node."), *CurrentCompilationFunction->GetPackage()->GetPathName(), *CurrentCompilationFunction->GetPathName());
 					return false;
 				}
 				InterfacePin = EntryNode->FindPin(Pin->GetName());
@@ -1243,7 +1243,7 @@ bool URigVMCompiler::Compile(const FRigVMCompileSettings& InSettings, TArray<URi
 			{
 				if(ReturnNode == nullptr)
 				{
-					WorkData.ReportErrorf(TEXT("Corrupt library node '%s' - Missing return node."), *CurrentCompilationFunction->GetPathName());
+					WorkData.ReportErrorf(TEXT("[%s] Corrupt library node '%s' - Missing return node."), *CurrentCompilationFunction->GetPackage()->GetPathName(), *CurrentCompilationFunction->GetPathName());
 					return false;
 				}
 				InterfacePin = ReturnNode->FindPin(Pin->GetName());
@@ -1251,7 +1251,7 @@ bool URigVMCompiler::Compile(const FRigVMCompileSettings& InSettings, TArray<URi
 
 			if(InterfacePin == nullptr)
 			{
-				WorkData.ReportErrorf(TEXT("Corrupt library node '%s' - Pin '%s' is not part of the entry / return node."), *CurrentCompilationFunction->GetPathName(), *Pin->GetPathName());
+				WorkData.ReportErrorf(TEXT("[%s} Corrupt library node '%s' - Pin '%s' is not part of the entry / return node."), *CurrentCompilationFunction->GetPackage()->GetPathName(), *CurrentCompilationFunction->GetPathName(), *Pin->GetPathName());
 				return false;
 			}
 			
@@ -2107,9 +2107,9 @@ bool URigVMCompiler::TraverseCallExtern(const FRigVMCallExternExprAST* InExpr, F
 			const UScriptStruct* SpecializedExecuteStruct = WorkData.ExecuteContextStruct;
 			if(!SpecializedExecuteStruct->IsChildOf(ExecuteStruct))
 			{
-				static constexpr TCHAR UnknownExecuteContextMessage[] = TEXT("Node @@ uses an unexpected execute type '%s'. This graph uses '%s'.");
+				static constexpr TCHAR UnknownExecuteContextMessage[] = TEXT("[%s] Node @@ uses an unexpected execute type '%s'. This graph uses '%s'.");
 				WorkData.Settings.Report(EMessageSeverity::Error, Subject, FString::Printf(
-					UnknownExecuteContextMessage, *ExecuteStruct->GetStructCPPName(), *SpecializedExecuteStruct->GetStructCPPName()));
+					UnknownExecuteContextMessage, *Subject->GetPackage()->GetPathName(), *ExecuteStruct->GetStructCPPName(), *SpecializedExecuteStruct->GetStructCPPName()));
 				return false;
 			}
 		}
@@ -2121,7 +2121,7 @@ bool URigVMCompiler::TraverseCallExtern(const FRigVMCallExternExprAST* InExpr, F
 		const UScriptStruct* ScriptStruct = UnitNode->GetScriptStruct(); 
 		if(ScriptStruct == nullptr)
 		{
-			static const FString UnresolvedMessage = TEXT("Node @@ is unresolved.");
+			static const FString UnresolvedMessage = FString::Printf(TEXT("[%s] Node @@ is unresolved."), *UnitNode->GetPackage()->GetPathName());
 			WorkData.Settings.Report(EMessageSeverity::Error, UnitNode, UnresolvedMessage);
 			return false;
 		}
@@ -2143,7 +2143,7 @@ bool URigVMCompiler::TraverseCallExtern(const FRigVMCallExternExprAST* InExpr, F
 	{
 		if(DispatchNode->GetFactory() == nullptr)
 		{
-			static const FString UnresolvedDispatchMessage = TEXT("Dispatch node @@ has no factory.");
+			static const FString UnresolvedDispatchMessage = FString::Printf(TEXT("[%s] Dispatch node @@ has no factory."), *DispatchNode->GetPackage()->GetPathName());
 			WorkData.Settings.Report(EMessageSeverity::Error, DispatchNode, UnresolvedDispatchMessage);
 			return false;
 		}
@@ -4675,8 +4675,7 @@ bool URigVMCompiler::ValidateNode(const FRigVMCompileSettings& InSettings, URigV
 	{
 		if(InNode->HasWildCardPin())
 		{
-			static const FString UnknownTypeMessage = TEXT("Node @@ has unresolved pins of wildcard type.");
-			InSettings.Report(EMessageSeverity::Error, InNode, UnknownTypeMessage);
+			InSettings.Reportf(EMessageSeverity::Error, InNode, TEXT("[%s] Node @@ has unresolved pins of wildcard type."), *InNode->GetPackage()->GetPathName());
 			return false;
 		}
 		return true;
