[RigVM] Added missing initialization for output variable #jira UE-354554 #rnx
#Synced-CL 48500089

diff --git a/Engine/Plugins/Runtime/RigVM/Source/RigVM/Public/RigVMFunctions/Math/RigVMFunction_MathBool.h b/Engine/Plugins/Runtime/RigVM/Source/RigVM/Public/RigVMFunctions/Math/RigVMFunction_MathBool.h
--- a/Engine/Plugins/Runtime/RigVM/Source/RigVM/Public/RigVMFunctions/Math/RigVMFunction_MathBool.h
+++ b/Engine/Plugins/Runtime/RigVM/Source/RigVM/Public/RigVMFunctions/Math/RigVMFunction_MathBool.h
@@ -450,6 +450,7 @@ struct FRigVMFunction_MathBoolAny : public FRigVMFunction_MathBoolBase
 	{
 		Array.Reset();
 		Result = false;
+		IsEmpty = false;
 	}
 
 	RIGVM_METHOD()
@@ -481,6 +482,7 @@ struct FRigVMFunction_MathBoolAll : public FRigVMFunction_MathBoolBase
 	{
 		Array.Reset();
 		Result = false;
+		IsEmpty = false;
 	}
 
 	RIGVM_METHOD()
