Clarify help text on some RigVM functions. No functional changes

#Synced-CL 48547444

diff --git a/Engine/Plugins/Runtime/RigVM/Source/RigVM/Public/RigVMFunctions/Math/RigVMFunction_MathVector.h b/Engine/Plugins/Runtime/RigVM/Source/RigVM/Public/RigVMFunctions/Math/RigVMFunction_MathVector.h
--- a/Engine/Plugins/Runtime/RigVM/Source/RigVM/Public/RigVMFunctions/Math/RigVMFunction_MathVector.h
+++ b/Engine/Plugins/Runtime/RigVM/Source/RigVM/Public/RigVMFunctions/Math/RigVMFunction_MathVector.h
@@ -550,7 +550,7 @@ struct FRigVMFunction_MathVectorNotEquals : public FRigVMFunction_MathVectorBase
 };
 
 /**
- * Returns true if the value is nearly zero
+ * Returns true if each component of the value is within a tolerance of zero
  */
 USTRUCT(meta=(DisplayName="Is Nearly Zero", TemplateName="IsNearlyZero", Keywords="AlmostZero,0"))
 struct FRigVMFunction_MathVectorIsNearlyZero : public FRigVMFunction_MathVectorBase
@@ -571,17 +571,17 @@ struct FRigVMFunction_MathVectorIsNearlyZero : public FRigVMFunction_MathVectorB
 	UPROPERTY(meta=(Input))
 	FVector Value;
 
-	// The tolerance to use when comparing with zero
+	// The tolerance to use when comparing each component with zero
 	UPROPERTY(meta=(Input))
 	float Tolerance;
 
-	// True if the input vector almost zero given the tolerance
+	// True if the input vector is within Tolerance of zero in each component
 	UPROPERTY(meta=(Output))
 	bool Result;
 };
 
 /**
- * Returns true if value A is almost equal to B
+ * Returns true if value A is within a tolerance (in each component) of B
  */
 USTRUCT(meta=(DisplayName="Is Nearly Equal", TemplateName="IsNearlyEqual", Keywords="AlmostEqual,=="))
 struct FRigVMFunction_MathVectorIsNearlyEqual : public FRigVMFunction_MathVectorBase
@@ -610,7 +610,7 @@ struct FRigVMFunction_MathVectorIsNearlyEqual : public FRigVMFunction_MathVector
 	UPROPERTY(meta=(Input))
 	float Tolerance;
 
-	// True if value A is almost equal to B given the tolerance
+	// True if every component of A is within Tolerance of the corresponding component of B
 	UPROPERTY(meta=(Output))
 	bool Result;
 };
diff --git a/Engine/Plugins/Runtime/RigVM/Source/RigVM/Public/RigVMFunctions/Simulation/RigVMFunction_Accumulate.h b/Engine/Plugins/Runtime/RigVM/Source/RigVM/Public/RigVMFunctions/Simulation/RigVMFunction_Accumulate.h
--- a/Engine/Plugins/Runtime/RigVM/Source/RigVM/Public/RigVMFunctions/Simulation/RigVMFunction_Accumulate.h
+++ b/Engine/Plugins/Runtime/RigVM/Source/RigVM/Public/RigVMFunctions/Simulation/RigVMFunction_Accumulate.h
@@ -34,7 +34,7 @@ struct FRigVMFunction_AccumulateFloatAdd : public FRigVMFunction_AccumulateBase
 	RIGVM_METHOD()
 	RIGVM_API virtual void Execute() override;
 
-	// The increment to add for every iteration
+	// The value added on every iteration (depending on the IntegrateDeltaTime flag)
 	UPROPERTY(meta=(Input))
 	float Increment;
 
@@ -42,8 +42,8 @@ struct FRigVMFunction_AccumulateFloatAdd : public FRigVMFunction_AccumulateBase
 	UPROPERTY(meta = (Input))
 	float InitialValue;
 
-	// If True the integration will be relying on the delta time
-	// to create more deterministic results with varying framerates.
+	// If True, the increment will be multiplied by the delta time as it is accumulated, treating
+	// it as a rate of change. This is advisable when dealing with variable frame rates.
 	UPROPERTY(meta=(Input))
 	bool bIntegrateDeltaTime;
 
@@ -78,7 +78,7 @@ struct FRigVMFunction_AccumulateVectorAdd : public FRigVMFunction_AccumulateBase
 	RIGVM_METHOD()
 	RIGVM_API virtual void Execute() override;
 
-	// The increment to add for every iteration
+	// The value added on every iteration (depending on the IntegrateDeltaTime flag)
 	UPROPERTY(meta = (Input))
 	FVector Increment;
 
@@ -86,8 +86,8 @@ struct FRigVMFunction_AccumulateVectorAdd : public FRigVMFunction_AccumulateBase
 	UPROPERTY(meta = (Input))
 	FVector InitialValue;
 
-	// If True the integration will be relying on the delta time
-	// to create more deterministic results with varying framerates.
+	// If True, the increment will be multiplied by the delta time as it is accumulated, treating
+	// it as a rate of change. This is advisable when dealing with variable frame rates.
 	UPROPERTY(meta = (Input))
 	bool bIntegrateDeltaTime;
 
@@ -122,7 +122,7 @@ struct FRigVMFunction_AccumulateFloatMul : public FRigVMFunction_AccumulateBase
 	RIGVM_METHOD()
 	RIGVM_API virtual void Execute() override;
 
-	// The multiplier to apply for every iteration
+	// The multiplier to apply on every iteration (depending on the IntegrateDeltaTime flag)
 	UPROPERTY(meta=(Input))
 	float Multiplier;
 
@@ -130,8 +130,8 @@ struct FRigVMFunction_AccumulateFloatMul : public FRigVMFunction_AccumulateBase
 	UPROPERTY(meta = (Input))
 	float InitialValue;
 
-	// If True the integration will be relying on the delta time
-	// to create more deterministic results with varying framerates.
+	// If True, the multiplier will be scaled depending on the delta time as it is accumulated,
+	// making the resulting accumulation less sensitive to varying frame rates.
 	UPROPERTY(meta=(Input))
 	bool bIntegrateDeltaTime;
 
@@ -166,7 +166,7 @@ struct FRigVMFunction_AccumulateVectorMul : public FRigVMFunction_AccumulateBase
 	RIGVM_METHOD()
 	RIGVM_API virtual void Execute() override;
 
-	// The multiplier to apply for every iteration
+	// The multiplier to apply on every iteration (depending on the IntegrateDeltaTime flag)
 	UPROPERTY(meta = (Input))
 	FVector Multiplier;
 
@@ -174,8 +174,8 @@ struct FRigVMFunction_AccumulateVectorMul : public FRigVMFunction_AccumulateBase
 	UPROPERTY(meta = (Input))
 	FVector InitialValue;
 
-	// If True the integration will be relying on the delta time
-	// to create more deterministic results with varying framerates.
+	// If True, the multiplier will be scaled depending on the delta time as it is accumulated,
+	// making the resulting accumulation less sensitive to varying frame rates.
 	UPROPERTY(meta = (Input))
 	bool bIntegrateDeltaTime;
 
@@ -210,7 +210,7 @@ struct FRigVMFunction_AccumulateQuatMul : public FRigVMFunction_AccumulateBase
 	RIGVM_METHOD()
 	RIGVM_API virtual void Execute() override;
 
-	// The multiplier to apply for every iteration
+	// The multiplier to apply on every iteration (depending on the IntegrateDeltaTime flag) 
 	UPROPERTY(meta=(Input))
 	FQuat Multiplier;
 
@@ -222,8 +222,8 @@ struct FRigVMFunction_AccumulateQuatMul : public FRigVMFunction_AccumulateBase
 	UPROPERTY(meta = (Input))
 	bool bFlipOrder;
 
-	// If True the integration will be relying on the delta time
-	// to create more deterministic results with varying framerates.
+	// If True, the multiplier will be scaled depending on the delta time as it is accumulated,
+	// making the resulting accumulation less sensitive to varying frame rates.
 	UPROPERTY(meta=(Input))
 	bool bIntegrateDeltaTime;
 
@@ -258,7 +258,7 @@ struct FRigVMFunction_AccumulateTransformMul : public FRigVMFunction_AccumulateB
 	RIGVM_METHOD()
 	RIGVM_API virtual void Execute() override;
 
-	// The multiplier to apply for every iteration
+	// The multiplier to apply for every iteration (depending on the IntegrateDeltaTime flag) 
 	UPROPERTY(meta = (Input))
 	FTransform Multiplier;
 
@@ -270,8 +270,8 @@ struct FRigVMFunction_AccumulateTransformMul : public FRigVMFunction_AccumulateB
 	UPROPERTY(meta = (Input))
 	bool bFlipOrder;
 
-	// If True the integration will be relying on the delta time
-	// to create more deterministic results with varying framerates.
+	// If True, the multiplier will be scaled depending on the delta time as it is accumulated,
+	// making the resulting accumulation less sensitive to varying frame rates.
 	UPROPERTY(meta = (Input))
 	bool bIntegrateDeltaTime;
 
