Control Rig: Fix graph explorer not changing variable public/private
#rnx

#Synced-CL 48500891

diff --git a/Engine/Plugins/Runtime/RigVM/Source/RigVMDeveloper/Private/RigVMBlueprintLegacy.cpp b/Engine/Plugins/Runtime/RigVM/Source/RigVMDeveloper/Private/RigVMBlueprintLegacy.cpp
--- a/Engine/Plugins/Runtime/RigVM/Source/RigVMDeveloper/Private/RigVMBlueprintLegacy.cpp
+++ b/Engine/Plugins/Runtime/RigVM/Source/RigVMDeveloper/Private/RigVMBlueprintLegacy.cpp
@@ -634,7 +634,7 @@ bool URigVMBlueprint::SetVariablePrivate(const FName& InName, const bool bInPriv
 
 bool URigVMBlueprint::SetVariablePublic(const FName& InName, const bool bIsPublic)
 {
-	FBlueprintEditorUtils::SetBlueprintOnlyEditableFlag(this, InName, bIsPublic);
+	FBlueprintEditorUtils::SetBlueprintOnlyEditableFlag(this, InName, !bIsPublic);
 	return true;
 }
 
