RigVM: RunOnce / Timer control flow nodes

#rnx

#Synced-CL 48645710

diff --git a/Engine/Plugins/Runtime/RigVM/Source/RigVM/Private/RigVMFunctions/RigVMFunction_ControlFlow.cpp b/Engine/Plugins/Runtime/RigVM/Source/RigVM/Private/RigVMFunctions/RigVMFunction_ControlFlow.cpp
--- a/Engine/Plugins/Runtime/RigVM/Source/RigVM/Private/RigVMFunctions/RigVMFunction_ControlFlow.cpp
+++ b/Engine/Plugins/Runtime/RigVM/Source/RigVM/Private/RigVMFunctions/RigVMFunction_ControlFlow.cpp
@@ -1,6 +1,7 @@
 // Copyright Epic Games, Inc. All Rights Reserved.
 
 #include "RigVMFunctions/RigVMFunction_ControlFlow.h"
+#include "RigVMFunctions/Simulation/RigVMFunction_Timeline.h"
 
 #include UE_INLINE_GENERATED_CPP_BY_NAME(RigVMFunction_ControlFlow)
 
@@ -24,3 +25,59 @@ FRigVMFunction_ControlFlowBranch_Execute()
 		BlockToRun = FRigVMStruct::ControlFlowCompletedName;
 	}
 }
+
+FRigVMFunction_ControlFlowRunOnce_Execute()
+{
+	if(BlockToRun.IsNone())
+	{
+		if (!HasRun)
+		{
+			static const FLazyName OnceName(GET_MEMBER_NAME_CHECKED(FRigVMFunction_ControlFlowRunOnce, Once)); 
+			BlockToRun = OnceName.Resolve();
+			HasRun = true;
+			return;
+		}
+	}
+	BlockToRun = FRigVMStruct::ControlFlowCompletedName;
+}
+
+FRigVMFunction_ControlFlowTimer_Execute()
+{
+	static const FLazyName JustStartedName(GET_MEMBER_NAME_CHECKED(FRigVMFunction_ControlFlowTimer, JustStarted)); 
+	static const FLazyName IsRunningName(GET_MEMBER_NAME_CHECKED(FRigVMFunction_ControlFlowTimer, IsRunning)); 
+	static const FLazyName JustFinishedName(GET_MEMBER_NAME_CHECKED(FRigVMFunction_ControlFlowTimer, JustFinished)); 
+	
+	if(BlockToRun.IsNone())
+	{
+		FRigVMFunction_TimerValue::StaticExecute(ExecuteContext, Reset, Speed, Duration, Repeats, bJustStarted, bIsRunning, bJustFinished, Time, Ratio, AccumulatedValue, bIsInitialized);
+
+		if (bJustStarted)
+		{
+			BlockToRun = JustStartedName;
+			return;
+		}
+		if (bIsRunning)
+		{
+			BlockToRun = IsRunningName;
+			return;
+		}
+		if (bJustFinished)
+		{
+			BlockToRun = JustFinishedName;
+			return;
+		}
+	}
+
+	if (BlockToRun == JustStartedName && bIsRunning)
+	{
+		BlockToRun = IsRunningName;
+		return;
+	}
+	if (BlockToRun == IsRunningName && bJustFinished)
+	{
+		BlockToRun = JustFinishedName;
+		return;
+	}
+
+	BlockToRun = FRigVMStruct::ControlFlowCompletedName;
+}
diff --git a/Engine/Plugins/Runtime/RigVM/Source/RigVM/Private/RigVMFunctions/Simulation/RigVMFunction_Timeline.cpp b/Engine/Plugins/Runtime/RigVM/Source/RigVM/Private/RigVMFunctions/Simulation/RigVMFunction_Timeline.cpp
--- a/Engine/Plugins/Runtime/RigVM/Source/RigVM/Private/RigVMFunctions/Simulation/RigVMFunction_Timeline.cpp
+++ b/Engine/Plugins/Runtime/RigVM/Source/RigVM/Private/RigVMFunctions/Simulation/RigVMFunction_Timeline.cpp
@@ -76,3 +76,55 @@ FRigVMFunction_TimeLoop_Execute()
 		FlipFlop = FlipFlop / DurationClamped;
 	}
 }
+
+FRigVMFunction_TimerValue_Execute()
+{
+	if (!bIsInitialized || Reset)
+	{
+		AccumulatedValue = 0.f;
+		bIsInitialized = true;
+		IsRunning = true;
+	}
+
+	const float AbsoluteDuration = FMath::Abs(Duration);
+	if(!IsRunning || AbsoluteDuration < SMALL_NUMBER)
+	{
+		JustStarted = JustFinished = false;
+	}
+	else if(IsRunning)
+	{
+		JustStarted = AccumulatedValue < SMALL_NUMBER;
+		IsRunning = AccumulatedValue < AbsoluteDuration + SMALL_NUMBER;
+
+		AccumulatedValue = AccumulatedValue + ExecuteContext.GetDeltaTime<float>() * Speed;
+
+		JustFinished = AccumulatedValue > AbsoluteDuration - SMALL_NUMBER;
+		if(JustFinished)
+		{
+			IsRunning = false;
+		}
+	}
+
+	Time = FMath::Clamp(AccumulatedValue, 0.f, AbsoluteDuration);
+	if (AbsoluteDuration < SMALL_NUMBER)
+	{
+		Ratio = 0.f;
+	}
+	else
+	{
+		Ratio = FMath::Clamp(Time / AbsoluteDuration, 0.f, 1.f);
+	}
+
+	if(JustFinished)
+	{
+		if(Repeats)
+		{
+			AccumulatedValue = 0.f;
+			IsRunning = true;
+		}
+		else
+		{
+			IsRunning = false;
+		}
+	}
+}
\ No newline at end of file
diff --git a/Engine/Plugins/Runtime/RigVM/Source/RigVM/Public/RigVMFunctions/RigVMFunction_ControlFlow.h b/Engine/Plugins/Runtime/RigVM/Source/RigVM/Public/RigVMFunctions/RigVMFunction_ControlFlow.h
--- a/Engine/Plugins/Runtime/RigVM/Source/RigVM/Public/RigVMFunctions/RigVMFunction_ControlFlow.h
+++ b/Engine/Plugins/Runtime/RigVM/Source/RigVM/Public/RigVMFunctions/RigVMFunction_ControlFlow.h
@@ -65,3 +65,147 @@ struct FRigVMFunction_ControlFlowBranch : public FRigVMFunction_ControlFlowBase
 	UPROPERTY(meta=(Singleton))
 	FName BlockToRun;
 };
+
+/**
+ * Executes the Once branch only once
+ */
+USTRUCT(meta = (DisplayName = "Run Once", Keywords = "Branch,if"))
+struct FRigVMFunction_ControlFlowRunOnce : public FRigVMFunction_ControlFlowBase
+{
+	GENERATED_BODY()
+
+	FRigVMFunction_ControlFlowRunOnce()
+	{
+		HasRun = false;
+		BlockToRun = NAME_None;
+	}
+
+	RIGVM_METHOD()
+	RIGVM_API virtual void Execute();
+
+	virtual const TArray<FName>& GetControlFlowBlocks_Impl() const override
+	{
+		static const TArray<FName> Blocks = {
+			GET_MEMBER_NAME_CHECKED(FRigVMFunction_ControlFlowRunOnce, Once),
+			ForLoopCompletedPinName
+		};
+		return Blocks;
+	}
+
+	// The input execution pin to hook up to the graph
+	UPROPERTY(meta=(Input, DisplayName="Execute"))
+	FRigVMExecuteContext ExecuteContext;
+
+	// The branch to run once
+	UPROPERTY(meta=(Output))
+	FRigVMExecuteContext Once;
+
+	// The execute flow to run every frame
+	UPROPERTY(meta=(Output))
+	FRigVMExecuteContext Completed;
+
+	// The internal block to run as the node progresses
+	UPROPERTY(meta=(Singleton))
+	FName BlockToRun;
+
+	// The internal condition indicating if this has run once or now
+	UPROPERTY()
+	bool HasRun;
+};
+
+/**
+ * Executes the output branches based on a timer
+ */
+USTRUCT(meta = (DisplayName = "Timer", Keywords = "Branch,If,Timeline,Execute"))
+struct FRigVMFunction_ControlFlowTimer : public FRigVMFunction_ControlFlowBase
+{
+	GENERATED_BODY()
+
+	FRigVMFunction_ControlFlowTimer()
+	{
+		BlockToRun = NAME_None;
+		Reset = Repeats = bJustStarted = bIsRunning = bJustFinished = bIsInitialized = bIsRunning = false;
+		Speed = Duration = 1.f;
+		Time = Ratio = AccumulatedValue = 0.f;
+	}
+
+	RIGVM_METHOD()
+	RIGVM_API virtual void Execute();
+
+	virtual const TArray<FName>& GetControlFlowBlocks_Impl() const override
+	{
+		static const TArray<FName> Blocks = {
+			GET_MEMBER_NAME_CHECKED(FRigVMFunction_ControlFlowTimer, JustStarted),
+			GET_MEMBER_NAME_CHECKED(FRigVMFunction_ControlFlowTimer, IsRunning),
+			GET_MEMBER_NAME_CHECKED(FRigVMFunction_ControlFlowTimer, JustFinished),
+			ForLoopCompletedPinName
+		};
+		return Blocks;
+	}
+
+	// The input execution pin to hook up to the graph
+	UPROPERTY(meta=(Input, DisplayName="Execute"))
+	FRigVMExecuteContext ExecuteContext;
+
+	// A flag to reset / restart the timer
+	UPROPERTY(meta = (Input))
+	bool Reset;
+
+	// The speed to apply to the played back time
+	UPROPERTY(meta = (Input))
+	float Speed;
+
+	// The duration in seconds the timer will take to complete
+	UPROPERTY(meta = (Input))
+	float Duration;
+
+	// A flag indicating if the timer is a one-shot or if it should repeat / loop
+	UPROPERTY(meta = (Input))
+	bool Repeats;
+
+	// The branch to run when the timer has just started
+	UPROPERTY(meta=(Output))
+	FRigVMExecuteContext JustStarted;
+
+	// The branch to run when the timer is running
+	UPROPERTY(meta=(Output))
+	FRigVMExecuteContext IsRunning;
+
+	// The branch to run when the timer has just finished
+	UPROPERTY(meta=(Output))
+	FRigVMExecuteContext JustFinished;
+
+	// The played back / accumulated time in seconds
+	UPROPERTY(meta=(Output))
+	float Time;
+
+	// The played back / accumulated time as a ratio of time / duration (0.0 - 1.0)
+	UPROPERTY(meta=(Output))
+	float Ratio;
+
+	// The execute flow to run every frame after all other branches
+	UPROPERTY(meta=(Output))
+	FRigVMExecuteContext Completed;
+
+	// The internal block to run as the node progresses
+	UPROPERTY(meta=(Singleton))
+	FName BlockToRun;
+
+	UPROPERTY()
+	float AccumulatedValue;
+
+	UPROPERTY()
+	bool bIsInitialized;
+
+	// True if the timer has just started
+	UPROPERTY(meta=(Singleton))
+	bool bJustStarted;
+
+	// True if the timer is running
+	UPROPERTY(meta=(Singleton))
+	bool bIsRunning;
+
+	// True if the timer has just finished
+	UPROPERTY(meta=(Singleton))
+	bool bJustFinished;
+};
\ No newline at end of file
diff --git a/Engine/Plugins/Runtime/RigVM/Source/RigVM/Public/RigVMFunctions/Simulation/RigVMFunction_Timeline.h b/Engine/Plugins/Runtime/RigVM/Source/RigVM/Public/RigVMFunctions/Simulation/RigVMFunction_Timeline.h
--- a/Engine/Plugins/Runtime/RigVM/Source/RigVM/Public/RigVMFunctions/Simulation/RigVMFunction_Timeline.h
+++ b/Engine/Plugins/Runtime/RigVM/Source/RigVM/Public/RigVMFunctions/Simulation/RigVMFunction_Timeline.h
@@ -5,11 +5,22 @@
 #include "RigVMFunctions/Simulation/RigVMFunction_SimBase.h"
 #include "RigVMFunction_Timeline.generated.h"
 
+/*
+ * The base class for all time nodes
+ */
+USTRUCT(meta=(Abstract, Category="Simulation|Time"))
+struct FRigVMFunction_TimeBase : public FRigVMFunction_SimBase
+{
+	GENERATED_BODY()
+
+	virtual void Execute() {}
+};
+
 /**
  * Simulates a time value - can act as a timeline playing back
  */
-USTRUCT(meta=(DisplayName="Accumulated Time", Category = "Simulation|Accumulate", Keywords="Playback,Pause,Timeline"))
-struct FRigVMFunction_Timeline : public FRigVMFunction_SimBase
+USTRUCT(meta=(DisplayName="Accumulated Time", Keywords="Playback,Pause,Timeline"))
+struct FRigVMFunction_Timeline : public FRigVMFunction_TimeBase
 {
 	GENERATED_BODY()
 	
@@ -43,8 +54,8 @@ struct FRigVMFunction_Timeline : public FRigVMFunction_SimBase
 /**
  * Simulates a time value - and outputs loop information
  */
-USTRUCT(meta=(DisplayName="Time Loop", Category = "Simulation|Accumulate", Keywords="Playback,Pause,Timeline"))
-struct FRigVMFunction_TimeLoop : public FRigVMFunction_SimBase
+USTRUCT(meta=(DisplayName="Time Loop", Keywords="Playback,Pause,Timeline"))
+struct FRigVMFunction_TimeLoop : public FRigVMFunction_TimeBase
 {
 	GENERATED_BODY()
 	
@@ -106,4 +117,67 @@ struct FRigVMFunction_TimeLoop : public FRigVMFunction_SimBase
 
 	UPROPERTY()
 	bool bIsInitialized;
-};
\ No newline at end of file
+};
+
+/**
+ * Starts a timer and provides access to the time simulation
+ */
+USTRUCT(meta=(DisplayName="Timer Value", Keywords="Playback,Pause,Timeline"))
+struct FRigVMFunction_TimerValue : public FRigVMFunction_TimeBase
+{
+	GENERATED_BODY()
+	
+	FRigVMFunction_TimerValue()
+	{
+		Reset = Repeats = JustStarted = IsRunning = JustFinished = bIsInitialized = false;
+		Speed = Duration = 1.f;
+		Time = Ratio = AccumulatedValue = 0.f;
+	}
+
+	virtual void Initialize() override { bIsInitialized = false; }
+
+	RIGVM_METHOD()
+	RIGVM_API virtual void Execute() override;
+
+	// A flag to reset / restart the timer
+	UPROPERTY(meta = (Input))
+	bool Reset;
+
+	// The speed to apply to the played back time
+	UPROPERTY(meta = (Input))
+	float Speed;
+
+	// The duration in seconds the timer will take to complete
+	UPROPERTY(meta = (Input))
+	float Duration;
+
+	// A flag indicating if the timer is a one-shot or if it should repeat / loop
+	UPROPERTY(meta = (Input))
+	bool Repeats;
+
+	// True if the timer has just started
+	UPROPERTY(meta=(Output))
+	bool JustStarted;
+
+	// True if the timer is running
+	UPROPERTY(meta=(Output))
+	bool IsRunning;
+
+	// True if the timer has just finished
+	UPROPERTY(meta=(Output))
+	bool JustFinished;
+
+	// The played back / accumulated time in seconds
+	UPROPERTY(meta=(Output))
+	float Time;
+
+	// The played back / accumulated time as a ratio of time / duration (0.0 - 1.0)
+	UPROPERTY(meta=(Output))
+	float Ratio;
+
+	UPROPERTY()
+	float AccumulatedValue;
+
+	UPROPERTY()
+	bool bIsInitialized;
+};
