RigVM: Fix static analysis warning in Control Flow node header

#rnx

#Synced-CL 48645710

diff --git a/Engine/Plugins/Runtime/RigVM/Source/RigVM/Public/RigVMFunctions/RigVMFunction_ControlFlow.h b/Engine/Plugins/Runtime/RigVM/Source/RigVM/Public/RigVMFunctions/RigVMFunction_ControlFlow.h
--- a/Engine/Plugins/Runtime/RigVM/Source/RigVM/Public/RigVMFunctions/RigVMFunction_ControlFlow.h
+++ b/Engine/Plugins/Runtime/RigVM/Source/RigVM/Public/RigVMFunctions/RigVMFunction_ControlFlow.h
@@ -124,7 +124,7 @@ struct FRigVMFunction_ControlFlowTimer : public FRigVMFunction_ControlFlowBase
 	FRigVMFunction_ControlFlowTimer()
 	{
 		BlockToRun = NAME_None;
-		Reset = Repeats = bJustStarted = bIsRunning = bJustFinished = bIsInitialized = bIsRunning = false;
+		Reset = Repeats = bJustStarted = bIsRunning = bJustFinished = bIsInitialized = false;
 		Speed = Duration = 1.f;
 		Time = Ratio = AccumulatedValue = 0.f;
 	}
