RigVM: Functions to support input variables

diff --git a/Engine/Plugins/Runtime/RigVM/Source/RigVM/Public/RigVMCore/RigVM.h b/Engine/Plugins/Runtime/RigVM/Source/RigVM/Public/RigVMCore/RigVM.h
--- a/Engine/Plugins/Runtime/RigVM/Source/RigVM/Public/RigVMCore/RigVM.h
+++ b/Engine/Plugins/Runtime/RigVM/Source/RigVM/Public/RigVMCore/RigVM.h
@@ -817,21 +817,21 @@ public:
 			FRigVMExternalVariableDef& ExternalVariableDef = ExternalVariables[i];
 
 			// Only check name and property, to allow the case where an UUserStruct is deleted while used inside a Rig
-			check(ExternalVariableDef.Name == InExternalVariable.Name);
-			check(ExternalVariableDef.Property == InExternalVariable.Property);
-			check(bAllowNullMemory || InExternalVariable.Memory != nullptr);
+			check(ExternalVariableDef.GetName() == InExternalVariable.GetName());
+			check(ExternalVariableDef.GetProperty() == InExternalVariable.GetProperty());
+			check(bAllowNullMemory || InExternalVariable.GetMemory() != nullptr);
 			
-			Context.ExternalVariableRuntimeData.Add(FRigVMExternalVariableRuntimeData(InExternalVariable.Memory));
+			Context.ExternalVariableRuntimeData.Add(FRigVMExternalVariableRuntimeData(const_cast<uint8*>(InExternalVariable.GetMemory())));
 		}
 	}
 
 	// Adds a new external / unowned variable to the VM
 	FRigVMOperand AddExternalVariable(FRigVMExtendedExecuteContext& Context, const FRigVMExternalVariable& InExternalVariable, bool bAllowNullMemory = false)
 	{
-		check(bAllowNullMemory || InExternalVariable.Memory != nullptr);
+		check(bAllowNullMemory || InExternalVariable.GetMemory() != nullptr);
 
 		const int32 VariableIndex = ExternalVariables.Add(InExternalVariable);
-		Context.ExternalVariableRuntimeData.Add(FRigVMExternalVariableRuntimeData(InExternalVariable.Memory));
+		Context.ExternalVariableRuntimeData.Add(FRigVMExternalVariableRuntimeData(const_cast<uint8*>(InExternalVariable.GetMemory())));
 
 		return FRigVMOperand(ERigVMMemoryType::External, VariableIndex);
 	}
diff --git a/Engine/Plugins/Runtime/RigVM/Source/RigVM/Public/RigVMCore/RigVMExternalVariable.h b/Engine/Plugins/Runtime/RigVM/Source/RigVM/Public/RigVMCore/RigVMExternalVariable.h
--- a/Engine/Plugins/Runtime/RigVM/Source/RigVM/Public/RigVMCore/RigVMExternalVariable.h
+++ b/Engine/Plugins/Runtime/RigVM/Source/RigVM/Public/RigVMCore/RigVMExternalVariable.h
@@ -9,7 +9,6 @@
 #include "RigVMTypeIndex.h"
 #include "RigVMTypeUtils.h"
 #include "UObject/UnrealType.h"
-#include "UObject/TextProperty.h"
 #if WITH_EDITOR
 #include "EdGraphSchema_K2.h"
 #endif
@@ -24,51 +23,132 @@ struct FRigVMExternalVariableDef
 {
 	GENERATED_BODY()
 
+public:
+	
 	FRigVMExternalVariableDef() = default;
 
 	bool operator == (const FRigVMExternalVariableDef& Other) const
 	{
 		return Name == Other.Name
 			&& Property == Other.Property
-			&& TypeName == Other.TypeName
-			&& TypeObject == Other.TypeObject
+			&& BaseCPPType == Other.BaseCPPType
+			&& CPPTypeObject == Other.CPPTypeObject
 			&& bIsArray == Other.bIsArray
 			&& bIsPublic == Other.bIsPublic
-			&& bIsReadOnly == Other.bIsReadOnly
-			&& Size == Other.Size;
+			&& bIsReadOnly == Other.bIsReadOnly;
 	}
 
 	bool IsValid() const
 	{
 		return Name.IsValid() &&
 			!Name.IsNone() &&
-			TypeName.IsValid() &&
-			!TypeName.IsNone();
+			BaseCPPType.IsValid() &&
+			!BaseCPPType.IsNone();
+	}
+
+	const FName& GetName() const
+	{
+		return Name;
+	}
+
+	void SetName(const FName& InName)
+	{
+		Name = InName;
+	}
+	
+	const FProperty* GetProperty() const
+	{
+		return Property;
+	};
+
+	void SetProperty(const FProperty* InProperty)
+	{
+		Property = InProperty;
+	}
+	
+	const FName& GetBaseCPPType() const
+	{
+		return BaseCPPType;
 	}
 
-	FName GetExtendedCPPType() const
+	const FName& GetExtendedCPPType() const
 	{
 		if(bIsArray)
 		{
-			return *RigVMTypeUtils::ArrayTypeFromBaseType(TypeName.ToString());
+			if (!ExtendedCPPType.IsSet())
+			{
+				ExtendedCPPType = *RigVMTypeUtils::ArrayTypeFromBaseType(BaseCPPType.ToString());
+			}
+			return ExtendedCPPType.GetValue();
 		}
-		return TypeName;
+		return BaseCPPType;
 	}
 
+	bool IsSameType(const FRigVMExternalVariableDef& InOther) const
+	{
+		return GetBaseCPPType() == InOther.GetBaseCPPType() &&
+			IsArray() == InOther.IsArray() &&
+			GetCPPTypeObject() == InOther.GetCPPTypeObject();
+	}
+
+	const UObject* GetCPPTypeObject() const
+	{
+		return CPPTypeObject;
+	}
+	
+	UObject* GetCPPTypeObject()
+	{
+		return CPPTypeObject;
+	}
+	
+	bool IsArray() const
+	{
+		return bIsArray;
+	}
+	
+	bool IsPublic() const
+	{
+		return bIsPublic;
+	}
+	
+	void SetIsPublic(bool InIsPublic)
+	{
+		bIsPublic = InIsPublic;
+	}
+	
+	bool IsReadOnly() const
+	{
+		return bIsReadOnly;
+	}
+	
+	void SetIsReadOnly(bool InIsReadOnly)
+	{
+		bIsReadOnly = InIsReadOnly;
+	}
+	
+	RIGVM_API friend FArchive& operator<<(FArchive& Ar, FRigVMExternalVariableDef& Variable);
+
+protected:
+
+	RIGVM_API void CheckCPPTypeIntegrity() const;
+	
 	FName Name = NAME_None;
 	const FProperty* Property =  nullptr;
-	FName TypeName = NAME_None;
-	UObject* TypeObject = nullptr;
+	FName BaseCPPType = NAME_None;
+	mutable TOptional<FName> ExtendedCPPType;
+	UObject* CPPTypeObject = nullptr;
 	bool bIsArray = false;
 	bool bIsPublic = false;
 	bool bIsReadOnly = false;
-	int32 Size = 0;
+
+	friend struct FRigVMExternalVariableModifier; 
 };
 
 USTRUCT()
 struct FRigVMExternalVariable : public FRigVMExternalVariableDef
 {
 	GENERATED_BODY()
+public:
 	
 	FRigVMExternalVariable()
 		: FRigVMExternalVariableDef()
@@ -79,11 +159,19 @@ struct FRigVMExternalVariable : public FRigVMExternalVariableDef
 		: FRigVMExternalVariableDef(Other)
 		, Memory(nullptr)
 	{
+		CheckCPPTypeIntegrity();
 	}
 
 	FRigVMExternalVariable(const FRigVMExternalVariableDef& Other, uint8* InMemory)
 		: FRigVMExternalVariableDef(Other)
 		, Memory(InMemory)
+	{
+		CheckCPPTypeIntegrity();
+	}
+
+	FRigVMExternalVariable(const FRigVMExternalVariable& Other)
+	: FRigVMExternalVariableDef(Other)
+	, Memory(Other.Memory)
 	{
 	}
 
@@ -91,7 +179,34 @@ struct FRigVMExternalVariable : public FRigVMExternalVariableDef
 	{
 		return RigVMPropertyUtils::GetPropertyHashStable(Property, Memory);
 	}
-	
+
+	const uint8* GetMemory() const
+	{
+		return Memory;
+	}
+
+	uint8* GetMemory()
+	{
+		return Memory;
+	}
+
+	void SetMemory(uint8* InMemory)
+	{
+		Memory = InMemory;
+	}
+
+	friend FArchive& operator<<(FArchive& Ar, FRigVMExternalVariable& Variable)
+	{
+		FRigVMExternalVariableDef& RigVMExternalVariableDef = static_cast<FRigVMExternalVariableDef&>(Variable);
+		Ar << RigVMExternalVariableDef;
+
+		return Ar;
+	}
+
+	bool operator == (const FRigVMExternalVariableDef& Other) const
+	{
+		return ((FRigVMExternalVariableDef)*this) == ((FRigVMExternalVariableDef)Other);  
+	}
 
 	static FRigVMExternalVariable Make(const FProperty* InProperty, void* InContainer, const FName& InOptionalName = NAME_None)
 	{
@@ -116,9 +231,9 @@ struct FRigVMExternalVariable : public FRigVMExternalVariableDef
 			Property = ArrayProperty->Inner;
 		}
 
-		ExternalVariable.Size = Property->GetSize();
-		RigVMPropertyUtils::GetTypeFromProperty(Property, ExternalVariable.TypeName, ExternalVariable.TypeObject);
+		RigVMPropertyUtils::GetTypeFromProperty(Property, ExternalVariable.BaseCPPType, ExternalVariable.CPPTypeObject);
 
+		ExternalVariable.CheckCPPTypeIntegrity();
 		return ExternalVariable;
 	}
 
@@ -126,11 +241,11 @@ struct FRigVMExternalVariable : public FRigVMExternalVariableDef
 	{
 		FRigVMExternalVariable Variable;
 		Variable.Name = InName;
-		Variable.TypeName = TEXT("bool");
-		Variable.TypeObject = nullptr;
+		Variable.BaseCPPType = TEXT("bool");
+		Variable.CPPTypeObject = nullptr;
 		Variable.bIsArray = false;
-		Variable.Size = sizeof(bool);
 		Variable.Memory = (uint8*)&InValue;
+		Variable.CheckCPPTypeIntegrity();
 		return Variable;
 	}
 
@@ -138,11 +253,11 @@ struct FRigVMExternalVariable : public FRigVMExternalVariableDef
 	{
 		FRigVMExternalVariable Variable;
 		Variable.Name = InName;
-		Variable.TypeName = TEXT("bool");
-		Variable.TypeObject = nullptr;
+		Variable.BaseCPPType = TEXT("bool");
+		Variable.CPPTypeObject = nullptr;
 		Variable.bIsArray = true;
-		Variable.Size = sizeof(bool);
 		Variable.Memory = (uint8*)&InValue;
+		Variable.CheckCPPTypeIntegrity();
 		return Variable;
 	}
 
@@ -150,11 +265,11 @@ struct FRigVMExternalVariable : public FRigVMExternalVariableDef
 	{
 		FRigVMExternalVariable Variable;
 		Variable.Name = InName;
-		Variable.TypeName = TEXT("int32");
-		Variable.TypeObject = nullptr;
+		Variable.BaseCPPType = TEXT("int32");
+		Variable.CPPTypeObject = nullptr;
 		Variable.bIsArray = false;
-		Variable.Size = sizeof(int32);
 		Variable.Memory = (uint8*)&InValue;
+		Variable.CheckCPPTypeIntegrity();
 		return Variable;
 	}
 
@@ -162,11 +277,11 @@ struct FRigVMExternalVariable : public FRigVMExternalVariableDef
 	{
 		FRigVMExternalVariable Variable;
 		Variable.Name = InName;
-		Variable.TypeName = TEXT("int32");
-		Variable.TypeObject = nullptr;
+		Variable.BaseCPPType = TEXT("int32");
+		Variable.CPPTypeObject = nullptr;
 		Variable.bIsArray = true;
-		Variable.Size = sizeof(int32);
 		Variable.Memory = (uint8*)&InValue;
+		Variable.CheckCPPTypeIntegrity();
 		return Variable;
 	}
 
@@ -174,11 +289,11 @@ struct FRigVMExternalVariable : public FRigVMExternalVariableDef
 	{
 		FRigVMExternalVariable Variable;
 		Variable.Name = InName;
-		Variable.TypeName = TEXT("uint8");
-		Variable.TypeObject = nullptr;
+		Variable.BaseCPPType = TEXT("uint8");
+		Variable.CPPTypeObject = nullptr;
 		Variable.bIsArray = false;
-		Variable.Size = sizeof(uint8);
 		Variable.Memory = (uint8*)&InValue;
+		Variable.CheckCPPTypeIntegrity();
 		return Variable;
 	}
 
@@ -186,11 +301,11 @@ struct FRigVMExternalVariable : public FRigVMExternalVariableDef
 	{
 		FRigVMExternalVariable Variable;
 		Variable.Name = InName;
-		Variable.TypeName = TEXT("uint8");
-		Variable.TypeObject = nullptr;
+		Variable.BaseCPPType = TEXT("uint8");
+		Variable.CPPTypeObject = nullptr;
 		Variable.bIsArray = true;
-		Variable.Size = sizeof(uint8);
 		Variable.Memory = (uint8*)&InValue;
+		Variable.CheckCPPTypeIntegrity();
 		return Variable;
 	}
 
@@ -198,11 +313,11 @@ struct FRigVMExternalVariable : public FRigVMExternalVariableDef
 	{
 		FRigVMExternalVariable Variable;
 		Variable.Name = InName;
-		Variable.TypeName = TEXT("float");
-		Variable.TypeObject = nullptr;
+		Variable.BaseCPPType = TEXT("float");
+		Variable.CPPTypeObject = nullptr;
 		Variable.bIsArray = false;
-		Variable.Size = sizeof(float);
 		Variable.Memory = (uint8*)&InValue;
+		Variable.CheckCPPTypeIntegrity();
 		return Variable;
 	}
 
@@ -210,11 +325,11 @@ struct FRigVMExternalVariable : public FRigVMExternalVariableDef
 	{
 		FRigVMExternalVariable Variable;
 		Variable.Name = InName;
-		Variable.TypeName = TEXT("float");
-		Variable.TypeObject = nullptr;
+		Variable.BaseCPPType = TEXT("float");
+		Variable.CPPTypeObject = nullptr;
 		Variable.bIsArray = true;
-		Variable.Size = sizeof(float);
 		Variable.Memory = (uint8*)&InValue;
+		Variable.CheckCPPTypeIntegrity();
 		return Variable;
 	}
 
@@ -222,11 +337,11 @@ struct FRigVMExternalVariable : public FRigVMExternalVariableDef
 	{
 		FRigVMExternalVariable Variable;
 		Variable.Name = InName;
-		Variable.TypeName = TEXT("double");
-		Variable.TypeObject = nullptr;
+		Variable.BaseCPPType = TEXT("double");
+		Variable.CPPTypeObject = nullptr;
 		Variable.bIsArray = false;
-		Variable.Size = sizeof(double);
 		Variable.Memory = (uint8*)&InValue;
+		Variable.CheckCPPTypeIntegrity();
 		return Variable;
 	}
 
@@ -234,11 +349,11 @@ struct FRigVMExternalVariable : public FRigVMExternalVariableDef
 	{
 		FRigVMExternalVariable Variable;
 		Variable.Name = InName;
-		Variable.TypeName = TEXT("double");
-		Variable.TypeObject = nullptr;
+		Variable.BaseCPPType = TEXT("double");
+		Variable.CPPTypeObject = nullptr;
 		Variable.bIsArray = true;
-		Variable.Size = sizeof(double);
 		Variable.Memory = (uint8*)&InValue;
+		Variable.CheckCPPTypeIntegrity();
 		return Variable;
 	}
 
@@ -246,11 +361,11 @@ struct FRigVMExternalVariable : public FRigVMExternalVariableDef
 	{
 		FRigVMExternalVariable Variable;
 		Variable.Name = InName;
-		Variable.TypeName = TEXT("FString");
-		Variable.TypeObject = nullptr;
+		Variable.BaseCPPType = TEXT("FString");
+		Variable.CPPTypeObject = nullptr;
 		Variable.bIsArray = false;
-		Variable.Size = sizeof(FString);
 		Variable.Memory = (uint8*)&InValue;
+		Variable.CheckCPPTypeIntegrity();
 		return Variable;
 	}
 
@@ -258,11 +373,11 @@ struct FRigVMExternalVariable : public FRigVMExternalVariableDef
 	{
 		FRigVMExternalVariable Variable;
 		Variable.Name = InName;
-		Variable.TypeName = TEXT("FString");
-		Variable.TypeObject = nullptr;
+		Variable.BaseCPPType = TEXT("FString");
+		Variable.CPPTypeObject = nullptr;
 		Variable.bIsArray = true;
-		Variable.Size = sizeof(FString);
 		Variable.Memory = (uint8*)&InValue;
+		Variable.CheckCPPTypeIntegrity();
 		return Variable;
 	}
 
@@ -270,11 +385,11 @@ struct FRigVMExternalVariable : public FRigVMExternalVariableDef
 	{
 		FRigVMExternalVariable Variable;
 		Variable.Name = InName;
-		Variable.TypeName = TEXT("FName");
-		Variable.TypeObject = nullptr;
+		Variable.BaseCPPType = TEXT("FName");
+		Variable.CPPTypeObject = nullptr;
 		Variable.bIsArray = false;
-		Variable.Size = sizeof(FName);
 		Variable.Memory = (uint8*)&InValue;
+		Variable.CheckCPPTypeIntegrity();
 		return Variable;
 	}
 
@@ -282,11 +397,11 @@ struct FRigVMExternalVariable : public FRigVMExternalVariableDef
 	{
 		FRigVMExternalVariable Variable;
 		Variable.Name = InName;
-		Variable.TypeName = TEXT("FName");
-		Variable.TypeObject = nullptr;
+		Variable.BaseCPPType = TEXT("FName");
+		Variable.CPPTypeObject = nullptr;
 		Variable.bIsArray = true;
-		Variable.Size = sizeof(FName);
 		Variable.Memory = (uint8*)&InValue;
+		Variable.CheckCPPTypeIntegrity();
 		return Variable;
 	}
 
@@ -298,11 +413,11 @@ struct FRigVMExternalVariable : public FRigVMExternalVariableDef
 	{
 		FRigVMExternalVariable Variable;
 		Variable.Name = InName;
-		Variable.TypeName = StaticEnum<T>()->GetFName();
-		Variable.TypeObject = StaticEnum<T>();
+		Variable.BaseCPPType = StaticEnum<T>()->GetFName();
+		Variable.CPPTypeObject = StaticEnum<T>();
 		Variable.bIsArray = false;
-		Variable.Size = sizeof(T);
 		Variable.Memory = (uint8*)&InValue;
+		Variable.CheckCPPTypeIntegrity();
 		return Variable;
 	}
 
@@ -314,11 +429,11 @@ struct FRigVMExternalVariable : public FRigVMExternalVariableDef
 	{
 		FRigVMExternalVariable Variable;
 		Variable.Name = InName;
-		Variable.TypeName = StaticEnum<T>()->GetFName();
-		Variable.TypeObject = StaticEnum<T>();
+		Variable.BaseCPPType = StaticEnum<T>()->GetFName();
+		Variable.CPPTypeObject = StaticEnum<T>();
 		Variable.bIsArray = true;
-		Variable.Size = sizeof(T);
 		Variable.Memory = (uint8*)&InValue;
+		Variable.CheckCPPTypeIntegrity();
 		return Variable;
 	}
 
@@ -330,11 +445,11 @@ struct FRigVMExternalVariable : public FRigVMExternalVariableDef
 	{
 		FRigVMExternalVariable Variable;
 		Variable.Name = InName;
-		Variable.TypeName = *RigVMTypeUtils::GetUniqueStructTypeName(TBaseStructure<T>::Get());
-		Variable.TypeObject = TBaseStructure<T>::Get();
+		Variable.BaseCPPType = *RigVMTypeUtils::GetUniqueStructTypeName(TBaseStructure<T>::Get());
+		Variable.CPPTypeObject = TBaseStructure<T>::Get();
 		Variable.bIsArray = false;
-		Variable.Size = TBaseStructure<T>::Get()->GetStructureSize();
 		Variable.Memory = (uint8*)&InValue;
+		Variable.CheckCPPTypeIntegrity();
 		return Variable;
 	}
 
@@ -346,11 +461,11 @@ struct FRigVMExternalVariable : public FRigVMExternalVariableDef
 	{
 		FRigVMExternalVariable Variable;
 		Variable.Name = InName;
-		Variable.TypeName = *RigVMTypeUtils::GetUniqueStructTypeName(TBaseStructure<T>::Get());
-		Variable.TypeObject = TBaseStructure<T>::Get();
+		Variable.BaseCPPType = *RigVMTypeUtils::GetUniqueStructTypeName(TBaseStructure<T>::Get());
+		Variable.CPPTypeObject = TBaseStructure<T>::Get();
 		Variable.bIsArray = true;
-		Variable.Size = TBaseStructure<T>::Get()->GetStructureSize();
 		Variable.Memory = (uint8*)&InValue;
+		Variable.CheckCPPTypeIntegrity();
 		return Variable;
 	}
 
@@ -362,11 +477,11 @@ struct FRigVMExternalVariable : public FRigVMExternalVariableDef
 	{
 		FRigVMExternalVariable Variable;
 		Variable.Name = InName;
-		Variable.TypeName = *RigVMTypeUtils::GetUniqueStructTypeName(T::StaticStruct());
-		Variable.TypeObject = T::StaticStruct();
+		Variable.BaseCPPType = *RigVMTypeUtils::GetUniqueStructTypeName(T::StaticStruct());
+		Variable.CPPTypeObject = T::StaticStruct();
 		Variable.bIsArray = false;
-		Variable.Size = T::StaticStruct()->GetStructureSize();
 		Variable.Memory = (uint8*)&InValue;
+		Variable.CheckCPPTypeIntegrity();
 		return Variable;
 	}
 
@@ -378,11 +493,11 @@ struct FRigVMExternalVariable : public FRigVMExternalVariableDef
 	{
 		FRigVMExternalVariable Variable;
 		Variable.Name = InName;
-		Variable.TypeName = *RigVMTypeUtils::GetUniqueStructTypeName(T::StaticStruct());
-		Variable.TypeObject = T::StaticStruct();
+		Variable.BaseCPPType = *RigVMTypeUtils::GetUniqueStructTypeName(T::StaticStruct());
+		Variable.CPPTypeObject = T::StaticStruct();
 		Variable.bIsArray = true;
-		Variable.Size = T::StaticStruct()->GetStructureSize();
 		Variable.Memory = (uint8*)&InValue;
+		Variable.CheckCPPTypeIntegrity();
 		return Variable;
 	}
 
@@ -394,11 +509,11 @@ struct FRigVMExternalVariable : public FRigVMExternalVariableDef
 	{
 		FRigVMExternalVariable Variable;
 		Variable.Name = InName;
-		Variable.TypeName = *T::StaticClass()->GetStructCPPName();
-		Variable.TypeObject = T::StaticClass();
+		Variable.BaseCPPType = *T::StaticClass()->GetStructCPPName();
+		Variable.CPPTypeObject = T::StaticClass();
 		Variable.bIsArray = false;
-		Variable.Size = T::StaticClass()->GetStructureSize();
 		Variable.Memory = (uint8*)&InValue;
+		Variable.CheckCPPTypeIntegrity();
 		return Variable;
 	}
 
@@ -410,11 +525,11 @@ struct FRigVMExternalVariable : public FRigVMExternalVariableDef
 	{
 		FRigVMExternalVariable Variable;
 		Variable.Name = InName;
-		Variable.TypeName = *T::StaticClass()->GetStructCPPName();
-		Variable.TypeObject = T::StaticClass();
+		Variable.BaseCPPType = *T::StaticClass()->GetStructCPPName();
+		Variable.CPPTypeObject = T::StaticClass();
 		Variable.bIsArray = true;
-		Variable.Size = T::StaticClass()->GetStructureSize();
 		Variable.Memory = (uint8*)&InValue;
+		Variable.CheckCPPTypeIntegrity();
 		return Variable;
 	}
 
@@ -468,6 +583,14 @@ struct FRigVMExternalVariable : public FRigVMExternalVariableDef
 
 	RIGVM_API TRigVMTypeIndex GetTypeIndex() const;
 
+	FString GetCPPTypeObjectPathName() const
+	{
+		if (CPPTypeObject)
+		{
+			return CPPTypeObject->GetPathName();
+		}
+		return FString();
+	}
 
 	template<typename VarType>
 	static void MergeExternalVariable(TArray<VarType>& OutVariables, const VarType& InVariable)
@@ -481,8 +604,6 @@ struct FRigVMExternalVariable : public FRigVMExternalVariableDef
 		{
 			if(ExistingVariable.Name == InVariable.Name)
 			{
-				ensure(ExistingVariable.TypeName == InVariable.TypeName);
-				ensure(ExistingVariable.TypeObject == InVariable.TypeObject);
 				return;
 			}
 		}
@@ -490,40 +611,22 @@ struct FRigVMExternalVariable : public FRigVMExternalVariableDef
 		OutVariables.Add(InVariable);
 	}
 
-	uint8* Memory = nullptr;
-};
+#if WITH_EDITOR
 
-inline FArchive& operator<<(FArchive& Ar, FRigVMExternalVariableDef& Variable)
-{
-	Ar << Variable.Name;
-	Ar << Variable.TypeName;
+	RIGVM_API static FRigVMExternalVariable Make(const FName& InName, const FEdGraphPinType& InPinType, bool bInPublic = false, bool bInReadonly = false);
 
-	if (Ar.IsSaving())
-	{
-		FSoftObjectPath TypeObjectPath(Variable.TypeObject);
-		Ar << TypeObjectPath;
-	}
-	else if (Ar.IsLoading())
-	{
-		FSoftObjectPath TypeObjectPath;
-		Ar << TypeObjectPath;
-		Variable.TypeObject = TypeObjectPath.ResolveObject();
-	}
+#endif
 	
-	Ar << Variable.bIsArray;
-	Ar << Variable.bIsPublic;
-	Ar << Variable.bIsReadOnly;
-	Ar << Variable.Size;
-	return Ar;
-}
+	RIGVM_API static FRigVMExternalVariable Make(const FName& InName, const FString& InCPPTypePath, bool bInPublic = false, bool bInReadonly = false);
 
-inline FArchive& operator<<(FArchive& Ar, FRigVMExternalVariable& Variable)
-{
-	FRigVMExternalVariableDef& RigVMExternalVariableDef = static_cast<FRigVMExternalVariableDef&>(Variable);
-	Ar << RigVMExternalVariableDef;
+	RIGVM_API static FRigVMExternalVariable Make(const FName& InName, const FString& InCPPType, UObject* InCPPTypeObject, bool bInPublic = false, bool bInReadonly = false);
 
-	return Ar;
-}
+private:
+	
+	uint8* Memory = nullptr;
+
+	friend struct FRigVMExternalVariableModifier; 
+};
 
 namespace RigVMTypeUtils
 {
diff --git a/Engine/Plugins/Runtime/RigVM/Source/RigVM/Public/RigVMCore/RigVMGraphFunctionDefinition.h b/Engine/Plugins/Runtime/RigVM/Source/RigVM/Public/RigVMCore/RigVMGraphFunctionDefinition.h
--- a/Engine/Plugins/Runtime/RigVM/Source/RigVM/Public/RigVMCore/RigVMGraphFunctionDefinition.h
+++ b/Engine/Plugins/Runtime/RigVM/Source/RigVM/Public/RigVMCore/RigVMGraphFunctionDefinition.h
@@ -233,8 +233,11 @@ struct FRigVMGraphFunctionArgument
 	, bIsArray(false)
 	, Direction(ERigVMPinDirection::Input)
 	, bIsConst(false)
+	, bIsInputVariable(false)
 	{}
 
+	UE_API FRigVMGraphFunctionArgument(const FRigVMExternalVariable& InExternalVariable);
+
 	UPROPERTY(BlueprintReadOnly, VisibleAnywhere, Category=FunctionArgument)
 	FName Name;
 
@@ -258,12 +261,13 @@ struct FRigVMGraphFunctionArgument
 
 	UPROPERTY(BlueprintReadOnly, VisibleAnywhere, Category=FunctionArgument)
 	bool bIsConst;
-	
+
+	UPROPERTY(BlueprintReadOnly, VisibleAnywhere, Category=FunctionArgument)
+	bool bIsInputVariable;
+
 	UPROPERTY()
 	TMap<FString, FText> PathToTooltip;
 
-	UE_API FRigVMExternalVariable GetExternalVariable() const;
-
 	// True if CPP Type or obj type is valid. May potentially load the CPP Type Object
 	UE_API bool IsValid() const;
 
@@ -282,6 +286,7 @@ struct FRigVMGraphFunctionArgument
 		Hash = HashCombine(Hash, GetTypeHash(Argument.Direction));
 		Hash = HashCombine(Hash, GetTypeHash(Argument.DefaultValue));
 		Hash = HashCombine(Hash, GetTypeHash(Argument.bIsConst));
+		Hash = HashCombine(Hash, GetTypeHash(Argument.bIsInputVariable));
 		for (const TPair<FString, FText>& Pair : Argument.PathToTooltip)
 		{
 			Hash = HashCombine(Hash, GetTypeHash(Pair.Key));
@@ -295,19 +300,15 @@ struct FRigVMGraphFunctionArgument
 		return true;
 	}
 
+	UE_API void Serialize(FArchive& Ar);
+
 	friend FArchive& operator<<(FArchive& Ar, FRigVMGraphFunctionArgument& Data)
 	{
-		Ar << Data.Name;
-		Ar << Data.DisplayName;
-		Ar << Data.CPPType;
-		Ar << Data.CPPTypeObject;
-		Ar << Data.bIsArray;
-		Ar << Data.Direction;
-		Ar << Data.DefaultValue;
-		Ar << Data.bIsConst;
-		Ar << Data.PathToTooltip;
+		Data.Serialize(Ar);
 		return Ar;
 	}
+
+	UE_API FRigVMExternalVariable ToExternalVariable() const;
 };
 
 USTRUCT(BlueprintType)
@@ -494,6 +495,8 @@ struct FRigVMGraphFunctionHeader
 
 	static UE_API FRigVMGraphFunctionHeader FindGraphFunctionHeader(const FRigVMGraphFunctionIdentifier& InIdentifier, bool* bOutIsPublic = nullptr, FString* OutErrorMessage = nullptr);
 
+	RIGVM_API TArray<FRigVMExternalVariable> GetExternalVariables() const;
+
 protected:
 
 	static UE_API FName GetFunctionNameFromObjectPath(const FString& InObjectPath, const FName& InOptionalFunctionName = NAME_None);
@@ -537,6 +540,8 @@ struct FRigVMGraphFunctionData
 
 	RIGVM_API bool IsMutable() const;
 
+	RIGVM_API bool IsPublic() const;
+
 	bool operator==(const FRigVMGraphFunctionData& Other) const
 	{
 		return Header == Other.Header;
@@ -557,6 +562,7 @@ struct FRigVMGraphFunctionData
 	RIGVM_API bool IsAnyOperandSharedAcrossArguments() const;
 
 	RIGVM_API bool PatchSharedArgumentOperandsIfRequired();
+	RIGVM_API bool PatchExternalVariablesToArguments();
 
 	static const inline TCHAR* EntryString = TEXT("Entry");
 	static const inline TCHAR* ReturnString = TEXT("Return");
diff --git a/Engine/Plugins/Runtime/RigVM/Source/RigVM/Public/RigVMCore/RigVMGraphFunctionHost.h b/Engine/Plugins/Runtime/RigVM/Source/RigVM/Public/RigVMCore/RigVMGraphFunctionHost.h
--- a/Engine/Plugins/Runtime/RigVM/Source/RigVM/Public/RigVMCore/RigVMGraphFunctionHost.h
+++ b/Engine/Plugins/Runtime/RigVM/Source/RigVM/Public/RigVMCore/RigVMGraphFunctionHost.h
@@ -62,7 +62,7 @@ public:
 
 	UE_API bool UpdateDependencies(const FRigVMGraphFunctionIdentifier& InLibraryPointer, TMap<FRigVMGraphFunctionIdentifier, uint32>& Dependencies);
 
-	UE_API bool UpdateExternalVariables(const FRigVMGraphFunctionIdentifier& InLibraryPointer, TArray<FRigVMExternalVariable> ExternalVariables);
+	UE_API bool UpdateExternalVariables(const FRigVMGraphFunctionIdentifier& InLibraryPointer, TArray<FRigVMExternalVariable> InExternalVariables);
 
 	UE_API bool UpdateFunctionCompilationData(const FRigVMGraphFunctionIdentifier& InLibraryPointer, const FRigVMFunctionCompilationData& CompilationData);
 
@@ -93,6 +93,7 @@ public:
 		for(FRigVMGraphFunctionData& Function: PublicFunctions)
 		{
 			Function.PatchSharedArgumentOperandsIfRequired();
+			Function.PatchExternalVariablesToArguments();
 		}
 	}
 
diff --git a/Engine/Plugins/Runtime/RigVM/Source/RigVM/Public/RigVMCore/RigVMNativized.h b/Engine/Plugins/Runtime/RigVM/Source/RigVM/Public/RigVMCore/RigVMNativized.h
--- a/Engine/Plugins/Runtime/RigVM/Source/RigVM/Public/RigVMCore/RigVMNativized.h
+++ b/Engine/Plugins/Runtime/RigVM/Source/RigVM/Public/RigVMCore/RigVMNativized.h
@@ -198,11 +198,11 @@ protected:
 	{
 		for(const FRigVMExternalVariable& ExternalVariable : GetExternalVariables(Context))
 		{
-			if(ExternalVariable.Name == InExternalVariableName && ExternalVariable.GetExtendedCPPType() == InExpectedType)
+			if(ExternalVariable.GetName() == InExternalVariableName && ExternalVariable.GetExtendedCPPType() == InExpectedType)
 			{
 				if(ExternalVariable.IsValid(false))
 				{
-					return *(T*)ExternalVariable.Memory;
+					return *(T*)ExternalVariable.GetMemory();
 				}
 			}
 		}
diff --git a/Engine/Plugins/Runtime/RigVM/Source/RigVM/Public/RigVMObjectVersion.h b/Engine/Plugins/Runtime/RigVM/Source/RigVM/Public/RigVMObjectVersion.h
--- a/Engine/Plugins/Runtime/RigVM/Source/RigVM/Public/RigVMObjectVersion.h
+++ b/Engine/Plugins/Runtime/RigVM/Source/RigVM/Public/RigVMObjectVersion.h
@@ -75,6 +75,9 @@ struct FRigVMObjectVersion
 		// The VM stores a relative seek offset to be able to skip the registry during load
 		LocalizedRegistryWithRelativeSeekOffset,
 
+		// Function arguments can now represent an input variable (an external variable passed into a funciton)
+		FunctionArgumentCanRepresentInputVariable,
+
 		// -----<new versions can be added above this line>-------------------------------------------------
 		VersionPlusOne,
 		LatestVersion = VersionPlusOne - 1
diff --git a/Engine/Plugins/Runtime/RigVM/Source/RigVMDeveloper/Private/EdGraph/RigVMEdGraphNode.cpp b/Engine/Plugins/Runtime/RigVM/Source/RigVMDeveloper/Private/EdGraph/RigVMEdGraphNode.cpp
--- a/Engine/Plugins/Runtime/RigVM/Source/RigVMDeveloper/Private/EdGraph/RigVMEdGraphNode.cpp
+++ b/Engine/Plugins/Runtime/RigVM/Source/RigVMDeveloper/Private/EdGraph/RigVMEdGraphNode.cpp
@@ -790,14 +790,16 @@ bool URigVMEdGraphNode::CreateGraphPinFromModelPin(const URigVMPin* InModelPin,
 		// check if we already have a pin like this in the last pins arrays
 		const FString PinPath = InModelPin->GetPinPath();
 		const FEdGraphPinType PinType = GetPinTypeForModelPin(InModelPin);
+		const bool bShouldDisplayAsMutableRef = InModelPin->IsWildCard() || InModelPin->IsDefinedAsInputVariable();
 
 		UEdGraphPin* GraphPin = nullptr;
-		UEdGraphPin** ExistingEdGraphPinPtr = LastEdGraphPins.FindByPredicate([PinPath, InDirection, PinType](const UEdGraphPin* ExistingPin) -> bool
+		UEdGraphPin** ExistingEdGraphPinPtr = LastEdGraphPins.FindByPredicate([PinPath, InDirection, PinType, bShouldDisplayAsMutableRef](const UEdGraphPin* ExistingPin) -> bool
 		{
 			return !ExistingPin->bWasTrashed &&
 				ExistingPin->GetName() == PinPath &&
 				ExistingPin->Direction == InDirection &&
-				ExistingPin->PinType == PinType;
+				ExistingPin->PinType == PinType &&
+				ExistingPin->bDisplayAsMutableRef == bShouldDisplayAsMutableRef;
 		});
 		if(ExistingEdGraphPinPtr)
 		{
@@ -1656,7 +1658,7 @@ void URigVMEdGraphNode::AllocateDefaultPins()
 	{
 		if(FunctionReferenceNode->RequiresVariableRemapping())
 		{
-			TArray<FRigVMExternalVariable> CurrentExternalVariables = FunctionReferenceNode->GetExternalVariables(false);
+			TArray<FRigVMExternalVariable> CurrentExternalVariables = FunctionReferenceNode->GetExternalVariables();
 			for(const FRigVMExternalVariable& CurrentExternalVariable : CurrentExternalVariables)
 			{
 				ExternalVariables.Add(MakeShared<FRigVMExternalVariable>(CurrentExternalVariable));
@@ -2259,6 +2261,7 @@ FEdGraphPinType URigVMEdGraphNode::GetPinTypeForModelPin(const URigVMPin* InMode
 {
 	FEdGraphPinType PinType = RigVMTypeUtils::PinTypeFromCPPType(*InModelPin->GetCPPType(), InModelPin->GetCPPTypeObject());
 	PinType.bIsConst = InModelPin->IsDefinedAsConstant();
+	PinType.bIsReference = InModelPin->IsDefinedAsInputVariable();
 	return PinType;
 }
 
@@ -2289,7 +2292,23 @@ void URigVMEdGraphNode::ConfigurePin(UEdGraphPin* EdGraphPin, const URigVMPin* M
 	EdGraphPin->PinFriendlyName = FText::FromName(ModelPin->GetDisplayName());
 	EdGraphPin->bNotConnectable = !bConnectable  || ModelPin->IsFixedSizeArray();
 	EdGraphPin->bOrphanedPin = ModelPin->IsOrphanPin() ? 1 : 0; 
-	EdGraphPin->bDisplayAsMutableRef = ModelPin->IsWildCard();
+	EdGraphPin->bDisplayAsMutableRef = ModelPin->IsWildCard() || ModelPin->IsDefinedAsInputVariable();
+	if (!EdGraphPin->bDisplayAsMutableRef)
+	{
+		if(ModelPin->GetDirection() == ERigVMPinDirection::Output && ModelPin->IsRootPin())
+		{
+			if (const URigVMVariableNode* VariableNode = Cast<URigVMVariableNode>(ModelPin->GetNode()))
+			{
+				if (!VariableNode->IsInputArgument())
+				{
+					if (VariableNode->GetValuePin() == ModelPin)
+					{
+						EdGraphPin->bDisplayAsMutableRef = true;
+					}
+				}
+			}
+		}
+	}
 }
 
 TArray<URigVMPin*>& URigVMEdGraphNode::PinListForPin(const URigVMPin* InModelPin)
diff --git a/Engine/Plugins/Runtime/RigVM/Source/RigVMDeveloper/Private/EdGraph/RigVMEdGraphSchema.cpp b/Engine/Plugins/Runtime/RigVM/Source/RigVMDeveloper/Private/EdGraph/RigVMEdGraphSchema.cpp
--- a/Engine/Plugins/Runtime/RigVM/Source/RigVMDeveloper/Private/EdGraph/RigVMEdGraphSchema.cpp
+++ b/Engine/Plugins/Runtime/RigVM/Source/RigVMDeveloper/Private/EdGraph/RigVMEdGraphSchema.cpp
@@ -397,11 +397,12 @@ UEdGraphNode* FRigVMEdGraphSchemaAction_PromoteToVariable::PerformAction(UEdGrap
 			}
 		}
 
-		FRigVMExternalVariable ExternalVariable;
-		ExternalVariable.Name = FromPin->GetFName();
-		ExternalVariable.bIsArray = ModelPin->IsArray();
-		ExternalVariable.TypeName = ModelPin->IsArray() ? *ModelPin->GetArrayElementCppType() : *ModelPin->GetCPPType();
-		ExternalVariable.TypeObject = ModelPin->GetCPPTypeObject();
+		FRigVMExternalVariable ExternalVariable = FRigVMExternalVariable::Make(
+			FromPin->GetFName(),
+			ModelPin->GetCPPType(),
+			ModelPin->GetCPPTypeObject(),
+			false,
+			false);
 		
 		VariableName = Blueprint->AddHostMemberVariableFromExternal(
 			ExternalVariable,
@@ -1851,17 +1852,13 @@ UEdGraphPin* URigVMEdGraphSchema::DropPinOnNode(UEdGraphNode* InTargetNode, cons
 						if (URigVMController* Controller = RigBlueprint->GetController(Model))
 						{
 							FName SourcePinName = InSourcePinName; 
-							FString TypeName = ExternalVar.TypeName.ToString();
-							if (ExternalVar.bIsArray)
-							{
-								TypeName = RigVMTypeUtils::ArrayTypeFromBaseType(*TypeName);
-							}
+							FString TypeName = ExternalVar.GetExtendedCPPType().ToString();
 							FName TypeObjectPathName = NAME_None;
-							if (ExternalVar.TypeObject)
+							if (ExternalVar.GetCPPTypeObject())
 							{
-								TypeObjectPathName = *ExternalVar.TypeObject->GetPathName();
+								TypeObjectPathName = *ExternalVar.GetCPPTypeObject()->GetPathName();
 
-								if(const UScriptStruct* CPPTypeStruct = Cast<UScriptStruct>(ExternalVar.TypeObject))
+								if(const UScriptStruct* CPPTypeStruct = Cast<UScriptStruct>(ExternalVar.GetCPPTypeObject()))
 								{
 									if(CPPTypeStruct->IsChildOf(FRigVMExecutePin::StaticStruct()))
 									{
@@ -1872,6 +1869,7 @@ UEdGraphPin* URigVMEdGraphSchema::DropPinOnNode(UEdGraphNode* InTargetNode, cons
 							}
 
 							FString DefaultValue;
+							bool bIsDefinedAsInputVariable = false;
 							if (PinBeingDropped)
 							{
 								if (const URigVMEdGraphNode* SourceNode = Cast<URigVMEdGraphNode>(PinBeingDropped->GetOwningNode()))
@@ -1879,6 +1877,7 @@ UEdGraphPin* URigVMEdGraphSchema::DropPinOnNode(UEdGraphNode* InTargetNode, cons
 									if (const URigVMPin* SourcePin = SourceNode->GetModelPinFromPinPath(PinBeingDropped->GetName()))
 									{
 										DefaultValue = SourcePin->GetDefaultValue();
+										bIsDefinedAsInputVariable = SourcePin->IsDefinedAsInputVariable();
 									}
 								}
 							}
@@ -1896,6 +1895,7 @@ UEdGraphPin* URigVMEdGraphSchema::DropPinOnNode(UEdGraphNode* InTargetNode, cons
 							if (!ExposedPinName.IsNone())
 							{
 								NewPinName = ExposedPinName.ToString();
+								Controller->MarkExposedPinAsInputVariable(ExposedPinName, bIsDefinedAsInputVariable);
 							}
 						}
 					}
@@ -2260,12 +2260,12 @@ bool URigVMEdGraphSchema::RequestVariableDropOnPin(UEdGraph* InGraph, const FRig
 					const FModifierKeysState KeyState = FSlateApplication::Get().GetModifierKeys();
 					if (KeyState.IsAltDown())
 					{
-						return Graph->GetController()->BindPinToVariable(ModelPin->GetPinPath(), InVariableToDrop.Name.ToString(), true, true);
+						return Graph->GetController()->BindPinToVariable(ModelPin->GetPinPath(), InVariableToDrop.GetName().ToString(), true, true);
 					}
 					else
 					{
 						Graph->GetController()->OpenUndoBracket(TEXT("Bind Variable to Pin"));
-						if (const URigVMVariableNode* VariableNode = Graph->GetController()->AddVariableNode(InVariableToDrop.Name, InVariableToDrop.TypeName.ToString(), InVariableToDrop.TypeObject, true, FString(), InDropPosition + FVector2D(0.f, -34.f)))
+						if (const URigVMVariableNode* VariableNode = Graph->GetController()->AddVariableNode(InVariableToDrop.GetName(), InVariableToDrop.GetExtendedCPPType().ToString(), const_cast<UObject*>(InVariableToDrop.GetCPPTypeObject()), true, FString(), InDropPosition + FVector2D(0.f, -34.f)))
 						{
 							Graph->GetController()->AddLink(VariableNode->FindPin(TEXT("Value"))->GetPinPath(), ModelPin->GetPinPath(), true);
 						}
diff --git a/Engine/Plugins/Runtime/RigVM/Source/RigVMDeveloper/Private/RigVMBlueprintLegacy.cpp b/Engine/Plugins/Runtime/RigVM/Source/RigVMDeveloper/Private/RigVMBlueprintLegacy.cpp
--- a/Engine/Plugins/Runtime/RigVM/Source/RigVMDeveloper/Private/RigVMBlueprintLegacy.cpp
+++ b/Engine/Plugins/Runtime/RigVM/Source/RigVMDeveloper/Private/RigVMBlueprintLegacy.cpp
@@ -50,16 +50,12 @@
 
 FEdGraphPinType FRigVMOldPublicFunctionArg::GetPinType() const
 {
-	FRigVMExternalVariable Variable;
-	Variable.Name = Name;
-	Variable.bIsArray = bIsArray;
-	Variable.TypeName = CPPType;
-	
+	UObject* CPPTypeObject = nullptr;
 	if(CPPTypeObjectPath.IsValid())
 	{
-		Variable.TypeObject = RigVMTypeUtils::FindObjectFromCPPTypeObjectPath(CPPTypeObjectPath.ToString());
+		CPPTypeObject = RigVMTypeUtils::FindObjectFromCPPTypeObjectPath(CPPTypeObjectPath.ToString());
 	}
-
+	FRigVMExternalVariable Variable = FRigVMExternalVariable::Make(Name, CPPType.ToString(), CPPTypeObject, false, false);
 	return RigVMTypeUtils::PinTypeFromExternalVariable(Variable);
 }
 
@@ -1032,67 +1028,14 @@ bool URigVMBlueprint::ChangeMemberVariableType(const FName& InName, const FStrin
 		return false;
 	}
 
-	FRigVMExternalVariable Variable;
-	Variable.Name = InName;
-	Variable.bIsPublic = bIsPublic;
-	Variable.bIsReadOnly = bIsReadOnly;
-
-	FString CPPType = InCPPType;
-	if (CPPType.StartsWith(TEXT("TMap<")))
-	{
-		UE_LOG(LogRigVMDeveloper, Warning, TEXT("TMap Variables are not supported."));
-		return false;
-	}
-
-	Variable.bIsArray = RigVMTypeUtils::IsArrayType(CPPType);
-	if (Variable.bIsArray)
-	{
-		CPPType = RigVMTypeUtils::BaseTypeFromArrayType(CPPType);
-	}
-
-	if (CPPType == TEXT("bool"))
-	{
-		Variable.TypeName = *CPPType;
-		Variable.Size = sizeof(bool);
-	}
-	else if (CPPType == TEXT("float"))
-	{
-		Variable.TypeName = *CPPType;
-		Variable.Size = sizeof(float);
-	}
-	else if (CPPType == TEXT("double"))
-	{
-		Variable.TypeName = *CPPType;
-		Variable.Size = sizeof(double);
-	}
-	else if (CPPType == TEXT("int32"))
-	{
-		Variable.TypeName = *CPPType;
-		Variable.Size = sizeof(int32);
-	}
-	else if (CPPType == TEXT("FString"))
-	{
-		Variable.TypeName = *CPPType;
-		Variable.Size = sizeof(FString);
-	}
-	else if (CPPType == TEXT("FName"))
-	{
-		Variable.TypeName = *CPPType;
-		Variable.Size = sizeof(FName);
-	}
-	else if(UScriptStruct* ScriptStruct = RigVMTypeUtils::FindObjectFromCPPTypeObjectPath<UScriptStruct>(CPPType))
-	{
-		Variable.TypeName = *RigVMTypeUtils::GetUniqueStructTypeName(ScriptStruct);
-		Variable.TypeObject = ScriptStruct;
-		Variable.Size = ScriptStruct->GetStructureSize();
-	}
-	else if (UEnum* Enum= RigVMTypeUtils::FindObjectFromCPPTypeObjectPath<UEnum>(CPPType))
+	FString BaseCPPType = InCPPType;
+	if (RigVMTypeUtils::IsArrayType(BaseCPPType))
 	{
-		Variable.TypeName = *RigVMTypeUtils::CPPTypeFromEnum(Enum);
-		Variable.TypeObject = Enum;
-		Variable.Size = static_cast<int32>(Enum->GetResourceSizeBytes(EResourceSizeMode::EstimatedTotal));
+		BaseCPPType = RigVMTypeUtils::BaseTypeFromArrayType(BaseCPPType);
 	}
 
+	UObject* CPPTypeObject = RigVMTypeUtils::FindObjectFromCPPTypeObjectPath<UScriptStruct>(BaseCPPType);
+	const FRigVMExternalVariable Variable = FRigVMExternalVariable::Make(InName, InCPPType, CPPTypeObject, bIsPublic, bIsReadOnly);
 	FEdGraphPinType PinType = RigVMTypeUtils::PinTypeFromExternalVariable(Variable);
 	if (!PinType.PinCategory.IsValid())
 	{
@@ -1196,8 +1139,8 @@ FName URigVMBlueprint::AddHostMemberVariableFromExternal(FRigVMExternalVariable
 	Modify();
 
 	TSharedPtr<FKismetNameValidator> NameValidator = MakeShareable(new FKismetNameValidator(this, NAME_None, nullptr));
-	FName VarName = FindHostMemberVariableUniqueName(NameValidator, InVariableToCreate.Name.ToString());
-	int32 VariableIndex = AddHostMemberVariable(this, VarName, PinType, InVariableToCreate.bIsPublic, InVariableToCreate.bIsReadOnly, InDefaultValue);
+	FName VarName = FindHostMemberVariableUniqueName(NameValidator, InVariableToCreate.GetName().ToString());
+	int32 VariableIndex = AddHostMemberVariable(this, VarName, PinType, InVariableToCreate.IsPublic(), InVariableToCreate.IsReadOnly(), InDefaultValue);
 	if (VariableIndex != INDEX_NONE)
 	{
 		FBlueprintEditorUtils::MarkBlueprintAsStructurallyModified(this);
diff --git a/Engine/Plugins/Runtime/RigVM/Source/RigVMDeveloper/Private/RigVMCompiler/RigVMAST.cpp b/Engine/Plugins/Runtime/RigVM/Source/RigVMDeveloper/Private/RigVMCompiler/RigVMAST.cpp
--- a/Engine/Plugins/Runtime/RigVM/Source/RigVMDeveloper/Private/RigVMCompiler/RigVMAST.cpp
+++ b/Engine/Plugins/Runtime/RigVM/Source/RigVMDeveloper/Private/RigVMCompiler/RigVMAST.cpp
@@ -65,10 +65,6 @@ FName FRigVMExprAST::GetTypeName() const
 		{
 			return TEXT("[Literal]");
 		}
-		case EType::ExternalVar:
-		{
-			return TEXT("[ExtVar.]");
-		}
 		case EType::Assign:
 		{
 			return TEXT("[.Assign]");
@@ -1220,6 +1216,10 @@ FRigVMExprAST* FRigVMParserAST::TraverseMutableNode(const FRigVMASTProxy& InNode
 		}
 
 		TraversePins(InNodeProxy, NodeExpr);
+		if (NodeExpr->IsNode())
+		{
+			TraverseNodeExpressionForRequiredVariableClone(NodeExpr->To<FRigVMNodeExprAST>());
+		}
 
 		TArray<URigVMPin*> SourcePins = Node->GetPins().FilterByPredicate([](const URigVMPin* InPin) -> bool
 		{
@@ -1363,71 +1363,9 @@ FRigVMExprAST* FRigVMParserAST::TraverseNode(const FRigVMASTProxy& InNodeProxy,
 	if (NodeExpr)
 	{
 		TraversePins(InNodeProxy, NodeExpr);
-
-		if (const URigVMVariableNode* VariableNode = Cast<URigVMVariableNode>(Node))
+		if (NodeExpr->IsNode())
 		{
-			bool bRequiresCloningVariableValue = false;
-			if (VariableNode->IsGetter())
-			{
-				URigVMPin* VariableValuePin = VariableNode->GetValuePin();
-				FRigVMASTProxy VariableValuePinProxy = NodeExpr->GetProxy().GetSibling(VariableValuePin);
-
-				// for variable nodes we need to deal with a specific case where a sub pin
-				// of the value pin is wired to a variable setter node with the same variable
-				// on a different sub pin. so for example Get MyVector.X -> Set MyVector.Y.
-				// In that case we need to inject a clone of the variable in the getter.
-				// Since arrays don't have sub pins on the getter this can only happen to non-arrays.
-				const TArray<int32>& LinkIndices = GetTargetLinkIndices(VariableValuePinProxy, true);
-				for(const int32 LinkIndex : LinkIndices)
-				{
-					const FRigVMASTLinkDescription& Link = Links[LinkIndex];
-					const URigVMPin* SourcePin = Link.SourceProxy.GetSubjectChecked<URigVMPin>();
-					const URigVMPin* TargetPin = Link.TargetProxy.GetSubjectChecked<URigVMPin>();
-					if (TargetPin && !Link.SegmentPath.IsEmpty())
-					{
-						TargetPin = TargetPin->FindSubPin(Link.SegmentPath);
-					}
-					if (!SourcePin || !TargetPin)
-					{
-						continue;
-					}
-					if (TargetPin->IsRootPin())
-					{
-						continue;
-					}
-					const URigVMVariableNode* TargetVariableNode = Cast<URigVMVariableNode>(TargetPin->GetNode());
-					if (!TargetVariableNode)
-					{
-						continue;
-					}
-					if (VariableNode->GetVariableName() != TargetVariableNode->GetVariableName())
-					{
-						continue;
-					}
-
-					if (SourcePin->GetSegmentPath() != TargetPin->GetSegmentPath())
-					{
-						bRequiresCloningVariableValue = true;
-						break;
-					}
-				}
-			}
-
-			if (bRequiresCloningVariableValue && NodeExpr->IsNode())
-			{
-				FRigVMNodeExprAST* VariableNodeExpr = NodeExpr->To<FRigVMNodeExprAST>();
-				const FRigVMVarExprAST* ConstVariableVarExpr = VariableNodeExpr->FindVarWithPinName(URigVMVariableNode::ValueName);
-				FRigVMVarExprAST* VariableVarExpr = const_cast<FRigVMVarExprAST*>(ConstVariableVarExpr);
-				
-				FRigVMASTProxy PinProxy = VariableVarExpr->GetProxy();
-				FRigVMExprAST* ClonedVariableVarExpr = MakeExpr<FRigVMVarExprAST>(FRigVMExprAST::EType::Var, PinProxy);
-				ClonedVariableVarExpr->Name = *FString::Printf(TEXT("%s::Clone"), *VariableVarExpr->Name.ToString());
-				const FRigVMASTLinkDescription CloneLink(PinProxy, PinProxy, FString());
-				FRigVMExprAST* PinCopyExpr = MakeExpr<FRigVMCopyExprAST>(CloneLink);
-				PinCopyExpr->AddParent(ClonedVariableVarExpr);
-				VariableVarExpr->AddParent(PinCopyExpr);
-				VariableNodeExpr->ReplaceChild(VariableVarExpr, ClonedVariableVarExpr);
-			}
+			TraverseNodeExpressionForRequiredVariableClone(NodeExpr->To<FRigVMNodeExprAST>());
 		}
 	}
 
@@ -1439,6 +1377,7 @@ FRigVMExprAST* FRigVMParserAST::CreateExpressionForNode(const FRigVMASTProxy& In
 	URigVMNode* Node = InNodeProxy.GetSubjectChecked<URigVMNode>();
 
 	bool bIsEntry = Node->IsEvent();
+	FName NodeName = Node->GetEventName();
 	if (LibraryNodeBeingCompiled)
 	{
 		if (Node->GetTypedOuter<URigVMLibraryNode>() == LibraryNodeBeingCompiled)
@@ -1448,6 +1387,7 @@ FRigVMExprAST* FRigVMParserAST::CreateExpressionForNode(const FRigVMASTProxy& In
 				if (EntryNode->IsMutable())
 				{
 					bIsEntry = true;
+					NodeName = TEXT("Function Entry");
 				}
 			}
 			if (URigVMFunctionReturnNode* ReturnNode = Cast<URigVMFunctionReturnNode>(Node))
@@ -1455,6 +1395,7 @@ FRigVMExprAST* FRigVMParserAST::CreateExpressionForNode(const FRigVMASTProxy& In
 				if (!ReturnNode->IsMutable())
 				{
 					bIsEntry = true;
+					NodeName = TEXT("Function Return");
 				}
 			}
 		}
@@ -1464,7 +1405,7 @@ FRigVMExprAST* FRigVMParserAST::CreateExpressionForNode(const FRigVMASTProxy& In
 	if (bIsEntry)
 	{
 		NodeExpr = MakeExpr<FRigVMEntryExprAST>(InNodeProxy);
-		NodeExpr->Name = Node->GetEventName();
+		NodeExpr->Name = NodeName;
 	}
 	else
 	{
@@ -1642,7 +1583,11 @@ FRigVMExprAST* FRigVMParserAST::TraversePin(const FRigVMASTProxy& InPinProxy, FR
 		}
 	}
 
-	if(Pin->IsTraitPin())
+	if (Pin->IsDefinedAsInputVariable())
+	{
+		PinExpr = MakeExpr<FRigVMVarExprAST>(FRigVMExprAST::EType::Var, InPinProxy);
+	}
+	else if(Pin->IsTraitPin())
 	{
 		PinExpr = MakeExpr<FRigVMVarExprAST>(FRigVMExprAST::EType::Var, InPinProxy);
 
@@ -1709,11 +1654,6 @@ FRigVMExprAST* FRigVMParserAST::TraversePin(const FRigVMASTProxy& InPinProxy, FR
 		return PinExpr;
 	}
 
-	if (PinExpr->IsA(FRigVMExprAST::ExternalVar))
-	{
-		return PinExpr;
-	}
-
 	if ((Pin->GetDirection() == ERigVMPinDirection::IO ||
 		Pin->GetDirection() == ERigVMPinDirection::Input)
 		&& !Pin->IsExecuteContext()
@@ -1944,7 +1884,26 @@ FRigVMExprAST* FRigVMParserAST::TraverseLink(int32 InLinkIndex, FRigVMExprAST* I
 			{
 				SourceRootProxy = FRigVMASTProxy::MakeFromUObject(SourceRootPin);
 			}
-			FRigVMVarExprAST* SourcePinExpr = MakeExpr<FRigVMVarExprAST>(FRigVMExprAST::EType::Var, SourceRootProxy);
+			
+			FRigVMVarExprAST* SourcePinExpr = nullptr;
+			if (NodeExpr->IsNode())
+			{
+				SourcePinExpr = const_cast<FRigVMVarExprAST*>(NodeExpr->To<FRigVMNodeExprAST>()->FindVarWithPinName(SourceRootPin->GetFName()));
+
+				if (SourcePinExpr)
+				{
+					if (const FRigVMVarExprAST* ClonedSourcePinExpr = SourcePinExpr->GetVariableCloneExpr())
+					{
+						SourcePinExpr = const_cast<FRigVMVarExprAST*>(ClonedSourcePinExpr);
+					}
+				}
+			}
+
+			if (!SourcePinExpr)
+			{
+				SourcePinExpr = MakeExpr<FRigVMVarExprAST>(FRigVMExprAST::EType::Var, SourceRootProxy);
+			}
+
 			AssignExpr->ReplaceChild(NodeExpr, SourcePinExpr);
 			return AssignExpr;
 		}			
@@ -1953,6 +1912,111 @@ FRigVMExprAST* FRigVMParserAST::TraverseLink(int32 InLinkIndex, FRigVMExprAST* I
 	return AssignExpr;
 }
 
+bool FRigVMParserAST::TraverseNodeExpressionForRequiredVariableClone(FRigVMNodeExprAST* InNodeExpr)
+{
+	const URigVMNode* Node = InNodeExpr->GetNode();
+	if (!Node)
+	{
+		return false;
+	}
+	
+	const URigVMVariableNode* SourceVariableNode = Cast<URigVMVariableNode>(Node);
+	const URigVMFunctionEntryNode* SourceEntryNode = Cast<URigVMFunctionEntryNode>(Node);
+	if (SourceVariableNode || SourceEntryNode)
+	{
+		TArray<TTuple<URigVMPin*, FName>> PinsToCheck;
+		if (SourceVariableNode && SourceVariableNode->IsGetter())
+		{
+			PinsToCheck.Emplace(SourceVariableNode->GetValuePin(), SourceVariableNode->GetVariableName());
+		}
+		else if (SourceEntryNode)
+		{
+			for (URigVMPin* Pin : SourceEntryNode->GetPins())
+			{
+				if (Pin->GetDirection() != ERigVMPinDirection::Output)
+				{
+					continue;
+				}
+				if (!Pin->IsDefinedAsInputVariable())
+				{
+					continue;
+				}
+				PinsToCheck.Emplace(Pin, Pin->GetFName());
+			}
+		}
+
+		for (int32 PinToCheckIndex = 0; PinToCheckIndex < PinsToCheck.Num(); PinToCheckIndex++)
+		{
+			URigVMPin* PinToCheck = PinsToCheck[PinToCheckIndex].Get<0>();
+			const FName& VariableName = PinsToCheck[PinToCheckIndex].Get<1>();
+
+			bool bRequiresCloningVariableValue = false;
+			FRigVMASTProxy PinProxyToCheck = InNodeExpr->GetProxy().GetSibling(PinToCheck);
+
+			// for variable / entry nodes we need to deal with a specific case where a sub pin
+			// of the pin is wired to a variable setter node with the same variable
+			// on a different sub pin. so for example Get MyVector.X -> Set MyVector.Y.
+			// In that case we need to inject a clone of the variable in the getter.
+			// Since arrays don't have sub pins on the getter this can only happen to non-arrays.
+			const TArray<int32>& LinkIndices = GetTargetLinkIndices(PinProxyToCheck, true);
+			for(const int32 LinkIndex : LinkIndices)
+			{
+				const FRigVMASTLinkDescription& Link = Links[LinkIndex];
+				const URigVMPin* SourcePin = Link.SourceProxy.GetSubjectChecked<URigVMPin>();
+				const URigVMPin* TargetPin = Link.TargetProxy.GetSubjectChecked<URigVMPin>();
+				if (TargetPin && !Link.SegmentPath.IsEmpty())
+				{
+					TargetPin = TargetPin->FindSubPin(Link.SegmentPath);
+				}
+				if (!SourcePin || !TargetPin)
+				{
+					continue;
+				}
+				if (TargetPin->IsRootPin())
+				{
+					continue;
+				}
+				const URigVMVariableNode* TargetVariableNode = Cast<URigVMVariableNode>(TargetPin->GetNode());
+				if (!TargetVariableNode)
+				{
+					continue;
+				}
+				if (VariableName != TargetVariableNode->GetVariableName())
+				{
+					continue;
+				}
+
+				if (SourcePin->GetSegmentPath() != TargetPin->GetSegmentPath())
+				{
+					bRequiresCloningVariableValue = true;
+					break;
+				}
+			}
+
+			if (bRequiresCloningVariableValue)
+			{
+				if (const FRigVMVarExprAST* ConstVariableVarExpr = InNodeExpr->FindVarWithPinName(PinToCheck->GetFName()))
+				{
+					FRigVMVarExprAST* VariableVarExpr = const_cast<FRigVMVarExprAST*>(ConstVariableVarExpr);
+			
+					FRigVMASTProxy PinProxy = VariableVarExpr->GetProxy();
+					FRigVMVarExprAST* ClonedVariableVarExpr = MakeExpr<FRigVMVarExprAST>(FRigVMExprAST::EType::Var, PinProxy);
+					ClonedVariableVarExpr->Name = *FString::Printf(TEXT("%s::Clone"), *VariableVarExpr->Name.ToString());
+					const FRigVMASTLinkDescription CloneLink(PinProxy, PinProxy, FString());
+					FRigVMExprAST* PinCopyExpr = MakeExpr<FRigVMCopyExprAST>(CloneLink);
+					PinCopyExpr->Name = *FString::Printf(TEXT("Cloning Variable %s"), *VariableVarExpr->Name.ToString());
+					PinCopyExpr->AddParent(ClonedVariableVarExpr);
+					VariableVarExpr->AddParent(PinCopyExpr);
+					VariableVarExpr->ClonedVarExpr = ClonedVariableVarExpr;
+					InNodeExpr->ReplaceChild(VariableVarExpr, ClonedVariableVarExpr);
+				}
+			}
+		}
+	}
+
+	return true;
+}
+
 void FRigVMParserAST::FoldEntries()
 {
 	DECLARE_SCOPE_HIERARCHICAL_COUNTER_FUNC()
@@ -2239,6 +2303,12 @@ void FRigVMParserAST::FoldAssignments()
 			continue;
 		}
 
+		// If the assign is between external variable pins, don't fold it
+		if (TargetPin->GetRootPin()->IsDefinedAsInputVariable())
+		{
+			continue;
+		}
+
 		FRigVMExprAST* Child = AssignExpr->Children[0];
 		AssignExpr->RemoveParent(Parent);
 		Child->RemoveParent(AssignExpr);
@@ -2734,6 +2804,14 @@ bool FRigVMParserAST::CanLink(URigVMPin* InSourcePin, URigVMPin* InTargetPin, FS
 			return true;
 		}
 
+		if (const URigVMVariableNode* TargetVariableNode = Cast<URigVMVariableNode>(TargetNode))
+		{
+			if (TargetVariableNode->GetValuePin() == InTargetPin->GetRootPin())
+			{
+				return true;
+			}
+		}
+
 		if (LastCycleCheckExpr != SourceExpression && LastCycleCheckExpr != TargetExpression)
 		{
 			PrepareCycleChecking(InSourcePin);
@@ -2950,11 +3028,6 @@ FString FRigVMParserAST::DumpDot() const
 					Label = FString::Printf(TEXT("%s(Literal)"), *InExpr->To<FRigVMLiteralExprAST>()->GetPin()->GetName());
 					break;
 				}
-				case FRigVMExprAST::EType::ExternalVar:
-				{
-					Label = FString::Printf(TEXT("%s(ExternalVar)"), *InExpr->To<FRigVMExternalVarExprAST>()->GetPin()->GetBoundVariableName());
-					break;
-				}
 				case FRigVMExprAST::EType::Var:
 				{
 					if (InExpr->To<FRigVMVarExprAST>()->IsGraphVariable())
@@ -2971,12 +3044,29 @@ FString FRigVMParserAST::DumpDot() const
 					}
 					else
 					{
-						Label = InExpr->To<FRigVMVarExprAST>()->GetPin()->GetPinPath(true);
+						const URigVMPin* Pin = InExpr->To<FRigVMVarExprAST>()->GetPin();
+						check(Pin);
+						if (Pin->IsDefinedAsInputVariable())
+						{
+							if (InExpr->Name.ToString().EndsWith(TEXT("::Clone")))
+							{
+								Label = FString::Printf(TEXT("Cloned Input Variable %s"), *Pin->GetName());
+							}
+							else
+							{
+								Label = FString::Printf(TEXT("Input Variable %s"), *Pin->GetName());
+							}
+						}
+						else
+						{
+							Label = Pin->GetPinPath(true);
+						}
 					}
 						
 					if (InExpr->To<FRigVMVarExprAST>()->IsExecuteContext())
 					{
 						Shape = EVisualGraphShape::House;
+						Label = TEXT("Execute");
 					}
 						
 					break;
@@ -3001,7 +3091,14 @@ FString FRigVMParserAST::DumpDot() const
 				}
 				case FRigVMExprAST::EType::Copy:
 				{
-					Label = TEXT("Copy");
+					if (InExpr->Name.ToString().EndsWith(TEXT("::Clone")))
+					{
+						Label = TEXT("Clone");
+					}
+					else
+					{
+						Label = TEXT("Copy");
+					}
 					break;
 				}
 				case FRigVMExprAST::EType::CachedValue:
diff --git a/Engine/Plugins/Runtime/RigVM/Source/RigVMDeveloper/Private/RigVMCompiler/RigVMCodeGenerator.cpp b/Engine/Plugins/Runtime/RigVM/Source/RigVMDeveloper/Private/RigVMCompiler/RigVMCodeGenerator.cpp
--- a/Engine/Plugins/Runtime/RigVM/Source/RigVMDeveloper/Private/RigVMCompiler/RigVMCodeGenerator.cpp
+++ b/Engine/Plugins/Runtime/RigVM/Source/RigVMDeveloper/Private/RigVMCompiler/RigVMCodeGenerator.cpp
@@ -160,7 +160,7 @@ FString FRigVMCodeGenerator::DumpExternalVariables(const FRigVMExtendedExecuteCo
 		}
 		else
 		{
-			Lines.Add(Format(RigVM_UpdateExternalVariableFormat, *OperandName, *ExternalVarCPPType, *ExternalVariable.Name.ToString(), *ExternalVarCPPType));
+			Lines.Add(Format(RigVM_UpdateExternalVariableFormat, *OperandName, *ExternalVarCPPType, *ExternalVariable.GetName().ToString(), *ExternalVarCPPType));
 		}
 	}
 
@@ -1999,7 +1999,7 @@ FRigVMPropertyDescription FRigVMCodeGenerator::GetPropertyDescForOperand(const F
 	{
 		const TArray<FRigVMExternalVariableDef> ExternalVariables = VM->GetExternalVariableDefs();
 		const FRigVMExternalVariableDef& ExternalVariable = ExternalVariables[InOperand.GetRegisterIndex()];
-		Property = ExternalVariable.Property;
+		Property = ExternalVariable.GetProperty();
 	}
 	else
 	{
@@ -2023,8 +2023,8 @@ FRigVMPropertyDescription FRigVMCodeGenerator::GetPropertyForOperand(const FRigV
 	{
 		const TArray<FRigVMExternalVariable> ExternalVariables = VM->GetExternalVariables(Context);
 		const FRigVMExternalVariable& ExternalVariable = ExternalVariables[InOperand.GetRegisterIndex()];
-		Property = ExternalVariable.Property;
-		Memory = ExternalVariable.Memory;
+		Property = ExternalVariable.GetProperty();
+		Memory = ExternalVariable.GetMemory();
 	}
 	else
 	{
diff --git a/Engine/Plugins/Runtime/RigVM/Source/RigVMDeveloper/Private/RigVMCompiler/RigVMCompiler.cpp b/Engine/Plugins/Runtime/RigVM/Source/RigVMDeveloper/Private/RigVMCompiler/RigVMCompiler.cpp
--- a/Engine/Plugins/Runtime/RigVM/Source/RigVMDeveloper/Private/RigVMCompiler/RigVMCompiler.cpp
+++ b/Engine/Plugins/Runtime/RigVM/Source/RigVMDeveloper/Private/RigVMCompiler/RigVMCompiler.cpp
@@ -248,7 +248,7 @@ const FProperty* FRigVMCompilerWorkData::GetPropertyForOperand(const FRigVMOpera
 		}
 	case ERigVMMemoryType::External:
 		{
-			Property = VM->GetExternalVariableDefs()[InOperand.GetRegisterIndex()].Property;
+			Property = VM->GetExternalVariableDefs()[InOperand.GetRegisterIndex()].GetProperty();
 			if(InOperand.GetRegisterOffset() != INDEX_NONE)
 			{
 				if(!VM->ExternalPropertyPaths.IsValidIndex(InOperand.GetRegisterOffset()))
@@ -489,13 +489,14 @@ bool URigVMCompiler::Compile(const FRigVMCompileSettings& InSettings, TArray<URi
 			if (FunctionLibrary)
 			{
 				TArray<URigVMLibraryNode*> Functions = FunctionLibrary->GetFunctions();
-				for (URigVMLibraryNode* Function : Functions)
+				for (URigVMLibraryNode* LibraryNode : Functions)
 				{
-					if (FunctionLibrary->IsFunctionPublic(Function->GetFName()))
+					if (FunctionLibrary->IsFunctionPublic(LibraryNode->GetFName()))
 					{
-						if (FRigVMGraphFunctionData* FunctionData = Function->GetFunctionHeader().GetFunctionData())
+						if (FRigVMGraphFunctionData* FunctionData = LibraryNode->GetFunctionHeader().GetFunctionData())
 						{
-							CompileFunction(WorkData.Settings, Function, InController, InExternalVariables, &FunctionData->CompilationData, OutVMContext);
+							TArray<FRigVMExternalVariable> FunctionVariables = LibraryNode->GetExternalVariables();
+							CompileFunction(WorkData.Settings, LibraryNode, InController, FunctionVariables, &FunctionData->CompilationData, OutVMContext);
 						}
 					}
 				}
@@ -557,45 +558,40 @@ bool URigVMCompiler::Compile(const FRigVMCompileSettings& InSettings, TArray<URi
 									IRigVMClientHost* ClientHost = LibraryNode->GetImplementingOuter<IRigVMClientHost>();
 									URigVMController* FunctionController = ClientHost->GetRigVMClient()->GetOrCreateController(LibraryNode->GetLibrary());
 									
-									TArray<FRigVMExternalVariable> FunctionExternalVariables = InExternalVariables;
+									TArray<FRigVMExternalVariable> FunctionVariables = LibraryNode->GetExternalVariables();
 									if(ReferenceNode->RequiresVariableRemapping())
 									{
-										if(FunctionController->GetExternalVariablesDelegate.IsBound())
+										for(FRigVMExternalVariable& FunctionExternalVariable : FunctionVariables)
 										{
-											FunctionExternalVariables = FunctionController->GetExternalVariablesDelegate.Execute(LibraryNode->GetContainedGraph());
-
-											for(FRigVMExternalVariable& FunctionExternalVariable : FunctionExternalVariables)
+											const FName OuterVariableName = ReferenceNode->GetOuterVariableName(FunctionExternalVariable.GetName());
+											if(OuterVariableName.IsNone())
 											{
-												const FName OuterVariableName = ReferenceNode->GetOuterVariableName(FunctionExternalVariable.Name);
-												if(OuterVariableName.IsNone())
-												{
-													const FString VariableRemappingErrorMessage =
-														FString::Printf(TEXT("The function's variable '%s' is not remapped on function reference @@."),
-														*FunctionExternalVariable.Name.ToString());
-													Settings.ASTSettings.Report(EMessageSeverity::Error, ReferenceNode, VariableRemappingErrorMessage);
-													bEncounteredGraphError = true;
-												}
-												else
-												{
-													const FRigVMExternalVariable* OuterExternalVariable = InExternalVariables.FindByPredicate(
-														[OuterVariableName](const FRigVMExternalVariable& ExternalVariable) -> bool
-														{
-															return ExternalVariable.Name.IsEqual(OuterVariableName, ENameCase::CaseSensitive);
-														}
-													);
-
-													if(OuterExternalVariable)
+												const FString VariableRemappingErrorMessage =
+													FString::Printf(TEXT("The function's variable '%s' is not remapped on function reference @@."),
+													*FunctionExternalVariable.GetName().ToString());
+												Settings.ASTSettings.Report(EMessageSeverity::Error, ReferenceNode, VariableRemappingErrorMessage);
+												bEncounteredGraphError = true;
+											}
+											else
+											{
+												const FRigVMExternalVariable* OuterExternalVariable = InExternalVariables.FindByPredicate(
+													[OuterVariableName](const FRigVMExternalVariable& ExternalVariable) -> bool
 													{
-														check(FRigVMRegistry::Get().CanMatchTypes(OuterExternalVariable->GetTypeIndex(), FunctionExternalVariable.GetTypeIndex(), true));
-														FunctionExternalVariable.Property = OuterExternalVariable->Property;
-														FunctionExternalVariable.Memory = OuterExternalVariable->Memory;
+														return ExternalVariable.GetName().IsEqual(OuterVariableName, ENameCase::CaseSensitive);
 													}
+												);
+
+												if(OuterExternalVariable)
+												{
+													check(FRigVMRegistry::Get().CanMatchTypes(OuterExternalVariable->GetTypeIndex(), FunctionExternalVariable.GetTypeIndex(), true));
+													FunctionExternalVariable.SetProperty(OuterExternalVariable->GetProperty());
+													FunctionExternalVariable.SetMemory(const_cast<uint8*>(OuterExternalVariable->GetMemory()));
 												}
 											}
 										}
 									}
-									
-									bSuccessfullCompilation = CompileFunction(WorkData.Settings, LibraryNode, FunctionController, FunctionExternalVariables, &FunctionData->CompilationData, OutVMContext);
+
+									bSuccessfullCompilation = CompileFunction(WorkData.Settings, LibraryNode, FunctionController, FunctionVariables, &FunctionData->CompilationData, OutVMContext);
 								}
 								else
 								{
@@ -781,6 +777,42 @@ bool URigVMCompiler::Compile(const FRigVMCompileSettings& InSettings, TArray<URi
 					FString FunctionHash = FunctionReferenceNode->GetReferencedFunctionHeader().GetHash();
 					if (const FRigVMFunctionCompilationData** CompilationData = CompiledFunctions.Find(FunctionHash))
 					{
+						// todo: this should go into a helper function
+						TArray<FRigVMExternalVariable> AvailableExternalVariables;
+						for (const FRigVMGraphVariableDescription& LocalVariable : VisitedGraph->LocalVariables)
+						{
+							FRigVMExternalVariable::MergeExternalVariable(AvailableExternalVariables, LocalVariable.ToExternalVariable());
+						}
+						const TArray<FRigVMGraphVariableDescription> InputVariables = VisitedGraph->GetInputVariables();
+						for (const FRigVMGraphVariableDescription& InputVariable : InputVariables)
+						{
+							FRigVMExternalVariable::MergeExternalVariable(AvailableExternalVariables, InputVariable.ToExternalVariable());
+						}
+						for (const FRigVMExternalVariable& InExternalVariable : InExternalVariables)
+						{
+							FRigVMExternalVariable::MergeExternalVariable(AvailableExternalVariables, InExternalVariable);
+						}
+						if (CurrentCompilationFunction)
+						{
+							if (const FRigVMGraphFunctionData* FunctionData = CurrentCompilationFunction->GetFunctionHeader().GetFunctionData())
+							{
+								// private functions are allowed to access all external variables
+								if (!FunctionData->IsPublic())
+								{
+									if(URigVMController* LibraryController = InController->GetControllerForGraph(VisitedGraph))
+									{
+										if (LibraryController->GetExternalVariablesDelegate.IsBound())
+										{
+											const TArray<FRigVMExternalVariable> ExternalVariablesFromDelegate = LibraryController->GetExternalVariablesDelegate.Execute(VisitedGraph);
+											for (const FRigVMExternalVariable& ExternalVariableFromDelegate : ExternalVariablesFromDelegate)
+											{
+												FRigVMExternalVariable::MergeExternalVariable(AvailableExternalVariables, ExternalVariableFromDelegate);
+											}
+										}
+									}
+								}
+							}
+						}
 						for (const TPair<int32, FName>& Pair : (*CompilationData)->ExternalRegisterIndexToVariable)
 						{
 							FName OuterName = FunctionReferenceNode->GetOuterVariableName(Pair.Value);
@@ -788,10 +820,10 @@ bool URigVMCompiler::Compile(const FRigVMCompileSettings& InSettings, TArray<URi
 							{
 								OuterName = Pair.Value;
 							}
-						
-							if (!InExternalVariables.ContainsByPredicate([OuterName](const FRigVMExternalVariable& ExternalVariable)
+
+							if (!AvailableExternalVariables.ContainsByPredicate([OuterName](const FRigVMExternalVariable& ExternalVariable)
 								{
-									return ExternalVariable.Name == OuterName;								
+									return ExternalVariable.GetName() == OuterName;								
 								}))
 							{
 								static const FString UnmappedMessage = TEXT("Function referenced in @@ using external variable not found in current rig.");
@@ -840,6 +872,7 @@ bool URigVMCompiler::Compile(const FRigVMCompileSettings& InSettings, TArray<URi
 
 				const FRigVMGraphVariableDescription VariableDescription = VariableNode->GetVariableDescription();
 				const TArray<FRigVMGraphVariableDescription> LocalVariables = VisitedGraph->GetLocalVariables(true);
+				const TArray<FRigVMGraphVariableDescription> InputVariables = VisitedGraph->GetInputVariables();
 
 				bool bFoundVariable = false;
 				for(const FRigVMGraphVariableDescription& LocalVariable : LocalVariables)
@@ -857,19 +890,35 @@ bool URigVMCompiler::Compile(const FRigVMCompileSettings& InSettings, TArray<URi
 					}
 				}
 
+				if (!bFoundVariable)
+				{
+					for(const FRigVMGraphVariableDescription& InputVariable : InputVariables)
+					{
+						if(InputVariable.Name == VariableDescription.Name)
+						{
+							bFoundVariable = true;
+						
+							if(InputVariable.CPPType != VariableDescription.CPPType ||
+								InputVariable.CPPTypeObject != VariableDescription.CPPTypeObject)
+							{
+								Settings.ASTSettings.Report(EMessageSeverity::Error, ModelNode, IllFormedVariableNodeMessage);
+								bEncounteredGraphError = true;
+							}
+						}
+					}
+				}
+
 				// if the variable is not a local variable, let's test against the external variables.
 				if(!bFoundVariable)
 				{
 					const FRigVMExternalVariable ExternalVariable = VariableDescription.ToExternalVariable();
 					for(const FRigVMExternalVariable& InExternalVariable : InExternalVariables)
 					{
-						if(InExternalVariable.Name == ExternalVariable.Name)
+						if(InExternalVariable.GetName() == ExternalVariable.GetName())
 						{
 							bFoundVariable = true;
 							
-							if(InExternalVariable.TypeName != ExternalVariable.TypeName ||
-								InExternalVariable.TypeObject != ExternalVariable.TypeObject ||
-								InExternalVariable.bIsArray != ExternalVariable.bIsArray)
+							if(!InExternalVariable.IsSameType(ExternalVariable))
 							{
 								Settings.ASTSettings.Report(EMessageSeverity::Error, ModelNode, IllFormedVariableNodeMessage);
 								bEncounteredGraphError = true;
@@ -1171,8 +1220,12 @@ bool URigVMCompiler::Compile(const FRigVMCompileSettings& InSettings, TArray<URi
 	
 	for (const FRigVMExternalVariable& ExternalVariable : InExternalVariables)
 	{
+		if (CurrentCompilationFunction && CurrentCompilationFunction->ContainsLocalVariable(ExternalVariable.GetName()))
+		{
+			continue;
+		}
 		FRigVMOperand Operand = OutVM->AddExternalVariable(OutVMContext, ExternalVariable, CurrentCompilationFunction != nullptr);
-		FString Hash = FString::Printf(TEXT("Variable::%s"), *ExternalVariable.Name.ToString());
+		FString Hash = FString::Printf(TEXT("Variable::%s"), *ExternalVariable.GetName().ToString());
 		OutOperands->Add(Hash, Operand);
 	}
 
@@ -1562,14 +1615,20 @@ bool URigVMCompiler::Compile(const FRigVMCompileSettings& InSettings, TArray<URi
 
 		for (const TPair<FString, FRigVMOperand>& Pair : (*WorkData.PinPathToOperand))
 		{
-			static const FString VariablePrefix = TEXT("Variable::");
-			if (Pair.Key.StartsWith(VariablePrefix))
+			if(Pair.Value.GetMemoryType() == ERigVMMemoryType::External)
 			{
-				ensure(Pair.Value.GetMemoryType() == ERigVMMemoryType::External);
-				if (UsedExternalVariableRegisters.Contains(Pair.Value.GetRegisterIndex()))
+				static const FString VariablePrefix = TEXT("Variable::");
+				if (Pair.Key.StartsWith(VariablePrefix))
 				{
-					FString VariableName = Pair.Key.RightChop(VariablePrefix.Len());
-					OutFunctionCompilationData->ExternalRegisterIndexToVariable.Add(Pair.Value.GetRegisterIndex(), *VariableName);
+					if (UsedExternalVariableRegisters.Contains(Pair.Value.GetRegisterIndex()))
+					{
+						FString VariableName = Pair.Key.RightChop(VariablePrefix.Len());
+						if (CurrentCompilationFunction && CurrentCompilationFunction->ContainsLocalVariable(*VariableName))
+						{
+							continue;
+						}
+						OutFunctionCompilationData->ExternalRegisterIndexToVariable.Add(Pair.Value.GetRegisterIndex(), *VariableName);
+					}
 				}
 			}
 		}
@@ -1630,15 +1689,35 @@ bool URigVMCompiler::CompileFunction(const FRigVMCompileSettings& InSettings, co
 	OutFunctionCompilationData->Hash = 0;
 	OutFunctionCompilationData->ByteCode.Reset();
 
-	TArray<FRigVMExternalVariable> ExternalVariables = InExternalVariables;
-	if (InExternalVariables.IsEmpty() && LibraryController->GetExternalVariablesDelegate.IsBound())
+	TArray<FRigVMExternalVariable> FunctionExternalVariables = InExternalVariables;
+	TMap<FString, FRigVMOperand> Operands;
+
+	// let's create a property bag for the external variables so they can be backed up with properties and memory
+	TArray<FRigVMPropertyDescription> FunctionVariableDescriptions;
+	for (int32 FunctionVariableIndex = 0; FunctionVariableIndex < FunctionExternalVariables.Num(); FunctionVariableIndex++)
 	{
-		ExternalVariables = LibraryController->GetExternalVariablesDelegate.Execute(InLibraryNode->GetContainedGraph());
+		const FRigVMExternalVariable& FunctionVariable = FunctionExternalVariables[FunctionVariableIndex];
+		static const FName PropertyNamePrefix = TEXT("Property");
+		const FName PropertyName(PropertyNamePrefix, FunctionVariableIndex + 1);
+		FunctionVariableDescriptions.Emplace(PropertyName, FunctionVariable.GetExtendedCPPType().ToString(), const_cast<UObject*>(FunctionVariable.GetCPPTypeObject()), FString());
+	}
+	FRigVMMemoryStorageStruct FunctionVariablesPropertyPag;
+	FunctionVariablesPropertyPag.AddProperties(FunctionVariableDescriptions, {});
+	const UStruct* FunctionVariablesStruct = FunctionVariablesPropertyPag.GetStruct();
+	FStructOnScope FunctionVariablesInstance(FunctionVariablesStruct);
+	
+	for (int32 FunctionVariableIndex = 0; FunctionVariableIndex < FunctionExternalVariables.Num(); FunctionVariableIndex++)
+	{
+		const FName PropertyName = FunctionVariableDescriptions[FunctionVariableIndex].Name;
+		FRigVMExternalVariable& FunctionVariable = FunctionExternalVariables[FunctionVariableIndex];
+		FunctionVariable.SetProperty(FunctionVariablesStruct->FindPropertyByName(PropertyName));
+		check(FunctionVariable.GetProperty());
+		FunctionVariable.SetMemory(FunctionVariable.GetProperty()->ContainerPtrToValuePtr<uint8>(FunctionVariablesInstance.GetStructMemory()));
+		check(FunctionVariable.GetMemory());
 	}
-	TMap<FString, FRigVMOperand> Operands;
 
 	URigVM* TempVM = NewObject<URigVM>(InLibraryNode->GetContainedGraph(), TEXT("CompilerTemp_VM"));
-	const bool bSuccess = Compile(InSettings, {InLibraryNode->GetContainedGraph()}, LibraryController, TempVM, OutVMContext, ExternalVariables, &Operands, nullptr, OutFunctionCompilationData);
+	const bool bSuccess = Compile(InSettings, {InLibraryNode->GetContainedGraph()}, LibraryController, TempVM, OutVMContext, FunctionExternalVariables, &Operands, nullptr, OutFunctionCompilationData);
 	TempVM->ClearMemory(OutVMContext);
 	TempVM->Rename(nullptr, GetTransientPackage(), REN_DoNotDirty | REN_DontCreateRedirectors | REN_NonTransactional);
 	TempVM->MarkAsGarbage();
@@ -1809,10 +1888,6 @@ bool URigVMCompiler::TraverseExpression(const FRigVMExprAST* InExpr, FRigVMCompi
 		{
 			return TraverseLiteral(InExpr->To<FRigVMLiteralExprAST>(), WorkData);
 		}
-		case FRigVMExprAST::EType::ExternalVar:
-		{
-			return TraverseExternalVar(InExpr->To<FRigVMExternalVarExprAST>(), WorkData);
-		}
 		case FRigVMExprAST::EType::Assign:
 		{
 			return TraverseAssign(InExpr->To<FRigVMAssignExprAST>(), WorkData);
@@ -3028,7 +3103,6 @@ bool URigVMCompiler::TraverseInlineFunction(const FRigVMInlineFunctionExprAST* I
 			const FRigVMOperandArray FunctionOperandArray = FunctionByteCode.GetOperandsForOp(FunctionInstructions[FunctionInstructionIndex]);
 			check(FunctionOperandArray.Num() == OperandArray.Num());
 			
-			int32 OperandsIndex = ByteCode.GetFirstOperandByteIndex(Instruction);
 			for (int32 j=0; j<OperandArray.Num(); ++j)
 			{
 				FRigVMOperand* Operand = const_cast<FRigVMOperand*>(&OperandArray[j]);
@@ -3042,10 +3116,11 @@ bool URigVMCompiler::TraverseInlineFunction(const FRigVMInlineFunctionExprAST* I
 				if (Operand->GetMemoryType() == ERigVMMemoryType::External)
 				{
 					const FName& InnerVariableName = FunctionCompilationData->ExternalRegisterIndexToVariable[Operand->GetRegisterIndex()];
-					FName OuterVariableName = InnerVariableName;
-					if (const FName* VariableRemapped = FunctionReferenceNode->GetVariableMap().Find(InnerVariableName))
+					FString OuterVariablePath = InnerVariableName.ToString();
+					const FString VariableRemapped = FunctionReferenceNode->GetOuterVariablePath(InnerVariableName);
+					if (!VariableRemapped.IsEmpty())
 					{
-						OuterVariableName = *VariableRemapped;
+						OuterVariablePath = VariableRemapped;
 					}
 					else
 					{
@@ -3058,8 +3133,24 @@ bool URigVMCompiler::TraverseInlineFunction(const FRigVMInlineFunctionExprAST* I
 							*GetPackage()->GetPathName(),
 							*LibraryPackagePath);
 					}
-					const FRigVMOperand& OuterOperand = WorkData.PinPathToOperand->FindChecked(FString::Printf(TEXT("Variable::%s"), *OuterVariableName.ToString()));
-					Operand->RegisterIndex = OuterOperand.RegisterIndex;
+
+					// for now variable paths are not supported for mapping.
+					if (OuterVariablePath.Contains(TEXT(".")))
+					{
+						UE_LOG(
+							LogRigVMDeveloper,
+							Error,
+							TEXT("Variable '%s' on function '%s' is mapped to path '%s'. Property paths on variable mappings are not yet supported."),
+							*InnerVariableName.ToString(),
+							*FunctionReferenceNode->GetReferencedFunctionHeader().Name.ToString(),
+							*OuterVariablePath);
+					}
+					else
+					{
+						const FRigVMOperand& OuterOperand = WorkData.PinPathToOperand->FindChecked(FString::Printf(TEXT("Variable::%s"), *OuterVariablePath));
+						Operand->MemoryType = OuterOperand.MemoryType;
+						Operand->RegisterIndex = OuterOperand.RegisterIndex;
+					}
 				}
 				// If Operand is an interface pin: replace the index and memory type
 				else if (const int32 FunctionInterfaceParameterIndex = GetOperandFunctionInterfaceParameterIndex(OperandsPinNames, FunctionCompilationData, *Operand); FunctionInterfaceParameterIndex != INDEX_NONE)
@@ -3403,11 +3494,6 @@ bool URigVMCompiler::TraverseLiteral(const FRigVMVarExprAST* InExpr, FRigVMCompi
 	return TraverseVar(InExpr, WorkData);
 }
 
-bool URigVMCompiler::TraverseExternalVar(const FRigVMExternalVarExprAST* InExpr, FRigVMCompilerWorkData& WorkData)
-{
-	return TraverseVar(InExpr, WorkData);
-}
-
 bool URigVMCompiler::TraverseAssign(const FRigVMAssignExprAST* InExpr, FRigVMCompilerWorkData& WorkData)
 {
 	if (!TraverseChildren(InExpr, WorkData))
@@ -3428,7 +3514,8 @@ bool URigVMCompiler::TraverseAssign(const FRigVMAssignExprAST* InExpr, FRigVMCom
 	{
 		SourceExpr = ChildExpr->To<FRigVMCachedValueExprAST>()->GetVarExpr();
 	}
-	else if (ChildExpr->GetType() == FRigVMExprAST::EType::NoOp)
+	else if (ChildExpr->GetType() == FRigVMExprAST::EType::NoOp ||
+		ChildExpr->GetType() == FRigVMExprAST::EType::Entry)
 	{
 		ensure(ChildExpr->NumChildren() > 0);
 
@@ -3437,12 +3524,35 @@ bool URigVMCompiler::TraverseAssign(const FRigVMAssignExprAST* InExpr, FRigVMCom
 			if (GrandChild->IsA(FRigVMExprAST::EType::Var))
 			{
 				const FRigVMVarExprAST* VarExpr = GrandChild->To<FRigVMVarExprAST>();
-				if (VarExpr->GetPin()->GetName() == TEXT("Value") ||
-					VarExpr->GetPin()->GetName() == TEXT("EnumIndex"))
+				const URigVMPin* Pin = VarExpr->GetPin();
+				if (!Pin)
+				{
+					continue;
+				}
+				const URigVMNode* Node = Pin->GetNode();
+				if (!Node)
+				{
+					continue;
+				}
+				if (Node->IsA<URigVMVariableNode>() && Pin->GetFName() == URigVMVariableNode::ValueName)
+				{
+					SourceExpr = VarExpr;
+					break;
+				}
+				if (Node->IsA<URigVMEnumNode>() && Pin->GetFName() == URigVMEnumNode::EnumIndexName)
 				{
 					SourceExpr = VarExpr;
 					break;
 				}
+				if (Node->IsA<URigVMFunctionEntryNode>() && Pin->GetRootPin()->IsDefinedAsInputVariable())
+				{
+					// the expression on the link equals the pin on the entry node or is a sub-pin thereof
+					if (InExpr->GetSourcePin() == VarExpr->GetPin() || InExpr->GetSourcePin()->IsInOuter(VarExpr->GetPin()))
+					{
+						SourceExpr = VarExpr;
+						break;
+					}
+				}
 			}
 		}
 
@@ -3775,6 +3885,11 @@ void URigVMCompiler::AddCopyOperator(
 
 FString URigVMCompiler::GetPinHashImpl(const URigVMPin* InPin, const FRigVMVarExprAST* InVarExpr, bool bIsDebugValue, const URigVMLibraryNode* FunctionCompiling, const FRigVMASTProxy& InPinProxy)
 {
+	if (!InPin)
+	{
+		return FString();
+	}
+
 	FString Prefix = bIsDebugValue ? TEXT("DebugWatch:") : TEXT("");
 	FString Suffix;
 
@@ -3792,13 +3907,6 @@ FString URigVMCompiler::GetPinHashImpl(const URigVMPin* InPin, const FRigVMVarEx
 
 	if (InVarExpr != nullptr && !bIsDebugValue)
 	{
-		if (InVarExpr->IsA(FRigVMExprAST::ExternalVar))
-		{
-			URigVMPin::FPinOverride PinOverride(InVarExpr->GetProxy(), InVarExpr->GetParser()->GetPinOverrides());
-			FString VariablePath = InPin->GetBoundVariablePath(PinOverride);
-			return FString::Printf(TEXT("%sVariable::%s%s"), *Prefix, *VariablePath, *Suffix);
-		}
-
 		// for IO array pins we'll walk left and use that pin hash instead
 		if(const FRigVMVarExprAST* SourceVarExpr = GetSourceVarExpr(InVarExpr))
 		{
@@ -3811,7 +3919,7 @@ FString URigVMCompiler::GetPinHashImpl(const URigVMPin* InPin, const FRigVMVarEx
 		bIsExecutePin = InPin->IsExecuteContext();
 		bIsLiteral = InVarExpr->GetType() == FRigVMExprAST::EType::Literal;
 
-		bIsVariable = Cast<URigVMVariableNode>(Node) != nullptr || InVarExpr->IsA(FRigVMExprAST::ExternalVar);
+		bIsVariable = Cast<URigVMVariableNode>(Node) != nullptr;
 		bIsFunctionInterfacePin = (Cast<URigVMFunctionEntryNode>(Node) || Cast<URigVMFunctionReturnNode>(Node)) &&
 			Node->GetTypedOuter<URigVMLibraryNode>() == FunctionCompiling;
 
@@ -3920,6 +4028,11 @@ FString URigVMCompiler::GetPinHashImpl(const URigVMPin* InPin, const FRigVMVarEx
 	}
 	else
 	{
+		if (InVarExpr && InPin->IsDefinedAsInputVariable() && !bIsVariableClone)
+		{
+			const FString VariableName = InPin->IsBoundToVariable() ? InPin->GetBoundVariableName() : InPin->GetName();
+			return FString::Printf(TEXT("%sVariable::%s%s"), *Prefix, *VariableName, *Suffix);
+		}
 		if (InVarExpr && !bIsDebugValue)
 		{
 			const FRigVMASTProxy NodeProxy = InVarExpr->GetProxy().GetSibling(Node);
@@ -4298,12 +4411,12 @@ FRigVMOperand URigVMCompiler::FindOrAddRegister(const FRigVMVarExprAST* InVarExp
 		InVarExpr->GetPin()->GetNode()->IsA<URigVMVariableNode>();
 
 	// external variables don't require to add any register.
-	if(bIsVariable && !bIsDebugValue)
+	if((bIsVariable || Pin->IsDefinedAsInputVariable()) && !bIsDebugValue)
 	{
 		const TArray<FRigVMExternalVariableDef>& VMExternalVariableDefs = WorkData.VM->GetExternalVariableDefs();
 		for(int32 ExternalVariableIndex = 0; ExternalVariableIndex < VMExternalVariableDefs.Num(); ExternalVariableIndex++)
 		{
-			const FName& ExternalVariableName = VMExternalVariableDefs[ExternalVariableIndex].Name;
+			const FName& ExternalVariableName = VMExternalVariableDefs[ExternalVariableIndex].GetName();
 			const FString ExternalVariableHash = FString::Printf(TEXT("Variable::%s"), *ExternalVariableName.ToString());
 			if(ExternalVariableHash == Hash)
 			{
diff --git a/Engine/Plugins/Runtime/RigVM/Source/RigVMDeveloper/Private/RigVMDeveloperTypeUtils.cpp b/Engine/Plugins/Runtime/RigVM/Source/RigVMDeveloper/Private/RigVMDeveloperTypeUtils.cpp
--- a/Engine/Plugins/Runtime/RigVM/Source/RigVMDeveloper/Private/RigVMDeveloperTypeUtils.cpp
+++ b/Engine/Plugins/Runtime/RigVM/Source/RigVMDeveloper/Private/RigVMDeveloperTypeUtils.cpp
@@ -91,49 +91,49 @@ FEdGraphPinType RigVMTypeUtils::PinTypeFromExternalVariable(const FRigVMExternal
 	PinType.ResetToDefaults();
 	PinType.PinCategory = NAME_None;
 
-	if (InExternalVariable.TypeName == BoolTypeName)
+	if (InExternalVariable.GetBaseCPPType() == BoolTypeName)
 	{
 		PinType.PinCategory = UEdGraphSchema_K2::PC_Boolean;
 	}
-	else if (InExternalVariable.TypeName == Int32TypeName)
+	else if (InExternalVariable.GetBaseCPPType() == Int32TypeName)
 	{
 		PinType.PinCategory = UEdGraphSchema_K2::PC_Int;
 	}
-	else if (InExternalVariable.TypeName == FloatTypeName)
+	else if (InExternalVariable.GetBaseCPPType() == FloatTypeName)
 	{
 		PinType.PinCategory = UEdGraphSchema_K2::PC_Real;
 		PinType.PinSubCategory = UEdGraphSchema_K2::PC_Float;
 	}
-	else if (InExternalVariable.TypeName == DoubleTypeName)
+	else if (InExternalVariable.GetBaseCPPType() == DoubleTypeName)
 	{
 		PinType.PinCategory = UEdGraphSchema_K2::PC_Real;
 		PinType.PinSubCategory = UEdGraphSchema_K2::PC_Double;
 	}
-	else if (InExternalVariable.TypeName == FNameTypeName)
+	else if (InExternalVariable.GetBaseCPPType() == FNameTypeName)
 	{
 		PinType.PinCategory = UEdGraphSchema_K2::PC_Name;
 	}
-	else if (InExternalVariable.TypeName == FStringTypeName)
+	else if (InExternalVariable.GetBaseCPPType() == FStringTypeName)
 	{
 		PinType.PinCategory = UEdGraphSchema_K2::PC_String;
 	}
-	else if (Cast<UScriptStruct>(InExternalVariable.TypeObject))
+	else if (Cast<UScriptStruct>(InExternalVariable.GetCPPTypeObject()))
 	{
 		PinType.PinCategory = UEdGraphSchema_K2::PC_Struct;
-		PinType.PinSubCategoryObject = InExternalVariable.TypeObject;
+		PinType.PinSubCategoryObject = const_cast<UObject*>(InExternalVariable.GetCPPTypeObject());
 	}
-	else if (Cast<UEnum>(InExternalVariable.TypeObject))
+	else if (Cast<UEnum>(InExternalVariable.GetCPPTypeObject()))
 	{
 		PinType.PinCategory = UEdGraphSchema_K2::PC_Byte;
-		PinType.PinSubCategoryObject = InExternalVariable.TypeObject;
+		PinType.PinSubCategoryObject = const_cast<UObject*>(InExternalVariable.GetCPPTypeObject());
 	}
-	else if (Cast<UClass>(InExternalVariable.TypeObject))
+	else if (Cast<UClass>(InExternalVariable.GetCPPTypeObject()))
 	{
 		PinType.PinCategory = UEdGraphSchema_K2::PC_Object;
-		PinType.PinSubCategoryObject = InExternalVariable.TypeObject;
+		PinType.PinSubCategoryObject = const_cast<UObject*>(InExternalVariable.GetCPPTypeObject());
 	}
 
-	if (InExternalVariable.bIsArray)
+	if (InExternalVariable.IsArray())
 	{
 		PinType.ContainerType = EPinContainerType::Array;
 	}
diff --git a/Engine/Plugins/Runtime/RigVM/Source/RigVMDeveloper/Private/RigVMEditorAsset.cpp b/Engine/Plugins/Runtime/RigVM/Source/RigVMDeveloper/Private/RigVMEditorAsset.cpp
--- a/Engine/Plugins/Runtime/RigVM/Source/RigVMDeveloper/Private/RigVMEditorAsset.cpp
+++ b/Engine/Plugins/Runtime/RigVM/Source/RigVMDeveloper/Private/RigVMEditorAsset.cpp
@@ -865,7 +865,7 @@ void IRigVMEditorAssetInterface::HandleConfigureRigVMController(const FRigVMClie
 				TArray<FRigVMExternalVariable> ExternalVariables = BaseAsset->GetExternalVariables(true);
 				for (FRigVMExternalVariable& ExistingVariable : ExternalVariables)
 				{
-					if (ExistingVariable.Name == InVariable.Name)
+					if (ExistingVariable.GetName() == InVariable.Name)
 					{
 						return FName();
 					}
@@ -873,7 +873,7 @@ void IRigVMEditorAssetInterface::HandleConfigureRigVMController(const FRigVMClie
 
 				FRigVMExternalVariable ExternalVariable = InVariable.ToExternalVariable();
 				return BaseAsset->AddMemberVariable(InVariable.Name,
-					ExternalVariable.TypeObject ? ExternalVariable.TypeObject->GetPathName() : ExternalVariable.TypeName.ToString(),
+					ExternalVariable.GetCPPTypeObject() ? ExternalVariable.GetCPPTypeObject()->GetPathName() : ExternalVariable.GetExtendedCPPType().ToString(),
 					bInIsPublic,
 					bInIsReadOnly,
 					InVariable.DefaultValue);
@@ -2410,7 +2410,7 @@ TArray<FString> IRigVMEditorAssetInterface::GeneratePythonCommands(const FString
 		RigVMTypeUtils::CPPTypeFromExternalVariable(ExternalVariable, CPPType, &CPPTypeObject);
 		if (CPPTypeObject)
 		{
-			if (ExternalVariable.bIsArray)
+			if (ExternalVariable.IsArray())
 			{
 				CPPType = RigVMTypeUtils::ArrayTypeFromBaseType(CPPTypeObject->GetPathName());
 			}
@@ -2421,10 +2421,10 @@ TArray<FString> IRigVMEditorAssetInterface::GeneratePythonCommands(const FString
 		}
 		// FName AddMemberVariable(const FName& InName, const FString& InCPPType, bool bIsPublic = false, bool bIsReadOnly = false, FString InDefaultValue = TEXT(""));
 		InternalCommands.Add(FString::Printf(TEXT("blueprint.add_member_variable('%s', '%s', %s, %s)"),
-					*ExternalVariable.Name.ToString(),
+					*ExternalVariable.GetName().ToString(),
 					*CPPType,
-					ExternalVariable.bIsPublic ? TEXT("True") : TEXT("False"),
-					ExternalVariable.bIsReadOnly ? TEXT("True") : TEXT("False")));	
+					ExternalVariable.IsPublic() ? TEXT("True") : TEXT("False"),
+					ExternalVariable.IsReadOnly() ? TEXT("True") : TEXT("False")));	
 	}
 	
 	// Create graphs
@@ -2683,7 +2683,7 @@ void IRigVMEditorAssetInterface::AddVariableSearchMetaDataInfo(const FName InVar
 {
 	for (const FRigVMExternalVariable& Variable : GetExternalVariables(true))
 	{
-		if (Variable.Name.IsEqual(InVariableName))
+		if (Variable.GetName().IsEqual(InVariableName))
 		{
 			OutTaggedMetaData.Emplace(FRigVMSearchTags::FiB_Name, FText::FromName(InVariableName));
 
@@ -3612,7 +3612,6 @@ void IRigVMEditorAssetInterface::HandleModifiedEvent(ERigVMGraphNotifType InNoti
 				break;
 			}
 			case ERigVMGraphNotifType::PinBoundVariableChanged:
-			case ERigVMGraphNotifType::VariableRemappingChanged:
 			{
 				RequestAutoVMRecompilation();
 				(void)GetObject()->MarkPackageDirty();
@@ -3691,7 +3690,7 @@ void IRigVMEditorAssetInterface::CreateMemberVariablesOnLoad()
 	TArray<FRigVMExternalVariable> Variables = GetExternalVariables(true);
 	for (int32 VariableIndex = 0; VariableIndex < Variables.Num(); VariableIndex++)
 	{
-		AddedMemberVariableMap.Add(Variables[VariableIndex].Name, VariableIndex);
+		AddedMemberVariableMap.Add(Variables[VariableIndex].GetName(), VariableIndex);
 	}
 
 	if (GetRigVMClient()->Num() == 0)
@@ -3961,7 +3960,7 @@ void IRigVMEditorAssetInterface::OnVariableAdded(const FName& InVarName)
     RigVMTypeUtils::CPPTypeFromExternalVariable(ExternalVariable, CPPType, &CPPTypeObject);
 	if (CPPTypeObject)
 	{
-		if (ExternalVariable.bIsArray)
+		if (ExternalVariable.IsArray())
 		{
 			CPPType = RigVMTypeUtils::ArrayTypeFromBaseType(CPPTypeObject->GetPathName());
 		}
@@ -3978,8 +3977,8 @@ void IRigVMEditorAssetInterface::OnVariableAdded(const FName& InVarName)
 		FString::Printf(TEXT("blueprint.add_member_variable('%s', '%s', %s, %s, '%s')"),
 			*InVarName.ToString(),
 			*CPPType,
-			(ExternalVariable.bIsPublic) ? TEXT("False") : TEXT("True"), 
-			(ExternalVariable.bIsReadOnly) ? TEXT("True") : TEXT("False"), 
+			(ExternalVariable.IsPublic()) ? TEXT("False") : TEXT("True"), 
+			(ExternalVariable.IsReadOnly()) ? TEXT("True") : TEXT("False"), 
 			*Variable.DefaultValue)); 
 	
 	BroadcastExternalVariablesChangedEvent();
@@ -4763,8 +4762,8 @@ bool URigVMEditorAsset::OnPasteVariable(const FString& InText)
 		{
 			FScopedTransaction Transaction(FText::Format(LOCTEXT("PasteVariable", "Paste Variable: {0}"), FText::FromName(NewVar.Name)));
 
+			NewVar.Name = FRigVMBlueprintUtils::FindUniqueVariableName(this, NewVar.Name.ToString());
 			FRigVMExternalVariable ExternalVariable = RigVMVariableUtils::ExternalVariableFromRigVMVariableDescription(NewVar);
-			ExternalVariable.Name = FRigVMBlueprintUtils::FindUniqueVariableName(this, ExternalVariable.Name.ToString());
 			FName Result = RuntimeAsset->AddHostMemberVariableFromExternal(ExternalVariable, NewVar.DefaultValue);
 
 			return !Result.IsNone();
diff --git a/Engine/Plugins/Runtime/RigVM/Source/RigVMDeveloper/Private/RigVMModel/Nodes/RigVMCollapseNode.cpp b/Engine/Plugins/Runtime/RigVM/Source/RigVMDeveloper/Private/RigVMModel/Nodes/RigVMCollapseNode.cpp
--- a/Engine/Plugins/Runtime/RigVM/Source/RigVMDeveloper/Private/RigVMModel/Nodes/RigVMCollapseNode.cpp
+++ b/Engine/Plugins/Runtime/RigVM/Source/RigVMDeveloper/Private/RigVMModel/Nodes/RigVMCollapseNode.cpp
@@ -26,6 +26,27 @@ URigVMFunctionLibrary* URigVMCollapseNode::GetLibrary() const
 	return Cast<URigVMFunctionLibrary>(GetOuter());
 }
 
+TArray<FRigVMExternalVariable> URigVMCollapseNode::GetExternalVariables() const
+{
+	TArray<FRigVMExternalVariable> ExternalVariables;
+	if (IsGraphFunctionDefinition())
+	{
+		for (const URigVMPin* Pin : GetPins())
+		{
+			if (Pin->IsDefinedAsInputVariable())
+			{
+				FRigVMExternalVariable::MergeExternalVariable(ExternalVariables, Pin->GetInputVariableDescription().ToExternalVariable());
+			}
+		}
+	}
+	const TArray<FRigVMExternalVariable> SuperExternalVariables = Super::GetExternalVariables();
+	for (const FRigVMExternalVariable& SuperExternalVariable : SuperExternalVariables)
+	{
+		FRigVMExternalVariable::MergeExternalVariable(ExternalVariables, SuperExternalVariable);
+	}
+	return ExternalVariables;
+}
+
 FString URigVMCollapseNode::GetEditorSubGraphName() const
 {
 	return FString::Printf(TEXT("%s_SubGraph"), *GetName());
diff --git a/Engine/Plugins/Runtime/RigVM/Source/RigVMDeveloper/Private/RigVMModel/Nodes/RigVMFunctionInterfaceNode.cpp b/Engine/Plugins/Runtime/RigVM/Source/RigVMDeveloper/Private/RigVMModel/Nodes/RigVMFunctionInterfaceNode.cpp
--- a/Engine/Plugins/Runtime/RigVM/Source/RigVMDeveloper/Private/RigVMModel/Nodes/RigVMFunctionInterfaceNode.cpp
+++ b/Engine/Plugins/Runtime/RigVM/Source/RigVMDeveloper/Private/RigVMModel/Nodes/RigVMFunctionInterfaceNode.cpp
@@ -97,6 +97,12 @@ bool URigVMFunctionInterfaceNode::IsInterfacePinUsed(const FName& InInterfacePin
 				continue;
 			}
 
+			// we found a variable node that sets the interface pin value
+			if (!VariableNode->IsGetter())
+			{
+				return true;
+			}
+
 			// we found a variable node that's bound to our input pin
 			if (VariableNode->GetValuePin()->IsLinked(true))
 			{
diff --git a/Engine/Plugins/Runtime/RigVM/Source/RigVMDeveloper/Private/RigVMModel/Nodes/RigVMFunctionReferenceNode.cpp b/Engine/Plugins/Runtime/RigVM/Source/RigVMDeveloper/Private/RigVMModel/Nodes/RigVMFunctionReferenceNode.cpp
--- a/Engine/Plugins/Runtime/RigVM/Source/RigVMDeveloper/Private/RigVMModel/Nodes/RigVMFunctionReferenceNode.cpp
+++ b/Engine/Plugins/Runtime/RigVM/Source/RigVMDeveloper/Private/RigVMModel/Nodes/RigVMFunctionReferenceNode.cpp
@@ -1,9 +1,11 @@
 // Copyright Epic Games, Inc. All Rights Reserved.
 
 #include "RigVMModel/Nodes/RigVMFunctionReferenceNode.h"
+#include "RigVMEditorAsset.h"
 #include "RigVMModel/RigVMFunctionLibrary.h"
 #include "RigVMCore/RigVMGraphFunctionDefinition.h"
 #include "RigVMCore/RigVMGraphFunctionHost.h"
+#include "Widgets/SRigVMEditorGraphExplorerTreeView.h"
 
 #include UE_INLINE_GENERATED_CPP_BY_NAME(RigVMFunctionReferenceNode)
 
@@ -108,109 +110,142 @@ FString URigVMFunctionReferenceNode::GetNodeKeywords() const
 	return ReferencedFunctionHeader.Keywords;
 }
 
-bool URigVMFunctionReferenceNode::RequiresVariableRemapping() const
+bool URigVMFunctionReferenceNode::ReferencesFunctionInSamePackage() const
 {
-	TArray<FRigVMExternalVariable> InnerVariables;
-	return RequiresVariableRemappingInternal(InnerVariables);
-}
-
-bool URigVMFunctionReferenceNode::RequiresVariableRemappingInternal(TArray<FRigVMExternalVariable>& InnerVariables) const
-{
-	bool bHostedInDifferencePackage = false;
-	
-	FString ThisPacakgePath = GetPackage()->GetPathName();
+	const FString ThisPackagePath = GetPackage()->GetPathName();
 	
 	// avoid function reference related validation for temp assets, a temp asset may get generated during
 	// certain content validation process. It is usually just a simple file-level copy of the source asset
 	// so these references are usually not fixed-up properly. Thus, it is meaningless to validate them.
-	if (ThisPacakgePath.StartsWith(TEXT("/Temp/")))
+	if (ThisPackagePath.StartsWith(TEXT("/Temp/")))
 	{
 		return false;
 	}
 	
-	FRigVMGraphFunctionIdentifier LibraryPointer = ReferencedFunctionHeader.LibraryPointer;
+	const FRigVMGraphFunctionIdentifier& LibraryPointer = ReferencedFunctionHeader.LibraryPointer;
 	const FString& LibraryPackagePath = LibraryPointer.GetNodeSoftPath().GetLongPackageName();
-	bHostedInDifferencePackage = LibraryPackagePath != ThisPacakgePath;
-	if(bHostedInDifferencePackage)
+	return LibraryPackagePath == ThisPackagePath;
+}
+
+bool URigVMFunctionReferenceNode::RequiresVariableRemapping() const
+{
+	if (GetPins().ContainsByPredicate([](const URigVMPin* Pin) -> bool
 	{
-		InnerVariables = GetExternalVariables(false);
-		if(InnerVariables.Num() == 0)
-		{
-			return false;
-		}
+		return Pin->IsDefinedAsInputVariable();
+	}))
+	{
+		return true;
 	}
-
-	return bHostedInDifferencePackage;
+	if (ReferencesFunctionInSamePackage())
+	{
+		return false;
+	}
+	return !GetExternalVariables().IsEmpty();
 }
 
 bool URigVMFunctionReferenceNode::IsFullyRemapped() const
 {
-	TArray<FRigVMExternalVariable> InnerVariables;
-	if(!RequiresVariableRemappingInternal(InnerVariables))
+	for (const URigVMPin* Pin : GetPins())
 	{
-		return true;
+		if (Pin->IsDefinedAsInputVariable())
+		{
+			if (GetOuterVariableName(Pin->GetFName()).IsNone())
+			{
+				return false;
+			}
+		}
 	}
 
-	for(const FRigVMExternalVariable& InnerVariable : InnerVariables)
+	if (!ReferencesFunctionInSamePackage())
 	{
-		const FName InnerVariableName = InnerVariable.Name;
-		const FName* OuterVariableName = VariableMap.Find(InnerVariableName);
-		if(OuterVariableName == nullptr)
+		TArray<FRigVMExternalVariable> ExternalVariables = GetExternalVariables();
+		for(const FRigVMExternalVariable& ExternalVariable : ExternalVariables)
 		{
-			return false;
+			if (GetOuterVariableName(ExternalVariable.GetName()).IsNone())
+			{
+				return false;
+			}
 		}
-		check(!OuterVariableName->IsNone());
 	}
-
 	return true;
 }
 
 TArray<FRigVMExternalVariable> URigVMFunctionReferenceNode::GetExternalVariables() const
 {
-	return GetExternalVariables(true);
+	if (const FRigVMGraphFunctionData* FunctionData = GetReferencedFunctionData())
+	{
+		return FunctionData->Header.ExternalVariables;
+	}
+	return Super::GetExternalVariables();
 }
 
-TArray<FRigVMExternalVariable> URigVMFunctionReferenceNode::GetExternalVariables(bool bRemapped) const
+const TMap<FName, FName> URigVMFunctionReferenceNode::GetVariableMap() const
 {
-	TArray<FRigVMExternalVariable> Variables;
-	
-	if(!bRemapped)
+	TMap<FName, FName> Result =  VariableMap_DEPRECATED;
+	for (const URigVMPin* InterfacePin : GetPins())
 	{
-		if (const FRigVMGraphFunctionData* FunctionData = GetReferencedFunctionData(false))
+		if (!InterfacePin->IsDefinedAsInputVariable())
+		{
+			continue;
+		}
+		const FString BoundVariable = InterfacePin->GetBoundVariableName();
+		if (BoundVariable.IsEmpty())
 		{
-			Variables = FunctionData->Header.ExternalVariables;
+			Result.FindOrAdd(InterfacePin->GetFName(), NAME_None);
 		}
 		else
 		{
-			Variables = GetReferencedFunctionHeader().ExternalVariables;
+			Result.FindOrAdd(InterfacePin->GetFName(), *BoundVariable);
 		}
 	}
-	else
+	return Result;
+}
+
+FName URigVMFunctionReferenceNode::GetOuterVariableName(const FName& InInnerVariableName) const
+{
+	FString VariablePath = GetOuterVariablePath(InInnerVariableName);
+	if (VariablePath.IsEmpty())
+	{
+		return NAME_None;
+	}
+	FString VariableName = VariablePath, VariableSubPath;
+	(void)VariablePath.Split(TEXT("."), &VariableName, &VariableSubPath);
+	return *VariableName;
+}
+
+FString URigVMFunctionReferenceNode::GetOuterVariablePath(const FName& InInnerVariableName) const
+{
+	if (const URigVMPin* InterfacePin = FindRootPinByName(InInnerVariableName))
 	{
-		if(RequiresVariableRemappingInternal(Variables))
+		if (InterfacePin->IsDefinedAsInputVariable())
 		{
-			for(FRigVMExternalVariable& Variable : Variables)
+			const FString BoundVariablePath = InterfacePin->GetBoundVariablePath();
+			if (!BoundVariablePath.IsEmpty())
 			{
-				const FName* OuterVariableName = VariableMap.Find(Variable.Name);
-				if(OuterVariableName != nullptr)
-				{
-					check(!OuterVariableName->IsNone());
-					Variable.Name = *OuterVariableName;
-				}
+				return BoundVariablePath;
 			}
 		}
 	}
+	if(const FName* OuterVariableName = VariableMap_DEPRECATED.Find(InInnerVariableName))
+	{
+		return OuterVariableName->ToString();
+	}
 	
-	return Variables; 
-}
-
-FName URigVMFunctionReferenceNode::GetOuterVariableName(const FName& InInnerVariableName) const
-{
-	if(const FName* OuterVariableName = VariableMap.Find(InInnerVariableName))
+	// we may be on a function reference within the same asset.
+	// in that case we allow to reference asset wide variables.
+	if (ReferencesFunctionInSamePackage())
 	{
-		return *OuterVariableName;
+		const TArray<FRigVMExternalVariable> AssetVariables = GetAssetExternalVariables();
+		if (AssetVariables.ContainsByPredicate(
+		[InInnerVariableName](const FRigVMExternalVariable& Variable) -> bool
+			{
+				return Variable.GetName() == InInnerVariableName;
+			}))
+		{
+			return InInnerVariableName.ToString();
+		}
 	}
-	return NAME_None;
+	return FString();
 }
 
 uint32 URigVMFunctionReferenceNode::GetStructureHash() const
@@ -239,10 +274,11 @@ uint32 URigVMFunctionReferenceNode::GetStructureHash() const
 		}
 	}
 
-	for(const FRigVMExternalVariable& ExternalVariable : ReferencedFunctionHeader.ExternalVariables)
+	const TArray<FRigVMExternalVariable> ExternalVariables = ReferencedFunctionHeader.GetExternalVariables();
+	for(const FRigVMExternalVariable& ExternalVariable : ExternalVariables)
 	{
-		Hash = HashCombine(Hash, GetTypeHash(ExternalVariable.Name.ToString()));
-		const TRigVMTypeIndex TypeIndex = Registry->GetTypeIndexFromCPPType_NoLock(ExternalVariable.TypeName.ToString());
+		Hash = HashCombine(Hash, GetTypeHash(ExternalVariable.GetName().ToString()));
+		const TRigVMTypeIndex TypeIndex = Registry->GetTypeIndexFromCPPType_NoLock(ExternalVariable.GetExtendedCPPType().ToString());
 		Hash = HashCombine(Hash, Registry->GetHashForType_NoLock(TypeIndex));
 	}
 
@@ -253,6 +289,8 @@ void URigVMFunctionReferenceNode::UpdateFunctionHeaderFromHost()
 {
 	if (const FRigVMGraphFunctionData* Data = GetReferencedFunctionData())
 	{
+		const_cast<FRigVMGraphFunctionData*>(Data)->PatchExternalVariablesToArguments();
+		
 		ReferencedFunctionHeader = Data->Header;
 		InvalidateCache();
 	}
@@ -299,6 +337,34 @@ FString URigVMFunctionReferenceNode::GetOriginalDefaultValueForRootPin(const URi
 	return Super::GetOriginalDefaultValueForRootPin(InRootPin);
 }
 
+TArray<FRigVMExternalVariable> URigVMFunctionReferenceNode::GetAssetExternalVariables() const
+{
+	if (const IRigVMAssetInterface* Asset = GetImplementingOuter<IRigVMAssetInterface>())
+	{
+		const TArray<FRigVMGraphVariableDescription> VariableDescriptions = Asset->GetAssetVariables();
+		TArray<FRigVMExternalVariable> ExternalVariables;
+		for (const FRigVMGraphVariableDescription& VariableDescription : VariableDescriptions)
+		{
+			FRigVMExternalVariable::MergeExternalVariable(ExternalVariables, VariableDescription.ToExternalVariable());
+		}
+		return ExternalVariables;
+	}
+	if (IRigVMClientHost* ClientHost = GetImplementingOuter<IRigVMClientHost>())
+	{
+		if (FRigVMClient* Client = ClientHost->GetRigVMClient())
+		{
+			if (const URigVMController* Controller = Client->GetOrCreateController(GetGraph()))
+			{
+				if (Controller->GetExternalVariablesDelegate.IsBound())
+				{
+					return Controller->GetExternalVariablesDelegate.Execute(GetGraph());
+				}
+			}
+		}
+	}
+	return {};
+}
+
 FText URigVMFunctionReferenceNode::GetToolTipTextForPin(const URigVMPin* InPin) const
 {
 	check(InPin);
