RigVM: Functions to support input variables

diff --git a/Engine/Plugins/Runtime/RigVM/Source/RigVMDeveloper/Private/RigVMModel/Nodes/RigVMLibraryNode.cpp b/Engine/Plugins/Runtime/RigVM/Source/RigVMDeveloper/Private/RigVMModel/Nodes/RigVMLibraryNode.cpp
--- a/Engine/Plugins/Runtime/RigVM/Source/RigVMDeveloper/Private/RigVMModel/Nodes/RigVMLibraryNode.cpp
+++ b/Engine/Plugins/Runtime/RigVM/Source/RigVMDeveloper/Private/RigVMModel/Nodes/RigVMLibraryNode.cpp
@@ -295,6 +295,7 @@ FRigVMGraphFunctionHeader URigVMLibraryNode::GetFunctionHeader(IRigVMGraphFuncti
     	Arg.CPPTypeObject = Pin->GetCPPTypeObject();
     	Arg.DefaultValue = Pin->GetDefaultValue();
     	Arg.bIsConst = Pin->IsDefinedAsConstant();
+    	Arg.bIsInputVariable = Pin->IsDefinedAsInputVariable();
 
     	TArray<URigVMPin*> PinsToProcess;
     	PinsToProcess.Add(Pin);
@@ -308,7 +309,7 @@ FRigVMGraphFunctionHeader URigVMLibraryNode::GetFunctionHeader(IRigVMGraphFuncti
     		
     		PinsToProcess.Append(PinsToProcess[i]->GetSubPins());
     	}
-    	
+
     	Header.Arguments.Add(Arg);
     }
 
@@ -324,7 +325,7 @@ FRigVMGraphFunctionHeader URigVMLibraryNode::GetFunctionHeader(IRigVMGraphFuncti
 		}
 		else
 		{
-			Header.ExternalVariables = GetExternalVariables();
+			Header.ExternalVariables.Reset();
 		}
 	}
 
@@ -353,4 +354,20 @@ TArray<FRigVMVariantRef> URigVMLibraryNode::GetMatchingVariants() const
 	return TArray<FRigVMVariantRef>();
 }
 
+bool URigVMLibraryNode::ContainsLocalVariable(const FName& InVariableName) const
+{
+	const URigVMGraph* ContainedGraph = GetContainedGraph();
+	if (!ContainedGraph)
+	{
+		return false;
+	}
+
+	return ContainedGraph->LocalVariables.ContainsByPredicate(
+		[InVariableName](const FRigVMGraphVariableDescription& LocalVariable) -> bool
+		{
+			return InVariableName == LocalVariable.Name;
+		}
+	);
+}
+
 
diff --git a/Engine/Plugins/Runtime/RigVM/Source/RigVMDeveloper/Private/RigVMModel/Nodes/RigVMUnitNode.cpp b/Engine/Plugins/Runtime/RigVM/Source/RigVMDeveloper/Private/RigVMModel/Nodes/RigVMUnitNode.cpp
--- a/Engine/Plugins/Runtime/RigVM/Source/RigVMDeveloper/Private/RigVMModel/Nodes/RigVMUnitNode.cpp
+++ b/Engine/Plugins/Runtime/RigVM/Source/RigVMDeveloper/Private/RigVMModel/Nodes/RigVMUnitNode.cpp
@@ -583,9 +583,9 @@ void URigVMUnitNode::EnumeratePropertiesOnHostAndStructInstance(URigVMHost* InHo
 					TArray<FRigVMExternalVariable> ExternalVariables = InHost->GetExternalVariables();
 					if(ExternalVariables.IsValidIndex(Operand.GetRegisterIndex()))
 					{
-						const FRigVMExternalVariable& ExternalVariable = ExternalVariables[Operand.GetRegisterIndex()];
-						HostMemory = ExternalVariable.Memory; 
-						HostProperty = ExternalVariable.Property;
+						FRigVMExternalVariable& ExternalVariable = ExternalVariables[Operand.GetRegisterIndex()];
+						HostMemory = ExternalVariable.GetMemory(); 
+						HostProperty = ExternalVariable.GetProperty();
 					}
 					break;
 				}
diff --git a/Engine/Plugins/Runtime/RigVM/Source/RigVMDeveloper/Private/RigVMModel/Nodes/RigVMVariableNode.cpp b/Engine/Plugins/Runtime/RigVM/Source/RigVMDeveloper/Private/RigVMModel/Nodes/RigVMVariableNode.cpp
--- a/Engine/Plugins/Runtime/RigVM/Source/RigVMDeveloper/Private/RigVMModel/Nodes/RigVMVariableNode.cpp
+++ b/Engine/Plugins/Runtime/RigVM/Source/RigVMDeveloper/Private/RigVMModel/Nodes/RigVMVariableNode.cpp
@@ -41,7 +41,7 @@ bool URigVMVariableNode::IsGetter() const
 
 bool URigVMVariableNode::IsExternalVariable() const
 {
-	return !IsLocalVariable() && !IsInputArgument();
+	return !IsLocalVariable() && !IsInputArgument() && !IsInputVariable();
 }
 
 bool URigVMVariableNode::IsLocalVariable() const
@@ -70,12 +70,26 @@ bool URigVMVariableNode::IsInputArgument() const
 	{
 		if (URigVMFunctionEntryNode* EntryNode = Graph->GetEntryNode())
 		{
-			for (const URigVMPin* Pin : EntryNode->GetPins())
+			if (const URigVMPin* Pin = EntryNode->FindRootPinByName(CurrentVariableName))
 			{
-				if(Pin->GetFName() == CurrentVariableName)
-				{
-					return true;
-				}
+				return !Pin->IsDefinedAsInputVariable();
+			}
+		}
+	}
+	return false;
+}
+
+bool URigVMVariableNode::IsInputVariable() const
+{
+	const FName CurrentVariableName = GetVariableName();
+	
+	if(URigVMGraph* Graph = GetGraph())
+	{
+		if (URigVMFunctionEntryNode* EntryNode = Graph->GetEntryNode())
+		{
+			if (const URigVMPin* Pin = EntryNode->FindRootPinByName(CurrentVariableName))
+			{
+				return Pin->IsDefinedAsInputVariable();
 			}
 		}
 	}
@@ -158,6 +172,10 @@ FString URigVMVariableNode::GetNodeSubTitle() const
 	{
 		SubTitle = TEXT("Input parameter");
 	}
+	else if (IsInputVariable())
+	{
+		SubTitle = TEXT("Input Variable");
+	}
 
 	return SubTitle;
 }
diff --git a/Engine/Plugins/Runtime/RigVM/Source/RigVMDeveloper/Private/RigVMModel/RigVMClient.cpp b/Engine/Plugins/Runtime/RigVM/Source/RigVMDeveloper/Private/RigVMModel/RigVMClient.cpp
--- a/Engine/Plugins/Runtime/RigVM/Source/RigVMDeveloper/Private/RigVMModel/RigVMClient.cpp
+++ b/Engine/Plugins/Runtime/RigVM/Source/RigVMDeveloper/Private/RigVMModel/RigVMClient.cpp
@@ -1345,6 +1345,7 @@ FRigVMClientPatchResult FRigVMClient::PatchModelsOnLoad()
 			Result.Merge(Controller->PatchExecutePins());
 			Result.Merge(Controller->PatchLazyPins());
 			Result.Merge(Controller->PatchUserDefinedStructPinNames());
+			Result.Merge(Controller->PatchExternalVariablesToArguments());
 		
 			if (URigVMCollapseNode* CollapseNode = Model->GetTypedOuter<URigVMCollapseNode>())
 			{
@@ -1634,9 +1635,9 @@ void FRigVMClient::HandleGraphModifiedEvent(ERigVMGraphNotifType InNotifType, UR
 		}
 		case ERigVMGraphNotifType::NodeRemoved: // A node has been removed from the graph (Subject == URigVMNode)
 		{		
-			if(InSubject->GetOuter()->IsA<URigVMFunctionLibrary>())
+			if (URigVMCollapseNode* CollapseNode = Cast<URigVMCollapseNode>(InSubject))
 			{
-				if (URigVMCollapseNode* CollapseNode = Cast<URigVMCollapseNode>(InSubject))
+				if(CollapseNode->IsGraphFunctionDefinition())
 				{
 					if (GetOuter()->Implements<URigVMClientHost>())
 					{
@@ -1647,7 +1648,7 @@ void FRigVMClient::HandleGraphModifiedEvent(ERigVMGraphNotifType InNotifType, UR
 					}
 				}
 			}
-			// A node was added into the contained graph of a function
+			// A node was removed from the contained graph of a function
 			else if(URigVMLibraryNode* LibraryNode = Cast<URigVMNode>(InSubject)->FindFunctionForNode())
 			{
 				DirtyGraphFunctionCompilationData(LibraryNode);
@@ -1674,7 +1675,6 @@ void FRigVMClient::HandleGraphModifiedEvent(ERigVMGraphNotifType InNotifType, UR
 		case ERigVMGraphNotifType::VariableAdded: // A variable has been added (Subject == URigVMVariableNode)
 		case ERigVMGraphNotifType::VariableRemoved: // A variable has been removed (Subject == URigVMVariableNode)
 		case ERigVMGraphNotifType::VariableRenamed: // A variable has been renamed (Subject == URigVMVariableNode)
-		case ERigVMGraphNotifType::VariableRemappingChanged: // A function reference node's remapping has changed (Subject == URigVMFunctionReferenceNode)
 		{
 			if(URigVMLibraryNode* LibraryNode = Cast<URigVMNode>(InSubject)->FindFunctionForNode())
 			{
@@ -1685,9 +1685,9 @@ void FRigVMClient::HandleGraphModifiedEvent(ERigVMGraphNotifType InNotifType, UR
 		}
 		case ERigVMGraphNotifType::FunctionRenamed: // A node has been renamed in the graph (Subject == URigVMNode)
 		{
-			if(InSubject->GetOuter()->IsA<URigVMFunctionLibrary>())
+			if (URigVMCollapseNode* CollapseNode = Cast<URigVMCollapseNode>(InSubject))
 			{
-				if (URigVMCollapseNode* CollapseNode = Cast<URigVMCollapseNode>(InSubject))
+				if (CollapseNode->IsGraphFunctionDefinition())
 				{
 					URigVMBuildData* BuildData = URigVMBuildData::Get();
 					if (IRigVMGraphFunctionHost* Host = Cast<IRigVMGraphFunctionHost>(CollapseNode->GetFunctionIdentifier().HostObject.ResolveObject()))
@@ -1736,9 +1736,9 @@ void FRigVMClient::HandleGraphModifiedEvent(ERigVMGraphNotifType InNotifType, UR
 		case ERigVMGraphNotifType::NodeTitleChanged: // A node's title has changed (Subject == URigVMNode)
 		case ERigVMGraphNotifType::VariantTagsChanged: // The tags in the header of this function variant have changed (Subject == URigVMLibraryNode)
 		{
-			if(InSubject->GetOuter()->IsA<URigVMFunctionLibrary>())
+			if (URigVMCollapseNode* CollapseNode = Cast<URigVMCollapseNode>(InSubject))
 			{
-				if (URigVMCollapseNode* CollapseNode = Cast<URigVMCollapseNode>(InSubject))
+				if(CollapseNode->IsGraphFunctionDefinition())
 				{
 					UpdateGraphFunctionData(CollapseNode);
 				}
@@ -1771,11 +1771,12 @@ void FRigVMClient::HandleGraphModifiedEvent(ERigVMGraphNotifType InNotifType, UR
 				{
 					DirtyGraphFunctionCompilationData(LibraryNode);
 				}
-				if(Node->GetOuter()->IsA<URigVMFunctionLibrary>())
+				if (const URigVMCollapseNode* CollapseNode = Cast<URigVMCollapseNode>(Node))
 				{
-					if (const URigVMCollapseNode* CollapseNode = Cast<URigVMCollapseNode>(Node))
+					if(CollapseNode->IsGraphFunctionDefinition())
 					{
 						UpdateGraphFunctionData(CollapseNode);
+						UpdateExternalVariablesForFunction(CollapseNode);
 					}
 				}
 			}
@@ -1783,14 +1784,11 @@ void FRigVMClient::HandleGraphModifiedEvent(ERigVMGraphNotifType InNotifType, UR
 		}
 		case ERigVMGraphNotifType::PinCategoriesChanged: // A node's category list has changed (Subject == URigVMNode)
 		{
-			if (URigVMNode* Node = Cast<URigVMNode>(InSubject))
+			if (const URigVMCollapseNode* CollapseNode = Cast<URigVMCollapseNode>(InSubject))
 			{
-				if(Node->GetOuter()->IsA<URigVMFunctionLibrary>())
+				if(CollapseNode->IsGraphFunctionDefinition())
 				{
-					if (const URigVMCollapseNode* CollapseNode = Cast<URigVMCollapseNode>(Node))
-					{
-						UpdateGraphFunctionData(CollapseNode);
-					}
+					UpdateGraphFunctionData(CollapseNode);
 				}
 			}
 			break;
diff --git a/Engine/Plugins/Runtime/RigVM/Source/RigVMDeveloper/Private/RigVMModel/RigVMController.cpp b/Engine/Plugins/Runtime/RigVM/Source/RigVMDeveloper/Private/RigVMModel/RigVMController.cpp
--- a/Engine/Plugins/Runtime/RigVM/Source/RigVMDeveloper/Private/RigVMModel/RigVMController.cpp
+++ b/Engine/Plugins/Runtime/RigVM/Source/RigVMDeveloper/Private/RigVMModel/RigVMController.cpp
@@ -433,6 +433,7 @@ FRigVMPinInfo::FRigVMPinInfo()
 	, DefaultValueType(ERigVMPinDefaultValueType::AutoDetect)
 	, bIsExpanded(false)
 	, bIsConstant(false)
+	, bIsInputVariable(false)
 	, bIsDynamicArray(false)
 	, bIsTrait(false)
 	, bIsLazy(false)
@@ -449,6 +450,7 @@ FRigVMPinInfo::FRigVMPinInfo(const URigVMPin* InPin, int32 InParentIndex, ERigVM
 	, DefaultValueType(InDefaultValueType)
 	, bIsExpanded(InPin->IsExpanded())
 	, bIsConstant(InPin->IsDefinedAsConstant())
+	, bIsInputVariable(InPin->IsDefinedAsInputVariable())
 	, bIsDynamicArray(InPin->IsDynamicArray())
 	, bIsTrait(InPin->IsTraitPin() && InPin->IsRootPin())
 	, bIsLazy(InPin->IsLazy() && InPin->IsRootPin())
@@ -493,6 +495,7 @@ FRigVMPinInfo::FRigVMPinInfo(FProperty* InProperty, ERigVMPinDirection InDirecti
 	, DefaultValueType(InDefaultValueType)
 	, bIsExpanded(false)
 	, bIsConstant(false)
+	, bIsInputVariable(false)
 	, bIsDynamicArray(false)
 	, bIsTrait(false)
 	, bIsLazy(false)
@@ -521,6 +524,7 @@ FRigVMPinInfo::FRigVMPinInfo(FProperty* InProperty, ERigVMPinDirection InDirecti
 	}
 	
 	bIsConstant = InProperty->HasMetaData(TEXT("Constant"));
+	bIsInputVariable = InProperty->HasMetaData(TEXT("ExternalVariable"));
 	CustomWidgetName = InProperty->GetMetaData(TEXT("CustomWidget"));
 	if (InProperty->HasMetaData(FRigVMStruct::ExpandPinByDefaultMetaName))
 	{
@@ -629,6 +633,7 @@ uint32 GetTypeHash(const FRigVMPinInfo& InPin)
 	Hash = HashCombine(Hash, GetTypeHash((int32)InPin.TypeIndex));
 	Hash = HashCombine(Hash, GetTypeHash(InPin.bIsArray));
 	Hash = HashCombine(Hash, GetTypeHash(InPin.bIsTrait));
+	Hash = HashCombine(Hash, GetTypeHash(InPin.bIsInputVariable));
 	// we are not hashing the parent index,  pinpath, default value or the property since
 	// it doesn't matter for the structure validity of the node
 	return Hash;
@@ -650,18 +655,20 @@ int32 FRigVMPinInfoArray::AddPin(const URigVMPin* InPin, int32 InParentIndex, ER
 {
 	// this method adds all pins as currently represented in the model.
 	const int32 Index = Pins.Emplace(InPin, InParentIndex, InDirection, InDefaultValueType);
-	for(const URigVMPin* SubPin : InPin->GetSubPins())
+	if (!Pins[Index].ShouldHideSubPins())
 	{
-		const int32 SubPinIndex = AddPin(SubPin, Index, InDirection, InDefaultValueType);
-		Pins[Index].SubPins.Add(SubPinIndex);
+		for(const URigVMPin* SubPin : InPin->GetSubPins())
+		{
+			const int32 SubPinIndex = AddPin(SubPin, Index, InDirection, InDefaultValueType);
+			Pins[Index].SubPins.Add(SubPinIndex);
+		}
 	}
+	Pins[Index].bIsInputVariable = InPin->IsDefinedAsInputVariable();
 	return Index;
 }
 
 FRigVMPinInfoArray::FRigVMPinInfoArray(const URigVMNode* InNode, URigVMController* InController, const FRigVMPinInfoArray* InPreviousPinInfos)
 {
-	const bool bAddSubPins = !InNode->IsA<URigVMRerouteNode>();
-	
 	// this method adds pins as needed based on the property structure
 	for(const URigVMPin* Pin : InNode->GetPins())
 	{
@@ -671,11 +678,13 @@ FRigVMPinInfoArray::FRigVMPinInfoArray(const URigVMNode* InNode, URigVMControlle
 		{
 			InController->ReportErrorf( TEXT("Invalid pin type %s for %s in %s"), *Pin->GetCPPType(), *Pin->GetPathName(), *InNode->GetPackage()->GetPathName());
 		}
+		const bool bAddSubPins = !InNode->IsA<URigVMRerouteNode>() && !Pin->ShouldHideSubPins();
 		int32 Index = AddPin(InController, INDEX_NONE, Pin->GetFName(), Pin->GetDirection(), Pin->GetTypeIndex(), DefaultValue, DefaultValueType, nullptr, InPreviousPinInfos, bAddSubPins);
 		if (InNode->IsA<URigVMLibraryNode>())
 		{
 			Pins[Index].bIsDynamicArray = Pins[Index].bIsArray;
 		}
+		Pins[Index].bIsInputVariable = Pin->IsDefinedAsInputVariable();
 	}
 }
 
@@ -697,8 +706,10 @@ FRigVMPinInfoArray::FRigVMPinInfoArray(const FRigVMGraphFunctionHeader& Function
 		{
 			DefaultValueType = PinInfo->DefaultValueType;
 		}
-		int32 Index = AddPin(InController, INDEX_NONE, FunctionArgument.Name, FunctionArgument.Direction, TypeIndex, FunctionArgument.DefaultValue, DefaultValueType, nullptr, InPreviousPinInfos, true);
+		const bool bAddSubPins = !(FunctionArgument.bIsInputVariable && FunctionArgument.Direction == ERigVMPinDirection::Input);
+		int32 Index = AddPin(InController, INDEX_NONE, FunctionArgument.Name, FunctionArgument.Direction, TypeIndex, FunctionArgument.DefaultValue, DefaultValueType, nullptr, InPreviousPinInfos, bAddSubPins);
 		Pins[Index].bIsDynamicArray = Pins[Index].bIsArray;
+		Pins[Index].bIsInputVariable = FunctionArgument.bIsInputVariable;
 	}
 }
 
@@ -715,6 +726,14 @@ int32 FRigVMPinInfoArray::AddPin(FProperty* InProperty, URigVMController* InCont
 		Pins[InParentIndex].SubPins.Add(Index);
 	}
 
+#if WITH_EDITOR
+	Pins[Index].bIsInputVariable = InProperty->HasMetaData(TEXT("ExternalVariable"));
+#endif
+	if (Pins[Index].ShouldHideSubPins())
+	{
+		bAddSubPins = false;
+	}
+
 	if(bAddSubPins)
 	{
 		if (const FStructProperty* StructProperty = CastField<FStructProperty>(InProperty))
@@ -839,6 +858,10 @@ void FRigVMPinInfoArray::AddPins(UScriptStruct* InScriptStruct, URigVMController
 		StructOnScope = MakeShareable(new FStructOnScope(InScriptStruct));
 		RigVMStructInstance = reinterpret_cast<FRigVMStruct*>(StructOnScope->GetStructMemory()); 
 	}
+	else if (InDefaultValueMemory == nullptr)
+	{
+		StructOnScope = MakeShareable(new FStructOnScope(InScriptStruct));
+	}
 
 	if (InController->GetSchema()->ShouldUnfoldStruct(InController, InScriptStruct))
 	{
@@ -863,6 +886,10 @@ void FRigVMPinInfoArray::AddPins(UScriptStruct* InScriptStruct, URigVMController
 				{
 					DefaultValueMemory = It->ContainerPtrToValuePtr<uint8>(InDefaultValueMemory);
 				}
+				else
+				{
+					DefaultValueMemory = It->ContainerPtrToValuePtr<uint8>(StructOnScope->GetStructMemory());
+				}
 
 				bool bAddSubPinsForProperty = bAddSubPins;
 #if WITH_EDITOR
@@ -1007,7 +1034,6 @@ void URigVMController::HandleModifiedEvent(ERigVMGraphNotifType InNotifType, URi
 		}
 		case ERigVMGraphNotifType::VariableAdded:
 		case ERigVMGraphNotifType::VariableRemoved:
-		case ERigVMGraphNotifType::VariableRemappingChanged:
 		{
 			URigVMGraph* RootGraph = InGraph->GetRootGraph();
 			if(URigVMFunctionLibrary* FunctionLibrary = Cast<URigVMFunctionLibrary>(RootGraph->GetRootGraph()))
@@ -1026,20 +1052,6 @@ void URigVMController::HandleModifiedEvent(ERigVMGraphNotifType InNotifType, URi
 						}
 					}
 				}
-
-				if (!bIsLocal)
-				{
-					if(const URigVMLibraryNode* Function = FunctionLibrary->FindFunctionForNode(Node))
-					{
-						FunctionLibrary->ForEachReference(Function->GetFName(), [this](URigVMFunctionReferenceNode* Reference)
-                    	{
-							if(const URigVMController* ReferenceController = GetControllerForGraph(Reference->GetGraph()))
-							{
-								ReferenceController->Notify(ERigVMGraphNotifType::VariableRemappingChanged, Reference);
-							}
-                    	});
-					}
-				}
 			}
 		}
 	}
@@ -1927,6 +1939,18 @@ URigVMVariableNode* URigVMController::AddVariableNode(const FName& InVariableNam
 					}
 				}
 
+				if (!bFoundLocalVariable)
+				{
+					for (FRigVMGraphVariableDescription& InputVariable : OuterFunction->GetContainedGraph()->GetInputVariables())
+					{
+						if (InputVariable.Name == InVariableName)
+						{
+							bFoundLocalVariable = true;
+							break;
+						}
+					}
+				}
+				
 				if (!bFoundLocalVariable)
 				{
 					// Make sure there is no external variable with that name
@@ -1934,7 +1958,7 @@ URigVMVariableNode* URigVMController::AddVariableNode(const FName& InVariableNam
 					bool bFoundExternalVariable = false;
 					for(const FRigVMExternalVariable& ExternalVariable : ExternalVariables)
 					{
-						if(ExternalVariable.Name == InVariableName)
+						if(ExternalVariable.GetName() == InVariableName)
 						{
 							bFoundExternalVariable = true;
 							break;
@@ -2002,18 +2026,13 @@ URigVMVariableNode* URigVMController::AddVariableNode(const FName& InVariableNam
 	AddNodePin(Node, VariablePin);
 
 	URigVMPin* ValuePin = NewObject<URigVMPin>(Node, URigVMVariableNode::ValueName);
-	ValuePin->CPPType = ExternalVariable.TypeName.ToString();
-	ValuePin->CPPTypeObject = ExternalVariable.TypeObject;
+	ValuePin->CPPType = ExternalVariable.GetExtendedCPPType().ToString();
+	ValuePin->CPPTypeObject = ExternalVariable.GetCPPTypeObject();
 	if (ValuePin->CPPTypeObject)
 	{
 		ValuePin->CPPTypeObjectPath = *ValuePin->CPPTypeObject->GetPathName();
 	}
-	ValuePin->bIsDynamicArray = ExternalVariable.bIsArray;
-
-	if(ValuePin->bIsDynamicArray && !RigVMTypeUtils::IsArrayType(ValuePin->CPPType))
-	{
-		ValuePin->CPPType = RigVMTypeUtils::ArrayTypeFromBaseType(*ValuePin->CPPType);
-	}
+	ValuePin->bIsDynamicArray = ExternalVariable.IsArray();
 
 	ValuePin->Direction = bIsGetter ? ERigVMPinDirection::Output : ERigVMPinDirection::Input;
 	AddNodePin(Node, ValuePin);
@@ -2235,37 +2254,7 @@ void URigVMController::OnExternalVariableRemoved(const FName& InVarName, bool bS
 				// call this function for the contained graph recursively 
 				CollapseController->OnExternalVariableRemoved(InVarName, bSetupUndoRedo);
 			}
-			
-			// if we are a function we need to notify all references!
-			if(const URigVMFunctionLibrary* FunctionLibrary = Cast<URigVMFunctionLibrary>(Graph))
-			{
-				FunctionLibrary->ForEachReference(CollapseNode->GetFName(), [this, InVarName](URigVMFunctionReferenceNode* Reference)
-				{
-					if(Reference->VariableMap.Contains(InVarName))
-					{
-						Reference->Modify();
-						Reference->VariableMap.Remove(InVarName);
-
-						if(URigVMController* ReferenceController = GetControllerForGraph(Reference->GetGraph()))
-						{
-							FRigVMControllerCompileBracketScope ReferenceCompileScope(ReferenceController);
-							ReferenceController->Notify(ERigVMGraphNotifType::VariableRemappingChanged, Reference);
-						}
-					}
-				});
-			}
-		}
-		else if(URigVMFunctionReferenceNode* FunctionReferenceNode = Cast<URigVMFunctionReferenceNode>(Node))
-		{
-			TMap<FName, FName> VariableMap = FunctionReferenceNode->GetVariableMap();
-			for(const TPair<FName, FName>& VariablePair : VariableMap)
-			{
-				if(VariablePair.Value == InVarName)
-				{
-					SetRemappedVariable(FunctionReferenceNode, VariablePair.Key, NAME_None, bSetupUndoRedo);
-				}
-			}
-		}
+		}			
 	}
 
 	if (bSetupUndoRedo)
@@ -2332,39 +2321,6 @@ bool URigVMController::OnExternalVariableRenamed(const FName& InOldVarName, cons
 				TGuardValue<bool> GuardEditGraph(CollapseNode->ContainedGraph->bEditable, true);
 				CollapseController->OnExternalVariableRenamed(InOldVarName, InNewVarName, bSetupUndoRedo);
 			}
-
-			// if we are a function we need to notify all references!
-			if(const URigVMFunctionLibrary* FunctionLibrary = Cast<URigVMFunctionLibrary>(Graph))
-			{
-				FunctionLibrary->ForEachReference(CollapseNode->GetFName(), [this, InOldVarName, InNewVarName](URigVMFunctionReferenceNode* Reference)
-                {
-					if(Reference->VariableMap.Contains(InOldVarName))
-					{
-						Reference->Modify();
-
-						const FName MappedVariable = Reference->VariableMap.FindChecked(InOldVarName);
-						Reference->VariableMap.Remove(InOldVarName);
-						Reference->VariableMap.FindOrAdd(InNewVarName) = MappedVariable; 
-
-						if(URigVMController* ReferenceController = GetControllerForGraph(Reference->GetGraph()))
-						{
-							FRigVMControllerCompileBracketScope ReferenceCompileScope(ReferenceController);
-							ReferenceController->Notify(ERigVMGraphNotifType::VariableRemappingChanged, Reference);
-						}
-                    }
-                });
-			}
-		}
-		else if(URigVMFunctionReferenceNode* FunctionReferenceNode = Cast<URigVMFunctionReferenceNode>(Node))
-		{
-			TMap<FName, FName> VariableMap = FunctionReferenceNode->GetVariableMap();
-			for(const TPair<FName, FName>& VariablePair : VariableMap)
-			{
-				if(VariablePair.Value == InOldVarName)
-				{
-					SetRemappedVariable(FunctionReferenceNode, VariablePair.Key, InNewVarName, bSetupUndoRedo);
-				}
-			}
 		}
 	}
 
@@ -2434,36 +2390,6 @@ void URigVMController::OnExternalVariableTypeChanged(const FName& InVarName, con
 				TGuardValue<bool> GuardEditGraph(CollapseNode->ContainedGraph->bEditable, true);
 				CollapseController->OnExternalVariableTypeChanged(InVarName, InCPPType, InCPPTypeObject, bSetupUndoRedo);
 			}
-
-			// if we are a function we need to notify all references!
-			if(const URigVMFunctionLibrary* FunctionLibrary = Cast<URigVMFunctionLibrary>(Graph))
-			{
-				FunctionLibrary->ForEachReference(CollapseNode->GetFName(), [this, InVarName](URigVMFunctionReferenceNode* Reference)
-                {
-                    if(Reference->VariableMap.Contains(InVarName))
-                    {
-                        Reference->Modify();
-                        Reference->VariableMap.Remove(InVarName); 
-
-                    	if(URigVMController* ReferenceController = GetControllerForGraph(Reference->GetGraph()))
-                    	{
-                    		FRigVMControllerCompileBracketScope ReferenceCompileScope(ReferenceController);
-							ReferenceController->Notify(ERigVMGraphNotifType::VariableRemappingChanged, Reference);
-						}
-                    }
-                });
-			}
-		}
-		else if(URigVMFunctionReferenceNode* FunctionReferenceNode = Cast<URigVMFunctionReferenceNode>(Node))
-		{
-			TMap<FName, FName> VariableMap = FunctionReferenceNode->GetVariableMap();
-			for(const TPair<FName, FName>& VariablePair : VariableMap)
-			{
-				if(VariablePair.Value == InVarName)
-				{
-					SetRemappedVariable(FunctionReferenceNode, VariablePair.Key, NAME_None, bSetupUndoRedo);
-				}
-			}
 		}
 
 		TArray<URigVMPin*> AllPins = Node->GetAllPinsRecursively();
@@ -5981,11 +5907,9 @@ TArray<URigVMNode*> URigVMController::ExpandLibraryNode(URigVMLibraryNode* InNod
 						FRigVMExternalVariable LocalVariableExternalType = LocalVariable.ToExternalVariable();
 						for (FRigVMExternalVariable& CurrentVariable : CurrentVariables)
 						{						
-							if (CurrentVariable.Name == LocalVariable.Name)
+							if (CurrentVariable.GetName() == LocalVariable.Name)
 							{
-								if (CurrentVariable.TypeName != LocalVariableExternalType.TypeName ||
-									CurrentVariable.TypeObject != LocalVariableExternalType.TypeObject ||
-									CurrentVariable.bIsArray != LocalVariableExternalType.bIsArray)
+								if (!CurrentVariable.IsSameType(LocalVariableExternalType))
 								{
 									bVariableIncompatible = true;	
 								}
@@ -6832,11 +6756,9 @@ URigVMCollapseNode* URigVMController::PromoteFunctionReferenceNodeToCollapseNode
 					FRigVMExternalVariable LocalVariableExternalType = LocalVariable.ToExternalVariable();
 					for (FRigVMExternalVariable& CurrentVariable : CurrentVariables)
 					{						
-						if (CurrentVariable.Name == LocalVariable.Name)
+						if (CurrentVariable.GetName() == LocalVariable.Name)
 						{
-							if (CurrentVariable.TypeName != LocalVariableExternalType.TypeName ||
-								CurrentVariable.TypeObject != LocalVariableExternalType.TypeObject ||
-								CurrentVariable.bIsArray != LocalVariableExternalType.bIsArray)
+							if (!CurrentVariable.IsSameType(LocalVariableExternalType))
 							{
 								bVariableIncompatible = true;	
 							}
@@ -7077,7 +6999,7 @@ bool URigVMController::RemoveNodes(TArray<URigVMNode*> InNodes, bool bSetupUndoR
 						{
 							const FName VariableToRemove = VariableNode->GetVariableName();
 							bool bIsLocalVariable = false;
-							for (FRigVMGraphVariableDescription VariableDescription : OuterFunction->GetContainedGraph()->GetLocalVariables(true))
+							for (const FRigVMGraphVariableDescription& VariableDescription : OuterFunction->GetContainedGraph()->GetLocalVariables(true))
 							{
 								if (VariableDescription.Name == VariableToRemove)
 								{
@@ -7098,13 +7020,25 @@ bool URigVMController::RemoveNodes(TArray<URigVMNode*> InNodes, bool bSetupUndoR
 								bool bFoundExternalVariable = false;
 								for(const FRigVMExternalVariable& ExternalVariable : ExternalVariablesWithoutVariableNode)
 								{
-									if(ExternalVariable.Name == VariableToRemove)
+									if(ExternalVariable.GetName() == VariableToRemove)
 									{
 										bFoundExternalVariable = true;
 										break;
 									}
 								}
 
+								if(!bFoundExternalVariable)
+								{
+									for (const FRigVMGraphVariableDescription& VariableDescription : OuterFunction->GetContainedGraph()->GetInputVariables())
+									{
+										if (VariableDescription.Name == VariableToRemove)
+										{
+											bFoundExternalVariable = true;
+											break;
+										}
+									}
+								}
+
 								if(!bFoundExternalVariable)
 								{
 									if(URigVMController* FunctionController = GetControllerForGraph(OuterFunction->GetContainedGraph()))
@@ -12339,30 +12273,30 @@ bool URigVMController::BindPinToVariable(URigVMPin* InPin, const FString& InNewB
 	InNewBoundVariablePath.Split(TEXT("."), &VariableName, &SegmentPath);
 
 	FRigVMExternalVariable Variable;
-	for (const FRigVMExternalVariable& VariableDescription : GetAllVariables(true))
+	const TArray<FRigVMExternalVariable> AllVariables = GetAllVariables(true);
+	for (const FRigVMExternalVariable& VariableDescription : AllVariables)
 	{
-		if (VariableDescription.Name.ToString() == VariableName)
+		if (VariableDescription.GetName().ToString().Equals(VariableName, ESearchCase::IgnoreCase))
 		{
 			Variable = VariableDescription;
 			break;
 		}
 	}
 
-	if(!GetSchema()->CanBindVariable(this, InPin, &Variable, InNewBoundVariablePath))
+	if (!Variable.IsValid(true))
 	{
+		ReportErrorf(TEXT("Cannot find variable '%s' in this graph."), *VariableName);
 		return false;
 	}
 
-	if (!Variable.Name.IsValid())
+	if(!GetSchema()->CanBindVariable(this, InPin, &Variable, InNewBoundVariablePath))
 	{
-		ReportError(TEXT("Cannot find variable in this graph."));
 		return false;
 	}
 
-	
 	if (!RigVMTypeUtils::AreCompatible(Variable, InPin->ToExternalVariable(), SegmentPath))
 	{
-		ReportError(TEXT("Cannot find variable in this graph."));
+		ReportErrorf(TEXT("Pin '%s' is not compatible with Variable '%s'"), *InPin->GetPinPath(), *Variable.GetName().ToString());
 		return false;
 	}
 	
@@ -12439,6 +12373,16 @@ bool URigVMController::BindPinToVariable(URigVMPin* InPin, const FString& InNewB
 		return false;
 	}
 	
+	// remove the mapping from the variable map since it may still be in there.
+	// we also need to remove it here since the bound variable now acts as the variable mapping instead.
+	if (InPin->IsDefinedAsInputVariable())
+	{
+		if (URigVMFunctionReferenceNode* FunctionReferenceNode = Cast<URigVMFunctionReferenceNode>(InPin->GetNode()))
+		{
+			FunctionReferenceNode->VariableMap_DEPRECATED.Remove(InPin->GetFName());
+		}
+	}
+	
 	if (bSetupUndoRedo)
 	{
 		GetActionStack()->EndAction(Action);
@@ -12510,6 +12454,21 @@ bool URigVMController::UnbindPinFromVariable(URigVMPin* InPin, bool bSetupUndoRe
 
 	RemoveInjectedNode(InPin->GetPinPath(), true, bSetupUndoRedo);
 
+	TArray<URigVMPin*> SourcePins = InPin->GetLinkedSourcePins();
+	if(SourcePins.Num() == 1)
+	{
+		BreakLink(SourcePins[0], InPin, bSetupUndoRedo);
+	}
+	
+	// remove the mapping from the variable map
+	if (InPin->IsDefinedAsInputVariable())
+	{
+		if (URigVMFunctionReferenceNode* FunctionReferenceNode = Cast<URigVMFunctionReferenceNode>(InPin->GetNode()))
+		{
+			FunctionReferenceNode->VariableMap_DEPRECATED.Remove(InPin->GetFName());
+		}
+	}
+
 	if (bSetupUndoRedo)
 	{
 		GetActionStack()->EndAction(Action);
@@ -12698,9 +12657,9 @@ bool URigVMController::PromotePinToVariable(URigVMPin* InPin, bool bCreateVariab
 		VariableForPin = GetVariableByName(*InPin->GetBoundVariableName());
 		check(VariableForPin.IsValid(true /* allow nullptr */));
 		SegmentPath = InPin->GetBoundVariablePath();
-		if (SegmentPath.StartsWith(VariableForPin.Name.ToString() + TEXT(".")))
+		if (SegmentPath.StartsWith(VariableForPin.GetName().ToString() + TEXT(".")))
 		{
-			SegmentPath = SegmentPath.RightChop(VariableForPin.Name.ToString().Len());
+			SegmentPath = SegmentPath.RightChop(VariableForPin.GetName().ToString().Len());
 		}
 		else
 		{
@@ -12732,9 +12691,9 @@ bool URigVMController::PromotePinToVariable(URigVMPin* InPin, bool bCreateVariab
 	{
 		FRigVMControllerCompileBracketScope CompileScope(this);
 		if (URigVMVariableNode* VariableNode = AddVariableNode(
-			VariableForPin.Name,
-			VariableForPin.TypeName.ToString(),
-			VariableForPin.TypeObject,
+			VariableForPin.GetName(),
+			VariableForPin.GetExtendedCPPType().ToString(),
+			const_cast<UObject*>(VariableForPin.GetCPPTypeObject()),
 			true,
 			FString(),
 			InNodePosition,
@@ -12750,7 +12709,7 @@ bool URigVMController::PromotePinToVariable(URigVMPin* InPin, bool bCreateVariab
 	else
 	{
 		FRigVMControllerCompileBracketScope CompileScope(this);
-		return BindPinToVariable(InPin, VariableForPin.Name.ToString(), bSetupUndoRedo);
+		return BindPinToVariable(InPin, VariableForPin.GetName().ToString(), bSetupUndoRedo);
 	}
 
 	return false;
@@ -13513,7 +13472,7 @@ bool URigVMController::BreakAllLinksRecursive(URigVMPin* Pin, bool bAsInput, boo
 	return bBrokenLinks;
 }
 
-FName URigVMController::AddExposedPin(const FName& InPinName, ERigVMPinDirection InDirection, const FString& InCPPType, const FName& InCPPTypeObjectPath, const FString& InDefaultValue, bool bSetupUndoRedo, bool bPrintPythonCommand)
+FName URigVMController::AddExposedPin(const FName& InPinName, ERigVMPinDirection InDirection, const FString& InCPPType, const FName& InCPPTypeObjectPath, const FString& InDefaultValue, bool bSetupUndoRedo, bool bPrintPythonCommand, bool bIsInputVariable)
 {
 	if (!IsValidGraph())
 	{
@@ -13622,6 +13581,8 @@ FName URigVMController::AddExposedPin(const FName& InPinName, ERigVMPinDirection
 	Pin->CPPType = CPPType;
 	Pin->CPPTypeObjectPath = InCPPTypeObjectPath;
 	Pin->bIsConstant = false;
+	Pin->bIsInputVariable = bIsInputVariable;
+
 	Pin->Direction = InDirection;
 	Pin->bIsDynamicArray = RigVMTypeUtils::IsArrayType(CPPType);
 	AddNodePin(LibraryNode, Pin);
@@ -13700,7 +13661,7 @@ FName URigVMController::AddExposedPin(const FName& InPinName, ERigVMPinDirection
 		//AddExposedPin(const FName& InPinName, ERigVMPinDirection InDirection, const FString& InCPPType, const FName& InCPPTypeObjectPath, const FString& InDefaultValue, bool bSetupUndoRedo, bool bPrintPythonCommand)
 
 		const FString GraphName = GetSchema()->GetSanitizedGraphName(GetGraph()->GetGraphName());
-		static constexpr TCHAR AddExposedPinFormat[] = TEXT("blueprint.get_controller_by_name('%s').add_exposed_pin('%s', %s, '%s', '%s', '%s')");
+		static constexpr TCHAR AddExposedPinFormat[] = TEXT("blueprint.get_controller_by_name('%s').add_exposed_pin('%s', %s, '%s', '%s', '%s', %s, %s, %s)");
 		
 		RigVMPythonUtils::Print(GetSchema()->GetGraphOuterName(GetGraph()), 
 			FString::Printf(AddExposedPinFormat,
@@ -13709,7 +13670,11 @@ FName URigVMController::AddExposedPin(const FName& InPinName, ERigVMPinDirection
 				*RigVMPythonUtils::EnumValueToPythonString<ERigVMPinDirection>((int64)InDirection),
 				*InCPPType,
 				*InCPPTypeObjectPath.ToString(),
-				*InDefaultValue));
+				*InDefaultValue,
+				bSetupUndoRedo ? TEXT("True") : TEXT("False"),
+				bPrintPythonCommand ? TEXT("True") : TEXT("False"),
+				bIsInputVariable ? TEXT("True") : TEXT("False")
+			));
 	}
 	
 	return PinName;
@@ -14376,6 +14341,114 @@ bool URigVMController::SetExposedPinIndex(const FName& InPinName, int32 InNewInd
 	return true;
 }
 
+bool URigVMController::MarkExposedPinAsInputVariable(const FName& InPinName, bool bDefinedAsExternalVariable, bool bSetupUndoRedo,
+	bool bPrintPythonCommand)
+{
+	if (!IsValidGraph())
+	{
+		return false;
+	}
+
+	if (!bIsTransacting && !IsGraphEditable())
+	{
+		return false;
+	}
+
+	URigVMGraph* Graph = GetGraph();
+	check(Graph);
+
+	FString PinPath = InPinName.ToString();
+	if (PinPath.Contains(TEXT(".")))
+	{
+		ReportError(TEXT("Cannot change pin index for pins on nodes for now - only within collapse nodes."));
+		return false;
+	}
+
+	URigVMLibraryNode* LibraryNode = Cast<URigVMLibraryNode>(Graph->GetOuter());
+	if (LibraryNode == nullptr)
+	{
+		ReportError(TEXT("Graph is not under a Collapse Node"));
+		return false;
+	}
+
+	if (!Cast<URigVMFunctionLibrary>(LibraryNode->GetGraph()))
+	{
+		ReportError(TEXT("Input Variables can only be used on functions - not on normal collapse nodes."));
+		return false;
+	}
+
+	URigVMPin* Pin = LibraryNode->FindPin(PinPath);
+	if (Pin == nullptr)
+	{
+		ReportErrorf(TEXT("Cannot find exposed pin '%s'."), *PinPath);
+		return false;
+	}
+
+	if (Pin->IsDefinedAsInputVariable() == bDefinedAsExternalVariable)
+	{
+		return true; // Nothing to do, do not fail
+	}
+
+	if (Pin->IsExecuteContext())
+	{
+		ReportErrorf(TEXT("Pin '%s' is an execute context. Only non-execute pins can be marked as external variable."), *PinPath);
+		return false;
+	}
+
+	if (Pin->GetDirection() != ERigVMPinDirection::Input)
+	{
+		ReportErrorf(TEXT("Pin '%s' is not an input pin."), *PinPath);
+		return false;
+	}
+
+	FRigVMControllerCompileBracketScope CompileBracketScope(this);
+
+	FRigVMMarkExposedPinAsInputVariableAction PinExternalVariableAction(this, Pin, bDefinedAsExternalVariable);
+	if (bSetupUndoRedo)
+	{
+		GetActionStack()->BeginAction(PinExternalVariableAction);
+	}
+	
+	{
+		Pin->Modify();
+		Pin->bIsInputVariable = bDefinedAsExternalVariable;
+		Pin->IncrementVersion();
+		Notify(ERigVMGraphNotifType::PinTypeChanged, Pin);
+
+		if(URigVMController* LibraryController = GetControllerForGraph(LibraryNode->GetGraph()))
+		{
+			// if we are marking an input variable as a normal argument
+			// we need to run the code below to clean up obsolete variable nodes (like SetVariable)
+			if (!bDefinedAsExternalVariable)
+			{
+				LibraryController->OnExternalVariableRemoved(Pin->GetFName(), bSetupUndoRedo);
+			}
+		}
+	}
+
+	RefreshFunctionPins(LibraryNode->GetEntryNode());
+	RefreshFunctionPins(LibraryNode->GetReturnNode());
+	RefreshFunctionReferences(LibraryNode, false, false);
+	
+	if (bSetupUndoRedo)
+	{
+		GetActionStack()->EndAction(PinExternalVariableAction);
+	}
+	
+	if (bPrintPythonCommand)
+	{
+		const FString GraphName = GetSchema()->GetSanitizedGraphName(GetGraph()->GetGraphName());
+
+		RigVMPythonUtils::Print(GetSchema()->GetGraphOuterName(GetGraph()), 
+			FString::Printf(TEXT("blueprint.get_controller_by_name('%s').mark_exposed_pin_as_input_variable('%s', %s)"),
+				*GraphName,
+				*GetSchema()->GetSanitizedPinName(InPinName.ToString()),
+				(bDefinedAsExternalVariable) ? TEXT("True") : TEXT("False")));
+	}
+
+	return true;
+}
+
 FRigVMGraphFunctionHeader URigVMController::FindGraphFunctionHeaderByName(FString InHostPath, FName InFunctionName) const
 {
 	FString ErrorMessage;
@@ -14489,36 +14562,6 @@ bool URigVMController::SwapFunctionReference(URigVMFunctionReferenceNode* InFunc
 		ExternalVariables.Append(GetExternalVariablesDelegate.Execute(GetGraph()));
 	}
 
-	TArray<FName> MappedVariablesToRemove;
-	for(const TPair<FName, FName>& Pair : InFunctionReferenceNode->VariableMap)
-	{
-		const FRigVMExternalVariable* OuterExternalVariable = ExternalVariables.FindByPredicate([Pair](const FRigVMExternalVariable& ExternalVariable) -> bool
-		{
-			return ExternalVariable.Name == Pair.Key;
-		});
-
-		const FRigVMExternalVariable* InnerExternalVariable = InFunctionReferenceNode->ReferencedFunctionHeader.ExternalVariables.FindByPredicate([Pair](const FRigVMExternalVariable& ExternalVariable) -> bool
-		{
-			return ExternalVariable.Name == Pair.Value;
-		});
-
-		if(OuterExternalVariable && InnerExternalVariable)
-		{
-			const TRigVMTypeIndex OuterTypeIndex = OuterExternalVariable->GetTypeIndex();
-			const TRigVMTypeIndex InnerTypeIndex = InnerExternalVariable->GetTypeIndex();
-			if(FRigVMRegistry::Get().CanMatchTypes(OuterTypeIndex, InnerTypeIndex, true))
-			{
-				continue;
-			}
-		}
-
-		MappedVariablesToRemove.Add(Pair.Key);
-	}
-	for(const FName& MappedVariableToRemove : MappedVariablesToRemove)
-	{
-		InFunctionReferenceNode->VariableMap.Remove(MappedVariableToRemove);
-	}
-
 	Notify(ERigVMGraphNotifType::NodeReferenceChanged, InFunctionReferenceNode);
 	
 	if (bSetupUndoRedo)
@@ -14802,100 +14845,11 @@ URigVMFunctionReferenceNode* URigVMController::AddExternalFunctionReferenceNode(
 bool URigVMController::SetRemappedVariable(URigVMFunctionReferenceNode* InFunctionRefNode,
 	const FName& InInnerVariableName, const FName& InOuterVariableName, bool bSetupUndoRedo)
 {
-	if(!InFunctionRefNode)
-	{
-		return false;
-	}
-
-	if (!IsValidGraph())
+	if (!InFunctionRefNode)
 	{
 		return false;
 	}
-
-	if (!bIsTransacting && !IsGraphEditable())
-	{
-		return false;
-	}
-
-	if(InInnerVariableName.IsNone())
-	{
-		return false;
-	}
-
-	const FName OldOuterVariableName = InFunctionRefNode->GetOuterVariableName(InInnerVariableName);
-	if(OldOuterVariableName == InOuterVariableName)
-	{
-		return false;
-	}
-
-	if(!InFunctionRefNode->RequiresVariableRemapping())
-	{
-		return false;
-	}
-	
-	URigVMGraph* Graph = GetGraph();
-	check(Graph);
-
-	FRigVMExternalVariable InnerExternalVariable;
-	{
-		TArray<FRigVMExternalVariable> Variables = InFunctionRefNode->GetExternalVariables(false);
-		if (FRigVMExternalVariable* Variable = Variables.FindByPredicate([InInnerVariableName](const FRigVMExternalVariable& Variable)
-		{
-			return Variable.Name == InInnerVariableName;
-		}))
-		{
-			InnerExternalVariable = *Variable;
-		}
-	}
-
-	if(!InnerExternalVariable.IsValid(true))
-	{
-		ReportErrorf(TEXT("External variable '%s' cannot be found."), *InInnerVariableName.ToString());
-		return false;
-	}
-
-	ensure(InnerExternalVariable.Name == InInnerVariableName);
-
-	if(InOuterVariableName.IsNone())
-	{
-		InFunctionRefNode->Modify();
-		InFunctionRefNode->VariableMap.Remove(InInnerVariableName);
-	}
-	else
-	{
-		const FRigVMExternalVariable OuterExternalVariable = GetVariableByName(InOuterVariableName);
-		if(!OuterExternalVariable.IsValid(true))
-		{
-			ReportErrorf(TEXT("External variable '%s' cannot be found."), *InOuterVariableName.ToString());
-			return false;
-		}
-
-		ensure(OuterExternalVariable.Name == InOuterVariableName);
-
-		if((InnerExternalVariable.TypeObject != nullptr) && (InnerExternalVariable.TypeObject != OuterExternalVariable.TypeObject))
-		{
-			ReportErrorf(TEXT("Inner and Outer External variables '%s' and '%s' are not compatible."), *InInnerVariableName.ToString(), *InOuterVariableName.ToString());
-			return false;
-		}
-		if((InnerExternalVariable.TypeObject == nullptr) && (InnerExternalVariable.TypeName != OuterExternalVariable.TypeName))
-		{
-			ReportErrorf(TEXT("Inner and Outer External variables '%s' and '%s' are not compatible."), *InInnerVariableName.ToString(), *InOuterVariableName.ToString());
-			return false;
-		}
-
-		InFunctionRefNode->Modify();
-		InFunctionRefNode->VariableMap.FindOrAdd(InInnerVariableName) = InOuterVariableName;
-	}
-
-	Notify(ERigVMGraphNotifType::VariableRemappingChanged, InFunctionRefNode);
-
-	FRigVMControllerCompileBracketScope CompileScope(this);
-	if(bSetupUndoRedo)
-	{
-		GetActionStack()->AddAction(FRigVMSetRemappedVariableAction(this, InFunctionRefNode, InInnerVariableName, OldOuterVariableName, InOuterVariableName));
-	}
-	
-	return true;
+	return BindPinToVariable(InFunctionRefNode->FindRootPinByName(InInnerVariableName), InOuterVariableName.ToString(), bSetupUndoRedo);
 }
 
 URigVMLibraryNode* URigVMController::AddFunctionToLibrary(const FName& InFunctionName, bool bMutable, const FVector2D& InNodePosition, bool bSetupUndoRedo, bool bPrintPythonCommand)
@@ -15860,7 +15814,7 @@ bool URigVMController::RemoveLocalVariable(const FName& InVariableName, bool bSe
 
 			for (const FRigVMExternalVariable& ExternalVariable : ExternalVariables)
 			{
-				if (ExternalVariable.Name == InVariableName)
+				if (ExternalVariable.GetName() == InVariableName)
 				{
 					bSwitchToMemberVariable = true;
 					ExternalVariableToSwitch = ExternalVariable;
@@ -15898,7 +15852,7 @@ bool URigVMController::RemoveLocalVariable(const FName& InVariableName, bool bSe
 					{
 						if (VariablePin->GetDefaultValue() == VarNameStr)
 						{
-							RefreshVariableNode(VariableNode->GetFName(), ExternalVariableToSwitch.Name, ExternalVariableToSwitch.TypeName.ToString(), ExternalVariableToSwitch.TypeObject, bSetupUndoRedo, false);
+							RefreshVariableNode(VariableNode->GetFName(), ExternalVariableToSwitch.GetName(), ExternalVariableToSwitch.GetExtendedCPPType().ToString(), ExternalVariableToSwitch.GetCPPTypeObject(), bSetupUndoRedo, false);
 							continue;
 						}
 					}
@@ -15909,7 +15863,7 @@ bool URigVMController::RemoveLocalVariable(const FName& InVariableName, bool bSe
 				{
 					if (Pin->GetBoundVariableName() == InVariableName.ToString())
 					{
-						if (Pin->GetCPPType() != ExternalVariableToSwitch.TypeName.ToString() || Pin->GetCPPTypeObject() == ExternalVariableToSwitch.TypeObject)
+						if (Pin->GetCPPType() != ExternalVariableToSwitch.GetExtendedCPPType().ToString() || Pin->GetCPPTypeObject() == ExternalVariableToSwitch.GetCPPTypeObject())
 						{
 							UnbindPinFromVariable(Pin, bSetupUndoRedo);
 						}
@@ -16538,7 +16492,7 @@ bool URigVMController::SetVariableName(URigVMVariableNode* InVariableNode, const
 	TMap<FName, int32> NameToIndex;
 	for (int32 VariableIndex = 0; VariableIndex < Descriptions.Num(); VariableIndex++)
 	{
-		NameToIndex.Add(Descriptions[VariableIndex].Name, VariableIndex);
+		NameToIndex.Add(Descriptions[VariableIndex].GetName(), VariableIndex);
 	}
 
 	const FRigVMExternalVariable VariableType = RigVMTypeUtils::ExternalVariableFromCPPType(InVariableName, InVariableNode->GetCPPType(), InVariableNode->GetCPPTypeObject());
@@ -16548,9 +16502,7 @@ bool URigVMController::SetVariableName(URigVMVariableNode* InVariableNode, const
 		{
 			return true;
 		}
-		return VariableType.TypeName == Descriptions[*FoundIndex].TypeName &&
-				VariableType.TypeObject == Descriptions[*FoundIndex].TypeObject &&
-				VariableType.bIsArray == Descriptions[*FoundIndex].bIsArray;
+		return VariableType.IsSameType(Descriptions[*FoundIndex]);
 	}, false, true);
 
 	int32 NodesSharingName = 0;
@@ -18413,6 +18365,7 @@ void URigVMController::ConfigurePinFromProperty(FProperty* InProperty, URigVMPin
 	}
 	
 	InOutPin->bIsConstant = InProperty->HasMetaData(TEXT("Constant"));
+	InOutPin->bIsInputVariable = InProperty->HasMetaData(TEXT("ExternalVariable"));
 	FString CustomWidgetName = InProperty->GetMetaData(TEXT("CustomWidget"));
 	InOutPin->CustomWidgetName = CustomWidgetName.IsEmpty() ? FName(NAME_None) : FName(*CustomWidgetName);
 
@@ -18547,6 +18500,7 @@ void URigVMController::ConfigurePinFromPin(URigVMPin* InOutPin, const URigVMPin*
 	InOutPin->DefaultValue = InPin->DefaultValue;
 	InOutPin->bIsDynamicArray = InPin->bIsDynamicArray;
 	InOutPin->bIsLazy = InPin->bIsLazy;
+	InOutPin->bIsInputVariable = InPin->bIsInputVariable;
 	if(bCopyDisplayName)
 	{
 		InOutPin->SetDisplayName(InPin->GetDisplayName());
@@ -18564,6 +18518,7 @@ void URigVMController::ConfigurePinFromArgument(URigVMPin* InOutPin, const FRigV
 	// things that defines the state of the pin is copied during GetPinState()
 	// though addmittedly these two functions have overlaps currently
 	InOutPin->bIsConstant = InArgument.bIsConst;
+	InOutPin->bIsInputVariable = InArgument.bIsInputVariable;
 	InOutPin->Direction = InArgument.Direction;
 	InOutPin->CPPType = InArgument.CPPType.ToString();
 	InOutPin->CPPTypeObjectPath = *InArgument.CPPTypeObject.ToSoftObjectPath().ToString();
@@ -18857,7 +18812,7 @@ void URigVMController::GenerateRepopulatePinsNodeData(TArray<FRepopulatePinsNode
 	URigVMCollapseNode* CollapseNode = Cast<URigVMCollapseNode>(InNode);
 	if (CollapseNode)
 	{
-		if (CollapseNode->GetOuter()->IsA<URigVMFunctionLibrary>())
+		if (CollapseNode->IsGraphFunctionDefinition())
 		{
 			NodeData.bSetupOrphanPinsForThisNode = false;
 		}
@@ -19183,7 +19138,7 @@ bool URigVMController::GenerateNewPinInfos(const FRigVMRegistry& Registry, URigV
 			{
 				for (const FRigVMExternalVariable& ExternalVariable : ExternalVariables)
 				{
-					if (ExternalVariable.Name == CurrentExternalVariable.Name)
+					if (ExternalVariable.GetName() == CurrentExternalVariable.GetName())
 					{
 						Variable = ExternalVariable;
 						break;
@@ -19193,23 +19148,21 @@ bool URigVMController::GenerateNewPinInfos(const FRigVMRegistry& Registry, URigV
 
 			if (Variable.IsValid(true))
 			{
-				if (Variable.TypeName != CurrentExternalVariable.TypeName ||
-					Variable.TypeObject != CurrentExternalVariable.TypeObject ||
-					Variable.bIsArray != CurrentExternalVariable.bIsArray)
+				if (!Variable.IsSameType(CurrentExternalVariable))
 				{
 					FString CPPType;
 					UObject* CPPTypeObject;
 
 					if (RigVMTypeUtils::CPPTypeFromExternalVariable(Variable, CPPType, &CPPTypeObject))
 					{
-						RefreshVariableNode(VariableNode->GetFName(), Variable.Name, CPPType, Variable.TypeObject, false, bSetupOrphanPinsForThisNode);
+						RefreshVariableNode(VariableNode->GetFName(), Variable.GetName(), CPPType, Variable.GetCPPTypeObject(), false, bSetupOrphanPinsForThisNode);
 					}
 					else
 					{
 						ReportAndNotifyErrorf(
 							TEXT("Control Rig '%s', Type of Variable '%s' cannot be resolved."),
 							*InNode->GetOutermost()->GetPathName(),
-							*Variable.Name.ToString()
+							*Variable.GetName().ToString()
 						);
 					}
 				}
@@ -19219,7 +19172,7 @@ bool URigVMController::GenerateNewPinInfos(const FRigVMRegistry& Registry, URigV
 				ReportWarningf(
 					TEXT("Control Rig '%s', Variable '%s' not found."),
 					*InNode->GetOutermost()->GetPathName(),
-					*CurrentExternalVariable.Name.ToString()
+					*CurrentExternalVariable.GetName().ToString()
 				);
 			}
 		}
@@ -19272,7 +19225,13 @@ bool URigVMController::GenerateNewPinInfos(const FRigVMRegistry& Registry, URigV
 				}
 
 				ERigVMPinDirection Direction = bIsEntryNode ? ERigVMPinDirection::Output : ERigVMPinDirection::Input;
-				(void)NewPinInfos.AddPin(LibraryPin, INDEX_NONE, Direction, LibraryPin->GetDefaultValueType());
+				const int32 NewPinIndex = NewPinInfos.AddPin(LibraryPin, INDEX_NONE, Direction, LibraryPin->GetDefaultValueType());
+
+				if (LibraryPin->IsRootPin() && LibraryPin->GetSubPins().IsEmpty() && LibraryPin->IsDefinedAsInputVariable() &&
+					LibraryPin->IsStruct() && !NewPinInfos[NewPinIndex].ShouldHideSubPins())
+				{
+					NewPinInfos.AddPins(Cast<UScriptStruct>(LibraryPin->GetCPPTypeObject()), this, Direction, NewPinIndex, nullptr, nullptr, true);
+				}
 			}
 		}
 		else
@@ -19311,7 +19270,8 @@ bool URigVMController::GenerateNewPinInfos(const FRigVMRegistry& Registry, URigV
 		const FRigVMPinInfo& PreviousPin = PreviousPinInfos[Index];
 		if (PreviousPin.bIsTrait)
 		{
-			const int32 NewPinIndex = NewPinInfos.AddPin(this, INDEX_NONE, PreviousPin.Name, PreviousPin.Direction, PreviousPin.TypeIndex, PreviousPin.DefaultValue, PreviousPin.DefaultValueType, nullptr, &PreviousPinInfos, true);
+			const bool bAddSubPins = !PreviousPin.ShouldHideSubPins();
+			const int32 NewPinIndex = NewPinInfos.AddPin(this, INDEX_NONE, PreviousPin.Name, PreviousPin.Direction, PreviousPin.TypeIndex, PreviousPin.DefaultValue, PreviousPin.DefaultValueType, nullptr, &PreviousPinInfos, bAddSubPins);
 			NewPinInfos[NewPinIndex].bIsTrait = true;
 
 			if (URigVMPin* Pin = InNode->FindPin(PreviousPin.PinPath))
@@ -19391,6 +19351,10 @@ void URigVMController::GenerateRepopulatePinLists(const FRigVMRegistry& Registry
 			{
 				NodeData.bRequireRecreateLinks = true;
 			}
+			else if (NodeData.PreviousPinInfos[Index].bIsInputVariable != NodeData.NewPinInfos[NewIndex].bIsInputVariable)
+			{
+				NodeData.bRequireRecreateLinks = true;
+			}
 			else if (NodeData.PreviousPinInfos[Index].Direction != ERigVMPinDirection::Hidden)
 			{
 				NodeData.bRequireRecreateLinks |= bTypesDiffer;
@@ -19645,6 +19609,14 @@ void URigVMController::RepopulatePinsOnNode(const FRigVMRegistry& Registry, cons
 				check(Pin->CPPTypeObject == nullptr);
 				Pin->CPPTypeObjectPath = NAME_None;
 				Pin->LastKnownTypeIndex = NewPinInfo->TypeIndex;
+				Pin->bIsInputVariable = NewPinInfo->bIsInputVariable;
+				Pin->IncrementVersion();
+
+				Notify(ERigVMGraphNotifType::PinTypeChanged, Pin);
+			}
+			else if (Pin->bIsInputVariable != NewPinInfo->bIsInputVariable)
+			{
+				Pin->bIsInputVariable = NewPinInfo->bIsInputVariable;
 				Pin->IncrementVersion();
 
 				Notify(ERigVMGraphNotifType::PinTypeChanged, Pin);
@@ -19730,7 +19702,7 @@ void URigVMController::RepopulatePinsOnNode(const FRigVMRegistry& Registry, cons
 	}
 	else if (CollapseNode)
 	{
-		if (!CollapseNode->GetOuter()->IsA<URigVMFunctionLibrary>())
+		if (!CollapseNode->IsGraphFunctionDefinition())
 		{
 			// no need to notify since the function library graph is invisible anyway
 			RemoveUnusedOrphanedPins(CollapseNode);
@@ -20228,6 +20200,7 @@ URigVMPin* URigVMController::CreatePinFromPinInfo(const FRigVMRegistry& InRegist
 		Pin->Direction = InPinInfo.Direction;
 		Pin->DisplayName = InPinInfo.DisplayName.IsEmpty() ? NAME_None : FName(*InPinInfo.DisplayName);
 		Pin->bIsConstant = InPinInfo.bIsConstant;
+		Pin->bIsInputVariable = InPinInfo.bIsInputVariable;
 		Pin->bIsDynamicArray = InPinInfo.bIsDynamicArray;
 		Pin->bIsLazy = InPinInfo.bIsLazy;
 		Pin->CustomWidgetName = InPinInfo.CustomWidgetName.IsEmpty() ? NAME_None : FName(*InPinInfo.CustomWidgetName);
@@ -20263,6 +20236,18 @@ URigVMPin* URigVMController::CreatePinFromPinInfo(const FRigVMRegistry& InRegist
 
 	Notify(ERigVMGraphNotifType::PinAdded, Pin);
 
+	if (Pin->IsDefinedAsInputVariable())
+	{
+		if (const URigVMFunctionReferenceNode* FunctionReferenceNode = Cast<URigVMFunctionReferenceNode>(Pin->GetNode()))
+		{
+			const FString MappedVariablePath = FunctionReferenceNode->GetOuterVariablePath(Pin->GetFName());
+			if (!MappedVariablePath.IsEmpty())
+			{
+				(void)BindPinToVariable(Pin, MappedVariablePath, false);
+			}
+		}
+	}
+
 	return Pin;
 }
 
@@ -22350,7 +22335,7 @@ FRigVMExternalVariable URigVMController::GetVariableByName(const FName& InExtern
 	TArray<FRigVMExternalVariable> Variables = GetAllVariables(bIncludeInputArguments);
 	for (const FRigVMExternalVariable& Variable : Variables)
 	{
-		if (Variable.Name == InExternalVariableName)
+		if (Variable.GetName() == InExternalVariableName)
 		{
 			return Variable;
 		}
@@ -22367,13 +22352,21 @@ TArray<FRigVMExternalVariable> URigVMController::GetAllVariables(const bool bInc
 	{
 		for (FRigVMGraphVariableDescription LocalVariable : Graph->GetLocalVariables(bIncludeInputArguments))
 		{
-			ExternalVariables.Add(LocalVariable.ToExternalVariable());
+			FRigVMExternalVariable::MergeExternalVariable(ExternalVariables, LocalVariable.ToExternalVariable());
+		}
+		for (FRigVMGraphVariableDescription InputVariable : Graph->GetInputVariables())
+		{
+			FRigVMExternalVariable::MergeExternalVariable(ExternalVariables, InputVariable.ToExternalVariable());
 		}
 	}
 	
 	if (GetExternalVariablesDelegate.IsBound())
 	{
-		ExternalVariables.Append(GetExternalVariablesDelegate.Execute(GetGraph()));
+		const TArray<FRigVMExternalVariable> ExternalVariablesFromDelegate = GetExternalVariablesDelegate.Execute(GetGraph());
+		for (const FRigVMExternalVariable& ExternalVariableFromDelegate : ExternalVariablesFromDelegate)
+		{
+			FRigVMExternalVariable::MergeExternalVariable(ExternalVariables, ExternalVariableFromDelegate);
+		}
 	}
 
 	return ExternalVariables;
@@ -24443,6 +24436,77 @@ FRigVMClientPatchResult URigVMController::PatchLocalVariableTypes()
 	return Result;
 }
 
+FRigVMClientPatchResult URigVMController::PatchExternalVariablesToArguments()
+{
+	FRigVMClientPatchResult Result;
+	if (const URigVMGraph* Graph = GetGraph())
+	{
+		if (const URigVMFunctionLibrary* FunctionLibrary = Cast<URigVMFunctionLibrary>(Graph->GetParentGraph()))
+		{
+			if (const URigVMLibraryNode* CollapseNode = Cast<URigVMLibraryNode>(Graph->GetOuter()))
+			{
+				// if this function isn't public - don't bother
+				if (!FunctionLibrary->PublicFunctionNames.Contains(CollapseNode->GetFName()))
+				{
+					return Result;
+				}
+				
+				const TArray<FRigVMExternalVariable> ExternalVariables = Graph->GetExternalVariables();
+				const TArray<FRigVMGraphVariableDescription> LocalVariables = Graph->GetLocalVariables(false);
+				for (const FRigVMExternalVariable& ExternalVariable : ExternalVariables)
+				{
+					if (CollapseNode->FindRootPinByName(ExternalVariable.GetName()))
+					{
+						continue;
+					}
+					if (LocalVariables.ContainsByPredicate([&ExternalVariable](const FRigVMGraphVariableDescription& LocalVariable) -> bool
+						{
+							return ExternalVariable.GetName() == LocalVariable.Name;
+						}))
+					{
+						continue;
+					}
+
+					const FString TypeName = ExternalVariable.GetExtendedCPPType().ToString();
+					const FName PinName = AddExposedPin(ExternalVariable.GetName(), ERigVMPinDirection::Input, TypeName, *ExternalVariable.GetCPPTypeObjectPathName(), FString(), false, false, true);
+					Result.bSucceeded = !PinName.IsNone();
+					if (Result.bSucceeded)
+					{
+						Result.bChangedContent = true;
+					}
+				}
+			}
+		}
+		else
+		{
+			for (const URigVMNode* Node : Graph->GetNodes())
+			{
+				const URigVMFunctionReferenceNode* FunctionReferenceNode = Cast<URigVMFunctionReferenceNode>(Node);
+				if (!FunctionReferenceNode)
+				{
+					continue;
+				}
+				for (URigVMPin* Pin : Node->GetPins())
+				{
+					if (Pin->IsDefinedAsInputVariable())
+					{
+						const FString MappedVariablePath = FunctionReferenceNode->GetOuterVariablePath(Pin->GetFName());
+						if (MappedVariablePath.IsEmpty())
+						{
+							continue;
+						}
+						if (BindPinToVariable(Pin, MappedVariablePath, false))
+						{
+							Result.bChangedContent = true;
+						}
+					}
+				}
+			}
+		}
+	}
+	return Result;
+}
+
 ERigVMPinDefaultValueType URigVMController::GetDefaultValueType(const URigVMPin* InPin, const FString& InDefaultValue) const
 {
 	if(!CVarRigVMEnablePinOverrides.GetValueOnAnyThread())
diff --git a/Engine/Plugins/Runtime/RigVM/Source/RigVMDeveloper/Private/RigVMModel/RigVMControllerActions.cpp b/Engine/Plugins/Runtime/RigVM/Source/RigVMDeveloper/Private/RigVMModel/RigVMControllerActions.cpp
--- a/Engine/Plugins/Runtime/RigVM/Source/RigVMDeveloper/Private/RigVMModel/RigVMControllerActions.cpp
+++ b/Engine/Plugins/Runtime/RigVM/Source/RigVMDeveloper/Private/RigVMModel/RigVMControllerActions.cpp
@@ -2034,56 +2034,42 @@ bool FRigVMSetPinIndexAction::Redo()
 	return FRigVMBaseAction::Redo();
 }
 
-FRigVMSetRemappedVariableAction::FRigVMSetRemappedVariableAction()
+FRigVMMarkExposedPinAsInputVariableAction::FRigVMMarkExposedPinAsInputVariableAction()
 : FRigVMBaseAction(nullptr)
-, InnerVariableName(NAME_None)
-, OldOuterVariableName(NAME_None)
-, NewOuterVariableName(NAME_None)
+, PinPath()
+, OldState(false)
+, NewState(false)
 {
 }
 
-FRigVMSetRemappedVariableAction::FRigVMSetRemappedVariableAction(
-	URigVMController* InController,
-	URigVMFunctionReferenceNode* InFunctionRefNode,
-	const FName& InInnerVariableName,
-	const FName& InOldOuterVariableName,
-	const FName& InNewOuterVariableName)
+FRigVMMarkExposedPinAsInputVariableAction::FRigVMMarkExposedPinAsInputVariableAction(URigVMController* InController, URigVMPin* InPin, bool bNewDefinedAsExternalVariableState)
 : FRigVMBaseAction(InController)
-, NodePath()
-, InnerVariableName(InInnerVariableName)
-, OldOuterVariableName(InOldOuterVariableName)
-, NewOuterVariableName(InNewOuterVariableName)
+, PinPath(InPin->GetName())
+, OldState(InPin->IsDefinedAsInputVariable())
+, NewState(bNewDefinedAsExternalVariableState)
 {
-	if(InFunctionRefNode)
-	{
-		NodePath = InFunctionRefNode->GetName();
-	}
 }
 
-bool FRigVMSetRemappedVariableAction::Undo()
+bool FRigVMMarkExposedPinAsInputVariableAction::Undo()
 {
 	if (!FRigVMBaseAction::Undo())
 	{
 		return false;
 	}
-	if (URigVMFunctionReferenceNode* Node = Cast<URigVMFunctionReferenceNode>(GetController()->GetGraph()->FindNode(NodePath)))
-	{
-		return GetController()->SetRemappedVariable(Node, InnerVariableName, OldOuterVariableName, false);
-	}
-	return false;
+	return GetController()->MarkExposedPinAsInputVariable(*PinPath, OldState, false);
 }
 
-bool FRigVMSetRemappedVariableAction::Redo()
+bool FRigVMMarkExposedPinAsInputVariableAction::Redo()
 {
 	if(!CanUndoRedo())
 	{
 		return false;
 	}
-	if (URigVMFunctionReferenceNode* Node = Cast<URigVMFunctionReferenceNode>(GetController()->GetGraph()->FindNode(NodePath)))
+	if (!GetController()->MarkExposedPinAsInputVariable(*PinPath, NewState, false))
 	{
-		return GetController()->SetRemappedVariable(Node, InnerVariableName, NewOuterVariableName, false);
+		return false;
 	}
-	return false;
+	return FRigVMBaseAction::Redo();
 }
 
 FRigVMAddLocalVariableAction::FRigVMAddLocalVariableAction()
diff --git a/Engine/Plugins/Runtime/RigVM/Source/RigVMDeveloper/Private/RigVMModel/RigVMGraph.cpp b/Engine/Plugins/Runtime/RigVM/Source/RigVMDeveloper/Private/RigVMModel/RigVMGraph.cpp
--- a/Engine/Plugins/Runtime/RigVM/Source/RigVMDeveloper/Private/RigVMModel/RigVMGraph.cpp
+++ b/Engine/Plugins/Runtime/RigVM/Source/RigVMDeveloper/Private/RigVMModel/RigVMGraph.cpp
@@ -34,6 +34,9 @@ void URigVMGraph::PostLoad()
 {
 	Super::PostLoad();
 
+	Nodes.Remove(nullptr);
+	Links.Remove(nullptr);
+
 	// update the bookkeeping on the links
 	for(URigVMLink* Link : Links)
 	{
@@ -355,17 +358,25 @@ TArray<FRigVMExternalVariable> URigVMGraph::GetExternalVariables() const
 {
 	TArray<FRigVMExternalVariable> Variables;
 	
-	for(URigVMNode* Node : GetNodes())
+	for(const URigVMNode* Node : GetNodes())
 	{
-		if(URigVMLibraryNode* LibraryNode = Cast<URigVMLibraryNode>(Node))
+		if(const URigVMFunctionReferenceNode* FunctionReferenceNode = Cast<URigVMFunctionReferenceNode>(Node))
+		{
+			const TArray<FRigVMExternalVariable> FuncRefVariables = FunctionReferenceNode->GetExternalVariables();
+			for(const FRigVMExternalVariable& FuncRefVariable : FuncRefVariables)
+			{
+				FRigVMExternalVariable::MergeExternalVariable(Variables, FuncRefVariable);
+			}
+		}
+		else if(const URigVMCollapseNode* CollapseNode = Cast<URigVMCollapseNode>(Node))
 		{
-			TArray<FRigVMExternalVariable> LibraryVariables = LibraryNode->GetExternalVariables();
-			for(const FRigVMExternalVariable& LibraryVariable : LibraryVariables)
+			const TArray<FRigVMExternalVariable> ContainedVariables = CollapseNode->GetExternalVariables();
+			for(const FRigVMExternalVariable& ContainedVariable : ContainedVariables)
 			{
-				FRigVMExternalVariable::MergeExternalVariable(Variables, LibraryVariable);
+				FRigVMExternalVariable::MergeExternalVariable(Variables, ContainedVariable);
 			}
 		}
-		else if(URigVMVariableNode* VariableNode = Cast<URigVMVariableNode>(Node))
+		else if(const URigVMVariableNode* VariableNode = Cast<URigVMVariableNode>(Node))
 		{
 			// Make sure it is not a local variable or input argument
 			if (VariableNode->IsExternalVariable())
@@ -402,6 +413,10 @@ TArray<FRigVMGraphVariableDescription> URigVMGraph::GetInputArguments() const
 			{
 				continue;
 			}
+			if (Pin->IsDefinedAsInputVariable())
+			{
+				continue;
+			}
 			FRigVMGraphVariableDescription Description;
 			Description.Name = Pin->GetFName();
 			Description.CPPType = Pin->GetCPPType();
@@ -441,6 +456,23 @@ TArray<FRigVMGraphVariableDescription> URigVMGraph::GetOutputArguments() const
 	return Outputs;
 }
 
+TArray<FRigVMGraphVariableDescription> URigVMGraph::GetInputVariables() const
+{
+	TArray<FRigVMGraphVariableDescription> InputVariables;
+	if (URigVMFunctionEntryNode* EntryNode = GetEntryNode())
+	{
+		for (URigVMPin* Pin : EntryNode->GetPins())
+		{
+			if (!Pin->IsDefinedAsInputVariable())
+			{
+				continue;
+			}
+			InputVariables.Add(Pin->GetInputVariableDescription());
+		}
+	}
+	return InputVariables;
+}
+
 URigVMSchema* URigVMGraph::GetSchema() const
 {
 	return SchemaClass ? SchemaClass->GetDefaultObject<URigVMSchema>() : nullptr;
diff --git a/Engine/Plugins/Runtime/RigVM/Source/RigVMDeveloper/Private/RigVMModel/RigVMPin.cpp b/Engine/Plugins/Runtime/RigVM/Source/RigVMDeveloper/Private/RigVMModel/RigVMPin.cpp
--- a/Engine/Plugins/Runtime/RigVM/Source/RigVMDeveloper/Private/RigVMModel/RigVMPin.cpp
+++ b/Engine/Plugins/Runtime/RigVM/Source/RigVMDeveloper/Private/RigVMModel/RigVMPin.cpp
@@ -24,6 +24,7 @@
 #include "RigVMModel/Nodes/RigVMFunctionReturnNode.h"
 #include "RigVMModel/Nodes/RigVMInvokeEntryNode.h"
 #include "RigVMModel/Nodes/RigVMSelectNode.h"
+#include "RigVMModel/Nodes/RigVMVariableNode.h"
 #include "RigVMStringUtils.h"
 #include "RigVMDeveloperTypeUtils.h"
 
@@ -114,6 +115,7 @@ URigVMPin::URigVMPin()
 	: Direction(ERigVMPinDirection::Invalid)
 	, bIsExpanded(false)
 	, bIsConstant(false)
+	, bIsInputVariable(false)
 	, bRequiresWatch(false)
 	, bIsDynamicArray(false)
 	, bIsLazy(false)
@@ -502,6 +504,35 @@ bool URigVMPin::IsDefinedAsConstant() const
 	return bIsConstant;
 }
 
+bool URigVMPin::IsDefinedAsInputVariable() const
+{
+	if (!IsRootPin())
+	{
+		return false;
+	}
+	if (IsDefinedAsConstant())
+	{
+		return false;
+	}
+	return bIsInputVariable;
+}
+
+FRigVMGraphVariableDescription URigVMPin::GetInputVariableDescription() const
+{
+	FRigVMGraphVariableDescription Description;
+	if (IsDefinedAsInputVariable())
+	{
+		Description.Name = GetFName();
+		Description.CPPType = GetCPPType();
+		Description.CPPTypeObject = GetCPPTypeObject();
+		if (Description.CPPTypeObject)
+		{
+			Description.CPPTypeObjectPath = *Description.CPPTypeObject->GetPathName();
+		}
+	}
+	return Description;
+}
+
 bool URigVMPin::RequiresWatch(const bool bCheckExposedPinChain) const
 {
 	if (!bRequiresWatch && bCheckExposedPinChain)
@@ -852,6 +883,10 @@ bool URigVMPin::ShouldHideSubPins() const
 				}
 			}
 		}
+		if (IsDefinedAsInputVariable() && GetDirection() == ERigVMPinDirection::Input)
+		{
+			return true;
+		}
 	}
 #endif
 
@@ -1538,6 +1573,22 @@ URigVMVariableNode* URigVMPin::GetBoundVariableNode() const
 			return VariableNode;
 		}
 	}
+	if (IsDefinedAsInputVariable() && Links.Num() == 1)
+	{
+		if (Links[0]->GetTargetPin() == this)
+		{
+			if (const URigVMPin* SourcePin = Links[0]->GetSourcePin())
+			{
+				if (URigVMVariableNode* VariableNode = Cast<URigVMVariableNode>(SourcePin->GetNode()))
+				{
+					if (VariableNode->GetValuePin() == SourcePin->GetRootPin())
+					{
+						return VariableNode;
+					}
+				}
+			}
+		}
+	}
 	
 	return nullptr;
 }
@@ -1560,15 +1611,44 @@ const FString URigVMPin::GetBoundVariablePath(const URigVMPin::FPinOverride& InO
 	{
 		if (URigVMVariableNode* VariableNode = Cast<URigVMVariableNode>(InjectionInfo->Node))
 		{
-			FString SegmentPath = InjectionInfo->OutputPin->GetSegmentPath(false);
+			FString SegmentPath = InjectionInfo->GetPin()->GetSegmentPath(false);
 			if (SegmentPath.IsEmpty())
 			{
 				return VariableNode->GetVariableName().ToString();
 			}
-			
 			return VariableNode->GetVariableName().ToString() + TEXT(".") + SegmentPath;
 		}
 	}
+	if (IsDefinedAsInputVariable() && Links.Num() == 1)
+	{
+		if (Links[0]->GetTargetPin() == this)
+		{
+			if (const URigVMPin* SourcePin = Links[0]->GetSourcePin())
+			{
+				if (URigVMVariableNode* VariableNode = Cast<URigVMVariableNode>(SourcePin->GetNode()))
+				{
+					if (VariableNode->GetValuePin() == SourcePin->GetRootPin())
+					{
+						const FString SegmentPath = SourcePin->GetSegmentPath(false);
+						if (SegmentPath.IsEmpty())
+						{
+							return VariableNode->GetVariableName().ToString();
+						}
+						return VariableNode->GetVariableName().ToString() + TEXT(".") + SegmentPath;
+					}
+				}
+				if (SourcePin->IsDefinedAsInputVariable() && Cast<URigVMFunctionEntryNode>(SourcePin->GetNode()))
+				{
+					const FString SegmentPath = SourcePin->GetSegmentPath(false);
+					if (SegmentPath.IsEmpty())
+					{
+						return SourcePin->GetName();
+					}
+					return SourcePin->GetName() + TEXT(".") + SegmentPath;
+				}
+			}
+		}
+	}
 	
 	return FString();
 }
@@ -1580,6 +1660,19 @@ FString URigVMPin::GetBoundVariableName() const
 	{
 		return VariableNode->GetVariableName().ToString();
 	}
+	if (IsDefinedAsInputVariable() && Links.Num() == 1)
+	{
+		if (Links[0]->GetTargetPin() == this)
+		{
+			if (const URigVMPin* SourcePin = Links[0]->GetSourcePin())
+			{
+				if (SourcePin->IsDefinedAsInputVariable() && Cast<URigVMFunctionEntryNode>(SourcePin->GetNode()))
+				{
+					return SourcePin->GetName();
+				}
+			}
+		}
+	}
 
 	return FString();
 }
@@ -1646,9 +1739,9 @@ bool URigVMPin::IsBoundToInputArgument() const
 
 	if (URigVMFunctionEntryNode* EntryNode = GetGraph()->GetEntryNode())
 	{
-		if (EntryNode->FindPin(VariableName))
+		if (const URigVMPin* InterfacePin = EntryNode->FindPin(VariableName))
 		{
-			return true;
+			return !InterfacePin->IsDefinedAsInputVariable();
 		}
 	}
 	
@@ -1667,6 +1760,11 @@ bool URigVMPin::CanBeBoundToVariable(const FRigVMExternalVariable& InExternalVar
 		return false;
 	}
 
+	if (bIsInputVariable && InExternalVariable.IsReadOnly())
+	{
+		return false;
+	}
+
 	// only allow to bind variables to input pins for now
 	if (Direction == ERigVMPinDirection::Output)
 	{
@@ -1679,19 +1777,19 @@ bool URigVMPin::CanBeBoundToVariable(const FRigVMExternalVariable& InExternalVar
 	{
 		return false;
 	}
-	if (IsArray() != InExternalVariable.bIsArray)
+	if (IsArray() != InExternalVariable.IsArray())
 	{
 		return false;
 	}
 
-	FName ExternalCPPType = InExternalVariable.TypeName;
-	UObject* ExternalCPPTypeObject = InExternalVariable.TypeObject;
+	FName ExternalCPPType = InExternalVariable.GetExtendedCPPType();
+	UObject* ExternalCPPTypeObject = const_cast<UObject*>(InExternalVariable.GetCPPTypeObject());
 
 	if(!InSegmentPath.IsEmpty())
 	{
-		check(InExternalVariable.Property);
+		check(InExternalVariable.GetProperty());
 
-		const FProperty* Property = InExternalVariable.Property;
+		const FProperty* Property = InExternalVariable.GetProperty();
 		const FRigVMPropertyPath PropertyPath(Property, InSegmentPath);
 		Property = PropertyPath.GetTailProperty();
 
@@ -1699,7 +1797,7 @@ bool URigVMPin::CanBeBoundToVariable(const FRigVMExternalVariable& InExternalVar
 	}
 
 	const FString CPPBaseType = IsArray() ? GetArrayElementCppType() : GetCPPType();
-	return RigVMTypeUtils::AreCompatible(*CPPBaseType, GetCPPTypeObject(), ExternalCPPType, ExternalCPPTypeObject);
+	return RigVMTypeUtils::AreCompatible(ExternalCPPType, ExternalCPPTypeObject, *CPPBaseType, GetCPPTypeObject());
 }
 
 bool URigVMPin::ShowInDetailsPanelOnly() const
@@ -2508,6 +2606,49 @@ bool URigVMPin::CanLink(const URigVMPin* InSourcePin, const URigVMPin* InTargetP
 		return false;
 	}
 
+	if (InTargetPin->IsDefinedAsInputVariable())
+	{
+		if (!InSourcePin->IsDefinedAsInputVariable())
+		{
+			if (!InSourcePin->IsRootPin())
+			{
+				if (OutFailureReason)
+				{
+					*OutFailureReason = TEXT("Variable / Ref pins can only be connected to other root variable pins.");
+				}
+				return false;
+			}
+
+			const URigVMVariableNode* SourceVariableNode = Cast<URigVMVariableNode>(SourceNode);
+			if (!SourceVariableNode)
+			{
+				if (OutFailureReason)
+				{
+					*OutFailureReason = TEXT("Variable / Ref pins can only be connected to other root variable pins.");
+				}
+				return false;
+			}
+
+			if (SourceVariableNode->GetValuePin() != InSourcePin)
+			{
+				if (OutFailureReason)
+				{
+					*OutFailureReason = TEXT("Variable / Ref pins can only be connected to other root variable pins.");
+				}
+				return false;
+			}
+
+			if (SourceVariableNode->IsInputArgument())
+			{
+				if (OutFailureReason)
+				{
+					*OutFailureReason = TEXT("Variable / Ref pins cannot be connected to input arguments.");
+				}
+				return false;
+			}
+		}
+	}
+
 	if (!bInAllowNonArgumentPins)
 	{
 		if (const URigVMTemplateNode* SourceTemplateNode = Cast<URigVMTemplateNode>(SourceNode))
diff --git a/Engine/Plugins/Runtime/RigVM/Source/RigVMDeveloper/Private/RigVMModel/RigVMSchema.cpp b/Engine/Plugins/Runtime/RigVM/Source/RigVMDeveloper/Private/RigVMModel/RigVMSchema.cpp
--- a/Engine/Plugins/Runtime/RigVM/Source/RigVMDeveloper/Private/RigVMModel/RigVMSchema.cpp
+++ b/Engine/Plugins/Runtime/RigVM/Source/RigVMDeveloper/Private/RigVMModel/RigVMSchema.cpp
@@ -294,7 +294,7 @@ bool URigVMSchema::SupportsExternalVariable_NoLock(URigVMController* InControlle
 
 	if(InExternalVariable)
 	{
-		TRigVMTypeIndex Type = InRegistry->GetTypeIndex_NoLock(InExternalVariable->GetExtendedCPPType(), InExternalVariable->TypeObject);
+		TRigVMTypeIndex Type = InRegistry->GetTypeIndex_NoLock(InExternalVariable->GetExtendedCPPType(), const_cast<UObject*>(InExternalVariable->GetCPPTypeObject()));
 		if(Type == INDEX_NONE)
 		{
 			return false;
@@ -460,7 +460,7 @@ bool URigVMSchema::CanAddNode(URigVMController* InController, const URigVMNode*
 				TArray<FRigVMExternalVariable> AllVariables = InController->GetAllVariables(true);
 				for(const FRigVMExternalVariable& Variable : AllVariables)
 				{
-					if(Variable.Name.IsEqual(VarName, ENameCase::CaseSensitive))
+					if(Variable.GetName().IsEqual(VarName, ENameCase::CaseSensitive))
 					{
 						return true;
 					}
@@ -745,6 +745,10 @@ bool URigVMSchema::CanUnfoldPin(URigVMController* InController, const URigVMPin*
 {
 	RIGVMSCHEMA_DEFAULT_FUNCTION_BODY
 
+	if (InPinToUnfold->ShouldHideSubPins())
+	{
+		return false;
+	}
 	if (InPinToUnfold->IsStruct())
 	{
 		return ShouldUnfoldStruct(InController, InPinToUnfold->GetScriptStruct());
diff --git a/Engine/Plugins/Runtime/RigVM/Source/RigVMDeveloper/Public/RigVMCompiler/RigVMAST.h b/Engine/Plugins/Runtime/RigVM/Source/RigVMDeveloper/Public/RigVMCompiler/RigVMAST.h
--- a/Engine/Plugins/Runtime/RigVM/Source/RigVMDeveloper/Public/RigVMCompiler/RigVMAST.h
+++ b/Engine/Plugins/Runtime/RigVM/Source/RigVMDeveloper/Public/RigVMCompiler/RigVMAST.h
@@ -23,7 +23,6 @@ class FRigVMInlineFunctionExprAST;
 class FRigVMNoOpExprAST;
 class FRigVMVarExprAST;
 class FRigVMLiteralExprAST;
-class FRigVMExternalVarExprAST;
 class FRigVMAssignExprAST;
 class FRigVMCopyExprAST;
 class FRigVMCachedValueExprAST;
@@ -90,7 +89,6 @@ public:
 		NoOp,
 		Var,
 		Literal,
-		ExternalVar,
 		Assign,
 		Copy,
 		CachedValue,
@@ -462,17 +460,6 @@ inline const FRigVMLiteralExprAST* FRigVMExprAST::To() const
 	return (const FRigVMLiteralExprAST*)this;
 }
 
-// specialized cast for type checking
-// for a External Variable / FRigVMExternalVarExprAST expression
-// will raise if types are not compatible
-// @return this expression cast to FRigVMExternalVarExprAST
-template<>
-inline const FRigVMExternalVarExprAST* FRigVMExprAST::To() const
-{
-	ensure(IsA(EType::ExternalVar));
-	return (const FRigVMExternalVarExprAST*)this;
-}
-
 // specialized cast for type checking
 // for a Assign / FRigVMAssignExprAST expression
 // will raise if types are not compatible
@@ -875,6 +862,12 @@ public:
 	// returns true if this variable allows links to be "soft", so without a cache / computed value.
 	UE_API bool SupportsSoftLinks() const;
 
+	// Returns the clone this variable expression is using.
+	const FRigVMVarExprAST* GetVariableCloneExpr() const
+	{
+		return ClonedVarExpr;
+	}
+
 protected:
 
 	// default constructor (protected so that only parser can access it)
@@ -887,6 +880,7 @@ protected:
 private:
 
 	FRigVMASTProxy Proxy;
+	const FRigVMVarExprAST* ClonedVarExpr = nullptr;
 
 	friend class FRigVMParserAST;
 	friend class URigVMCompiler;
@@ -935,42 +929,6 @@ private:
 	friend class FRigVMParserAST;
 };
 
-/*
- * An abstract syntax tree external variable expression represents the reference 
- * of unowned / external memory.
- */
-class FRigVMExternalVarExprAST : public FRigVMVarExprAST
-{
-public:
-
-	// virtual destructor
-	virtual ~FRigVMExternalVarExprAST() = default;
-
-	// disable copy constructor
-	FRigVMExternalVarExprAST(const FRigVMExternalVarExprAST&) = delete;
-
-	// overload of the type checking mechanism
-	virtual bool IsA(EType InType) const override
-	{
-		if (FRigVMVarExprAST::IsA(InType))
-		{
-			return true;
-		}
-		return InType == EType::ExternalVar;
-	};
-
-protected:
-
-	// default constructor (protected so that only parser can access it)
-	FRigVMExternalVarExprAST(const FRigVMASTProxy& InPinProxy)
-		: FRigVMVarExprAST(EType::ExternalVar, InPinProxy)
-	{}
-
-private:
-
-	friend class FRigVMParserAST;
-};
-
 /*
  * An abstract syntax tree assign expression represents the assignment of one
  * expression to another. This can result in referencing memory from a to b or
@@ -1438,6 +1396,8 @@ private:
 	// traverse a single link (constructs assign + copy expressions)
 	UE_API FRigVMExprAST* TraverseLink(int32 InLinkIndex, FRigVMExprAST* InParentExpr);
 
+	UE_API bool TraverseNodeExpressionForRequiredVariableClone(FRigVMNodeExprAST* InNodeExpr);
+
 	UE_API bool ShouldLinkBeSkipped(const FRigVMASTLinkDescription& InLink) const;
 
 	static UE_API FString GetLinkAsString(const FRigVMASTLinkDescription& InLink);
diff --git a/Engine/Plugins/Runtime/RigVM/Source/RigVMDeveloper/Public/RigVMCompiler/RigVMCompiler.h b/Engine/Plugins/Runtime/RigVM/Source/RigVMDeveloper/Public/RigVMCompiler/RigVMCompiler.h
--- a/Engine/Plugins/Runtime/RigVM/Source/RigVMDeveloper/Public/RigVMCompiler/RigVMCompiler.h
+++ b/Engine/Plugins/Runtime/RigVM/Source/RigVMDeveloper/Public/RigVMCompiler/RigVMCompiler.h
@@ -362,7 +362,6 @@ private:
 	UE_API bool TraverseNoOp(const FRigVMNoOpExprAST* InExpr, FRigVMCompilerWorkData& WorkData);
 	UE_API bool TraverseVar(const FRigVMVarExprAST* InExpr, FRigVMCompilerWorkData& WorkData);
 	UE_API bool TraverseLiteral(const FRigVMVarExprAST* InExpr, FRigVMCompilerWorkData& WorkData);
-	UE_API bool TraverseExternalVar(const FRigVMExternalVarExprAST* InExpr, FRigVMCompilerWorkData& WorkData);
 	UE_API bool TraverseAssign(const FRigVMAssignExprAST* InExpr, FRigVMCompilerWorkData& WorkData);
 	UE_API bool TraverseCopy(const FRigVMCopyExprAST* InExpr, FRigVMCompilerWorkData& WorkData);
 	UE_API bool TraverseCachedValue(const FRigVMCachedValueExprAST* InExpr, FRigVMCompilerWorkData& WorkData);
diff --git a/Engine/Plugins/Runtime/RigVM/Source/RigVMDeveloper/Public/RigVMModel/Nodes/RigVMCollapseNode.h b/Engine/Plugins/Runtime/RigVM/Source/RigVMDeveloper/Public/RigVMModel/Nodes/RigVMCollapseNode.h
--- a/Engine/Plugins/Runtime/RigVM/Source/RigVMDeveloper/Public/RigVMModel/Nodes/RigVMCollapseNode.h
+++ b/Engine/Plugins/Runtime/RigVM/Source/RigVMDeveloper/Public/RigVMModel/Nodes/RigVMCollapseNode.h
@@ -29,7 +29,9 @@ public:
 	virtual FString GetNodeKeywords() const override { return NodeKeywords; }
 	virtual FString GetNodeDescription() const override { return NodeDescription; }
 	UE_API virtual URigVMFunctionLibrary* GetLibrary() const override;
+	bool IsGraphFunctionDefinition() const { return GetLibrary() != nullptr;}
 	virtual URigVMGraph* GetContainedGraph() const override { return ContainedGraph; }
+	UE_API virtual TArray<FRigVMExternalVariable> GetExternalVariables() const override;
 
 	UE_API FString GetEditorSubGraphName() const;
 
diff --git a/Engine/Plugins/Runtime/RigVM/Source/RigVMDeveloper/Public/RigVMModel/Nodes/RigVMFunctionReferenceNode.h b/Engine/Plugins/Runtime/RigVM/Source/RigVMDeveloper/Public/RigVMModel/Nodes/RigVMFunctionReferenceNode.h
--- a/Engine/Plugins/Runtime/RigVM/Source/RigVMDeveloper/Public/RigVMModel/Nodes/RigVMFunctionReferenceNode.h
+++ b/Engine/Plugins/Runtime/RigVM/Source/RigVMDeveloper/Public/RigVMModel/Nodes/RigVMFunctionReferenceNode.h
@@ -47,11 +47,12 @@ public:
 	UE_API URigVMLibraryNode* LoadReferencedNode() const;
 
 	// Variable remapping
+	UE_API bool ReferencesFunctionInSamePackage() const;
 	UE_API bool RequiresVariableRemapping() const;
 	UE_API bool IsFullyRemapped() const;
-	UE_API TArray<FRigVMExternalVariable> GetExternalVariables(bool bRemapped) const;
-	const TMap<FName, FName>& GetVariableMap() const { return VariableMap; }
+	UE_API const TMap<FName, FName> GetVariableMap() const;
 	UE_API FName GetOuterVariableName(const FName& InInnerVariableName) const;
+	UE_API FString GetOuterVariablePath(const FName& InInnerVariableName) const;
 	// end Variable remapping
 
 	UE_API virtual uint32 GetStructureHash() const override;
@@ -70,10 +71,10 @@ public:
 protected:
 
 	UE_API virtual FString GetOriginalDefaultValueForRootPin(const URigVMPin* InRootPin) const override;
+	UE_API TArray<FRigVMExternalVariable> GetAssetExternalVariables() const;
 	
 private:
 
-	UE_API bool RequiresVariableRemappingInternal(TArray<FRigVMExternalVariable>& InnerVariables) const;
 	UE_API virtual TArray<int32> GetInstructionsForVMImpl(const FRigVMExtendedExecuteContext& Context, URigVM* InVM, const FRigVMASTProxy& InProxy = FRigVMASTProxy()) const override;
 	UE_API const URigVMPin* FindReferencedPin(const URigVMPin* InPin) const;
 	UE_API const URigVMPin* FindReferencedPin(const FString& InPinPath) const;
@@ -87,7 +88,7 @@ private:
 	TSoftObjectPtr<URigVMLibraryNode> ReferencedNodePtr_DEPRECATED;
 
 	UPROPERTY()
-	TMap<FName, FName> VariableMap;
+	TMap<FName, FName> VariableMap_DEPRECATED;
 
 	friend class URigVMController;
 	friend class FRigVMParserAST;
diff --git a/Engine/Plugins/Runtime/RigVM/Source/RigVMDeveloper/Public/RigVMModel/Nodes/RigVMLibraryNode.h b/Engine/Plugins/Runtime/RigVM/Source/RigVMDeveloper/Public/RigVMModel/Nodes/RigVMLibraryNode.h
--- a/Engine/Plugins/Runtime/RigVM/Source/RigVMDeveloper/Public/RigVMModel/Nodes/RigVMLibraryNode.h
+++ b/Engine/Plugins/Runtime/RigVM/Source/RigVMDeveloper/Public/RigVMModel/Nodes/RigVMLibraryNode.h
@@ -72,6 +72,9 @@ public:
 	UFUNCTION(BlueprintPure, Category = RigVMLibraryNode)
 	UE_API TArray<FRigVMVariantRef> GetMatchingVariants() const;
 
+	UFUNCTION(BlueprintPure, Category = RigVMLibraryNode)
+	bool ContainsLocalVariable(const FName& InVariableName) const;
+
 protected:
 
 	UE_API virtual TArray<int32> GetInstructionsForVMImpl(const FRigVMExtendedExecuteContext& Context, URigVM* InVM, const FRigVMASTProxy& InProxy = FRigVMASTProxy()) const override; 
diff --git a/Engine/Plugins/Runtime/RigVM/Source/RigVMDeveloper/Public/RigVMModel/Nodes/RigVMParameterNode.h b/Engine/Plugins/Runtime/RigVM/Source/RigVMDeveloper/Public/RigVMModel/Nodes/RigVMParameterNode.h
--- a/Engine/Plugins/Runtime/RigVM/Source/RigVMDeveloper/Public/RigVMModel/Nodes/RigVMParameterNode.h
+++ b/Engine/Plugins/Runtime/RigVM/Source/RigVMDeveloper/Public/RigVMModel/Nodes/RigVMParameterNode.h
@@ -61,26 +61,7 @@ public:
 	// Returns nullptr external variable matching this description
 	FRigVMExternalVariable ToExternalVariable() const
 	{
-		FRigVMExternalVariable ExternalVariable;
-		ExternalVariable.Name = Name;
-
-		if (RigVMTypeUtils::IsArrayType(CPPType))
-		{
-			ExternalVariable.bIsArray = true;
-			ExternalVariable.TypeName = *CPPType.Mid(7, CPPType.Len() - 8);
-			ExternalVariable.TypeObject = CPPTypeObject;
-		}
-		else
-		{
-			ExternalVariable.bIsArray = false;
-			ExternalVariable.TypeName = *CPPType;
-			ExternalVariable.TypeObject = CPPTypeObject;
-		}
-
-		ExternalVariable.bIsPublic = false;
-		ExternalVariable.bIsReadOnly = false;
-		ExternalVariable.Memory = nullptr;
-		return ExternalVariable;
+		return FRigVMExternalVariable::Make(Name, CPPType, CPPTypeObject, false, false);
 	}
 };
 
diff --git a/Engine/Plugins/Runtime/RigVM/Source/RigVMDeveloper/Public/RigVMModel/Nodes/RigVMVariableNode.h b/Engine/Plugins/Runtime/RigVM/Source/RigVMDeveloper/Public/RigVMModel/Nodes/RigVMVariableNode.h
--- a/Engine/Plugins/Runtime/RigVM/Source/RigVMDeveloper/Public/RigVMModel/Nodes/RigVMVariableNode.h
+++ b/Engine/Plugins/Runtime/RigVM/Source/RigVMDeveloper/Public/RigVMModel/Nodes/RigVMVariableNode.h
@@ -48,6 +48,10 @@ public:
 	UFUNCTION(BlueprintCallable, Category = RigVMVariableNode)
 	UE_API bool IsInputArgument() const;
 
+	// Returns true if this variable is related to an input variable
+	UFUNCTION(BlueprintCallable, Category = RigVMVariableNode)
+	UE_API bool IsInputVariable() const;
+
 	// Returns the C++ data type of the variable
 	UFUNCTION(BlueprintCallable, Category = RigVMVariableNode)
 	UE_API FString GetCPPType() const;
diff --git a/Engine/Plugins/Runtime/RigVM/Source/RigVMDeveloper/Public/RigVMModel/RigVMController.h b/Engine/Plugins/Runtime/RigVM/Source/RigVMDeveloper/Public/RigVMModel/RigVMController.h
--- a/Engine/Plugins/Runtime/RigVM/Source/RigVMDeveloper/Public/RigVMModel/RigVMController.h
+++ b/Engine/Plugins/Runtime/RigVM/Source/RigVMDeveloper/Public/RigVMModel/RigVMController.h
@@ -192,12 +192,18 @@ struct FRigVMPinInfo
 	FString CustomWidgetName;
 	bool bIsExpanded;
 	bool bIsConstant;
+	bool bIsInputVariable;
 	bool bIsDynamicArray;
 	bool bIsTrait;
 	bool bIsLazy;
 	TArray<int32> SubPins;
 
 	friend uint32 GetTypeHash(const FRigVMPinInfo& InPin);
+
+	bool ShouldHideSubPins() const
+	{
+		return bIsInputVariable && Direction == ERigVMPinDirection::Input;
+	}
 };
 
 struct FRigVMPinInfoArray
@@ -1033,7 +1039,7 @@ public:
 
 	// Adds an exposed pin to the graph controlled by this
 	UFUNCTION(BlueprintCallable, Category = RigVMController)
-	UE_API FName AddExposedPin(const FName& InPinName, ERigVMPinDirection InDirection, const FString& InCPPType, const FName& InCPPTypeObjectPath, const FString& InDefaultValue, bool bSetupUndoRedo = true, bool bPrintPythonCommand = false);
+	UE_API FName AddExposedPin(const FName& InPinName, ERigVMPinDirection InDirection, const FString& InCPPType, const FName& InCPPTypeObjectPath, const FString& InDefaultValue, bool bSetupUndoRedo = true, bool bPrintPythonCommand = false, bool bIsInputVariable = false);
 
 	// Removes an exposed pin from the graph controlled by this
 	UFUNCTION(BlueprintCallable, Category = RigVMController)
@@ -1051,6 +1057,10 @@ public:
 	UFUNCTION(BlueprintCallable, Category = RigVMController)
 	UE_API bool SetExposedPinIndex(const FName& InPinName, int32 InNewIndex, bool bSetupUndoRedo = true, bool bPrintPythonCommand = false);
 
+	// Marks an exposed pin as defined as an input variable
+	UFUNCTION(BlueprintCallable, Category = RigVMController)
+	UE_API bool MarkExposedPinAsInputVariable(const FName& InPinName, bool bDefinedAsExternalVariable, bool bSetupUndoRedo = true, bool bPrintPythonCommand = false);
+	
 	UFUNCTION(BlueprintPure, Category = RigVMController)
 	UE_API FRigVMGraphFunctionHeader FindGraphFunctionHeaderByName(FString InHostPath, FName InFunctionName) const;
 
@@ -1080,7 +1090,8 @@ public:
 	UE_API bool SwapAllFunctionReferences(const FRigVMGraphFunctionIdentifier& InOldFunctionIdentifier, const FRigVMGraphFunctionIdentifier& InNewFunctionIdentifier, bool bSetupOrphanPins, bool bSetupUndoRedo = true, bool bPrintPythonCommand = false);
 
 	// Sets the remapped variable on a function reference node
-	UFUNCTION(BlueprintCallable, Category = RigVMController)
+	UE_DEPRECATED(5.8, "Please use BindPinToVariable instead.")
+	UFUNCTION(BlueprintCallable, Category = RigVMController, meta=(DeprecatedFunction, DeprecationMessage="Function has been deprecated, please use BindPinToVariable instead."))
     UE_API bool SetRemappedVariable(URigVMFunctionReferenceNode* InFunctionRefNode, const FName& InInnerVariableName, const FName& InOuterVariableName, bool bSetupUndoRedo = true);
 
 	// Adds a function definition to a function library graph
@@ -1560,6 +1571,7 @@ protected:
 	UE_API FRigVMClientPatchResult PatchPinDefaultValues();
 	UE_API FRigVMClientPatchResult PatchUserDefinedStructPinNames();
 	UE_API FRigVMClientPatchResult PatchLocalVariableTypes();
+	UE_API FRigVMClientPatchResult PatchExternalVariablesToArguments();
 
 	UE_API ERigVMPinDefaultValueType GetDefaultValueType(const URigVMPin* InPin, const FString& InDefaultValue) const;
 
diff --git a/Engine/Plugins/Runtime/RigVM/Source/RigVMDeveloper/Public/RigVMModel/RigVMControllerActions.h b/Engine/Plugins/Runtime/RigVM/Source/RigVMDeveloper/Public/RigVMModel/RigVMControllerActions.h
--- a/Engine/Plugins/Runtime/RigVM/Source/RigVMDeveloper/Public/RigVMModel/RigVMControllerActions.h
+++ b/Engine/Plugins/Runtime/RigVM/Source/RigVMDeveloper/Public/RigVMModel/RigVMControllerActions.h
@@ -1401,34 +1401,30 @@ public:
 };
 
 /**
-* An action to remap a variable inside of a function reference node renaming a node in the graph.
-*/
+ * An action to mark an exposed pin as an input variable
+ */
 USTRUCT()
-struct FRigVMSetRemappedVariableAction : public FRigVMBaseAction
+struct FRigVMMarkExposedPinAsInputVariableAction : public FRigVMBaseAction
 {
 	GENERATED_BODY()
 
 public:
 
-	FRigVMSetRemappedVariableAction();
-	FRigVMSetRemappedVariableAction(URigVMController* InController, URigVMFunctionReferenceNode* InFunctionRefNode, const FName& InInnerVariableName,
-		const FName& InOldOuterVariableName, const FName& InNewOuterVariableName);
-	virtual ~FRigVMSetRemappedVariableAction() = default;
-	virtual UScriptStruct* GetScriptStruct() const override { return FRigVMSetRemappedVariableAction::StaticStruct(); }
+	FRigVMMarkExposedPinAsInputVariableAction();
+	FRigVMMarkExposedPinAsInputVariableAction(URigVMController* InController, URigVMPin* InPin, bool bNewDefinedAsExternalVariableState);
+	virtual ~FRigVMMarkExposedPinAsInputVariableAction() = default;
+	virtual UScriptStruct* GetScriptStruct() const override { return FRigVMMarkExposedPinAsInputVariableAction::StaticStruct(); }
 	virtual bool Undo() override;
 	virtual bool Redo() override;
 
 	UPROPERTY()
-	FString NodePath;
-
-	UPROPERTY()
-	FName InnerVariableName;
+	FString PinPath;
 
 	UPROPERTY()
-	FName OldOuterVariableName;
+	bool OldState;
 
 	UPROPERTY()
-	FName NewOuterVariableName;
+	bool NewState;
 };
 
 /**
diff --git a/Engine/Plugins/Runtime/RigVM/Source/RigVMDeveloper/Public/RigVMModel/RigVMGraph.h b/Engine/Plugins/Runtime/RigVM/Source/RigVMDeveloper/Public/RigVMModel/RigVMGraph.h
--- a/Engine/Plugins/Runtime/RigVM/Source/RigVMDeveloper/Public/RigVMModel/RigVMGraph.h
+++ b/Engine/Plugins/Runtime/RigVM/Source/RigVMDeveloper/Public/RigVMModel/RigVMGraph.h
@@ -151,6 +151,10 @@ public:
     UFUNCTION(BlueprintCallable, Category = RigVMGraph)
     UE_API TArray<FRigVMGraphVariableDescription> GetOutputArguments() const;
 
+	// Returns the input variables of this function
+	UFUNCTION(BlueprintCallable, Category = RigVMGraph)
+	UE_API TArray<FRigVMGraphVariableDescription> GetInputVariables() const;
+
 	// Returns the schema used by this graph
 	UFUNCTION(BlueprintCallable, Category = RigVMGraph)
 	UE_API URigVMSchema* GetSchema() const;
@@ -261,6 +265,7 @@ private:
 	friend class FRigVMFunctionArgumentLayout;
 	friend class FRigVMFunctionArgumentReorderDragDropHandler;
 	friend class FRigVMGraphDetailCustomizationImpl;
+	friend class URigVMLibraryNode;
 };
 
 #undef UE_API
diff --git a/Engine/Plugins/Runtime/RigVM/Source/RigVMDeveloper/Public/RigVMModel/RigVMNotifications.h b/Engine/Plugins/Runtime/RigVM/Source/RigVMDeveloper/Public/RigVMModel/RigVMNotifications.h
--- a/Engine/Plugins/Runtime/RigVM/Source/RigVMDeveloper/Public/RigVMModel/RigVMNotifications.h
+++ b/Engine/Plugins/Runtime/RigVM/Source/RigVMDeveloper/Public/RigVMModel/RigVMNotifications.h
@@ -51,7 +51,7 @@ enum class ERigVMGraphNotifType : uint8
 	NodeCategoryChanged, // A node's category has changed (Subject == URigVMNode)
 	NodeKeywordsChanged, // A node's keywords have changed (Subject == URigVMNode)
 	NodeDescriptionChanged, // A node's description has changed (Subject == URigVMNode)
-	VariableRemappingChanged, // A function reference node's remapping has changed (Subject == URigVMFunctionReferenceNode)
+	VariableRemappingChanged, // (DEPRECATED) A function reference node's remapping has changed (Subject == URigVMFunctionReferenceNode)
 	LibraryTemplateChanged, // The definition of a library node's template has changed (Subject == URigVMLibraryNode)
 	FunctionAccessChanged, // The function has been made public/private (Subject == URigVMLibraryNode)
 	VariantTagsChanged, // The tags in the header of this function variant have changed (Subject == URigVMLibraryNode)
diff --git a/Engine/Plugins/Runtime/RigVM/Source/RigVMDeveloper/Public/RigVMModel/RigVMPin.h b/Engine/Plugins/Runtime/RigVM/Source/RigVMDeveloper/Public/RigVMModel/RigVMPin.h
--- a/Engine/Plugins/Runtime/RigVM/Source/RigVMDeveloper/Public/RigVMModel/RigVMPin.h
+++ b/Engine/Plugins/Runtime/RigVM/Source/RigVMDeveloper/Public/RigVMModel/RigVMPin.h
@@ -12,6 +12,7 @@
 #include "UObject/StructOnScope.h"
 #include <RigVMCore/RigVMTrait.h>
 #include "RigVMModelCachedValue.h"
+#include "RigVMCore/RigVMVariableDescription.h"
 #include "RigVMPin.generated.h"
 
 #define UE_API RIGVMDEVELOPER_API
@@ -211,6 +212,13 @@ public:
 	UFUNCTION(BlueprintCallable, Category = RigVMPin)
 	UE_API bool IsDefinedAsConstant() const;
 
+	// Returns true if the pin is defined as an input variable
+	UFUNCTION(BlueprintCallable, Category = RigVMPin)
+	UE_API bool IsDefinedAsInputVariable() const;
+
+	// Returns the description of the input variable
+	UE_API FRigVMGraphVariableDescription GetInputVariableDescription() const; 
+
 	// Returns true if the pin should be watched
 	UFUNCTION(BlueprintCallable, Category = RigVMPin)
 	UE_API bool RequiresWatch(const bool bCheckExposedPinChain = false) const;
@@ -570,6 +578,9 @@ private:
 	UPROPERTY()
 	bool bIsConstant;
 
+	UPROPERTY()
+	bool bIsInputVariable;
+
 	UPROPERTY(transient)
 	bool bRequiresWatch;
 
diff --git a/Engine/Plugins/Runtime/RigVM/Source/RigVMEditor/Private/EdGraph/NodeSpawners/RigVMEdGraphFunctionRefNodeSpawner.cpp b/Engine/Plugins/Runtime/RigVM/Source/RigVMEditor/Private/EdGraph/NodeSpawners/RigVMEdGraphFunctionRefNodeSpawner.cpp
--- a/Engine/Plugins/Runtime/RigVM/Source/RigVMEditor/Private/EdGraph/NodeSpawners/RigVMEdGraphFunctionRefNodeSpawner.cpp
+++ b/Engine/Plugins/Runtime/RigVM/Source/RigVMEditor/Private/EdGraph/NodeSpawners/RigVMEdGraphFunctionRefNodeSpawner.cpp
@@ -311,14 +311,6 @@ URigVMEdGraphNode* URigVMEdGraphFunctionRefNodeSpawner::SpawnNode(UEdGraph* Pare
 				{
 					continue;
 				}
-
-				for(TPair<FName, FName> MappedVariablePair : ExistingFunctionReferenceNode->GetVariableMap())
-				{
-					if(!MappedVariablePair.Value.IsNone())
-					{
-						Controller->SetRemappedVariable(ModelNode, MappedVariablePair.Key, MappedVariablePair.Value, true);
-					}
-				}
 			}
 
 			Controller->CloseUndoBracket();
diff --git a/Engine/Plugins/Runtime/RigVM/Source/RigVMEditor/Private/EdGraph/NodeSpawners/RigVMEdGraphVariableNodeSpawner.cpp b/Engine/Plugins/Runtime/RigVM/Source/RigVMEditor/Private/EdGraph/NodeSpawners/RigVMEdGraphVariableNodeSpawner.cpp
--- a/Engine/Plugins/Runtime/RigVM/Source/RigVMEditor/Private/EdGraph/NodeSpawners/RigVMEdGraphVariableNodeSpawner.cpp
+++ b/Engine/Plugins/Runtime/RigVM/Source/RigVMEditor/Private/EdGraph/NodeSpawners/RigVMEdGraphVariableNodeSpawner.cpp
@@ -87,7 +87,7 @@ bool URigVMEdGraphVariableNodeSpawner::IsTemplateNodeFilteredOut(TArray<UObject*
 
 FString URigVMEdGraphVariableNodeSpawner::GetSpawnerSignature() const
 {
-	return FString("ExternalVariable=" + ExternalVariable.Name.ToString());
+	return FString("ExternalVariable=" + ExternalVariable.GetName().ToString());
 }
 
 URigVMEdGraphNode* URigVMEdGraphVariableNodeSpawner::Invoke(URigVMEdGraph* ParentGraph, FVector2D const Location) const
@@ -112,16 +112,12 @@ URigVMEdGraphNode* URigVMEdGraphVariableNodeSpawner::Invoke(URigVMEdGraph* Paren
 #endif
 
 	FString ObjectPath;
-	if (ExternalVariable.TypeObject)
+	if (ExternalVariable.GetCPPTypeObject())
 	{
-		ObjectPath = ExternalVariable.TypeObject->GetPathName();
+		ObjectPath = ExternalVariable.GetCPPTypeObject()->GetPathName();
 	}
 
-	FString TypeName = ExternalVariable.TypeName.ToString();
-	if (ExternalVariable.bIsArray)
-	{
-		TypeName = FString::Printf(TEXT("TArray<%s>"), *TypeName);
-	}
+	const FString TypeName = ExternalVariable.GetExtendedCPPType().ToString();
 
 	FString NodeName;
 	if(bIsTemplateNode)
@@ -133,7 +129,7 @@ URigVMEdGraphNode* URigVMEdGraphVariableNodeSpawner::Invoke(URigVMEdGraph* Paren
 		static constexpr TCHAR VariableNodeNameFormat[] = TEXT("VariableNode_%s_%s");
 		static const FString GetterPrefix = TEXT("Getter");
 		static const FString SetterPrefix = TEXT("Setter");
-		NodeName = FString::Printf(VariableNodeNameFormat, bIsGetter ? *GetterPrefix : *SetterPrefix, *ExternalVariable.Name.ToString());
+		NodeName = FString::Printf(VariableNodeNameFormat, bIsGetter ? *GetterPrefix : *SetterPrefix, *ExternalVariable.GetName().ToString());
 
 		TArray<FPinInfo> Pins;
 
@@ -147,8 +143,8 @@ URigVMEdGraphNode* URigVMEdGraphVariableNodeSpawner::Invoke(URigVMEdGraph* Paren
 		Pins.Emplace(
 			URigVMVariableNode::ValueName,
 			bIsGetter ? ERigVMPinDirection::Output : ERigVMPinDirection::Input,
-			ExternalVariable.TypeName,
-			ExternalVariable.TypeObject);
+			ExternalVariable.GetExtendedCPPType(),
+			const_cast<UObject*>(ExternalVariable.GetCPPTypeObject()));
 		
 		return SpawnTemplateNode(ParentGraph, Pins, *NodeName);
 	}
@@ -156,7 +152,7 @@ URigVMEdGraphNode* URigVMEdGraphVariableNodeSpawner::Invoke(URigVMEdGraph* Paren
 	URigVMController* Controller = RigBlueprint->GetController(ParentGraph);
 	Controller->OpenUndoBracket(TEXT("Add Variable"));
 
-	if (URigVMNode* ModelNode = Controller->AddVariableNodeFromObjectPath(ExternalVariable.Name, TypeName, ObjectPath, bIsGetter, FString(), Location, NodeName, true, true))
+	if (URigVMNode* ModelNode = Controller->AddVariableNodeFromObjectPath(ExternalVariable.GetName(), TypeName, ObjectPath, bIsGetter, FString(), Location, NodeName, true, true))
 	{
 		for (UEdGraphNode* Node : ParentGraph->Nodes)
 		{
diff --git a/Engine/Plugins/Runtime/RigVM/Source/RigVMEditor/Private/EdGraph/RigVMEdGraphPanelPinFactory.cpp b/Engine/Plugins/Runtime/RigVM/Source/RigVMEditor/Private/EdGraph/RigVMEdGraphPanelPinFactory.cpp
--- a/Engine/Plugins/Runtime/RigVM/Source/RigVMEditor/Private/EdGraph/RigVMEdGraphPanelPinFactory.cpp
+++ b/Engine/Plugins/Runtime/RigVM/Source/RigVMEditor/Private/EdGraph/RigVMEdGraphPanelPinFactory.cpp
@@ -88,7 +88,8 @@ TSharedPtr<SGraphPin> FRigVMEdGraphPanelPinFactory::CreatePin_Internal(UEdGraphP
 			URigVMPin* ModelPin = RigNode->GetModelPinFromPinPath(InPin->GetName());
 			if (ModelPin)
 			{
-				if (ModelPin->IsBoundToVariable())
+				if ((ModelPin->GetDirection() == ERigVMPinDirection::Input || ModelPin->GetDirection() == ERigVMPinDirection::Visible) && 
+					(ModelPin->IsBoundToVariable() || ModelPin->IsDefinedAsInputVariable()))
 				{
 					if (FRigVMEditorAssetInterfacePtr Blueprint = IRigVMEditorAssetInterface::GetInterfaceOuter(RigGraph))
 					{
diff --git a/Engine/Plugins/Runtime/RigVM/Source/RigVMEditor/Private/Editor/RigVMEditor.cpp b/Engine/Plugins/Runtime/RigVM/Source/RigVMEditor/Private/Editor/RigVMEditor.cpp
--- a/Engine/Plugins/Runtime/RigVM/Source/RigVMEditor/Private/Editor/RigVMEditor.cpp
+++ b/Engine/Plugins/Runtime/RigVM/Source/RigVMEditor/Private/Editor/RigVMEditor.cpp
@@ -3835,18 +3835,18 @@ void FRigVMEditorBase::HandleVariableDroppedFromBlueprint(UObject* InSubject, FP
 	}
 
 	FMenuBuilder MenuBuilder(true, NULL);
-	const FText SectionText = FText::FromString(FString::Printf(TEXT("Variable %s"), *ExternalVariable.Name.ToString()));
+	const FText SectionText = FText::FromString(FString::Printf(TEXT("Variable %s"), *ExternalVariable.GetName().ToString()));
 
 	MenuBuilder.BeginSection("VariableDropped", SectionText);
 
 	MenuBuilder.AddMenuEntry(
-		FText::FromString(FString::Printf(TEXT("Get %s"), *ExternalVariable.Name.ToString())),
-		FText::FromString(FString::Printf(TEXT("Adds a getter node for variable %s"), *ExternalVariable.Name.ToString())),
+		FText::FromString(FString::Printf(TEXT("Get %s"), *ExternalVariable.GetName().ToString())),
+		FText::FromString(FString::Printf(TEXT("Adds a getter node for variable %s"), *ExternalVariable.GetName().ToString())),
 		FSlateIcon(),
 		FUIAction(
 			FExecuteAction::CreateLambda([ExternalVariable, Controller, InDropPosition] {
 
-				Controller->AddVariableNode(ExternalVariable.Name, ExternalVariable.TypeName.ToString(), ExternalVariable.TypeObject, true, FString(), InDropPosition, FString(), true, true);
+				Controller->AddVariableNode(ExternalVariable.GetName(), ExternalVariable.GetExtendedCPPType().ToString(), const_cast<UObject*>(ExternalVariable.GetCPPTypeObject()), true, FString(), InDropPosition, FString(), true, true);
 
 			}),
 			FCanExecuteAction()
@@ -3854,13 +3854,13 @@ void FRigVMEditorBase::HandleVariableDroppedFromBlueprint(UObject* InSubject, FP
 	);
 
 	MenuBuilder.AddMenuEntry(
-		FText::FromString(FString::Printf(TEXT("Set %s"), *ExternalVariable.Name.ToString())),
-		FText::FromString(FString::Printf(TEXT("Adds a setter node for variable %s"), *ExternalVariable.Name.ToString())),
+		FText::FromString(FString::Printf(TEXT("Set %s"), *ExternalVariable.GetName().ToString())),
+		FText::FromString(FString::Printf(TEXT("Adds a setter node for variable %s"), *ExternalVariable.GetName().ToString())),
 		FSlateIcon(),
 		FUIAction(
 			FExecuteAction::CreateLambda([ExternalVariable, Controller, InDropPosition] {
 
-				Controller->AddVariableNode(ExternalVariable.Name, ExternalVariable.TypeName.ToString(), ExternalVariable.TypeObject, false, FString(), InDropPosition, FString(), true, true);
+				Controller->AddVariableNode(ExternalVariable.GetName(), ExternalVariable.GetExtendedCPPType().ToString(), const_cast<UObject*>(ExternalVariable.GetCPPTypeObject()), false, FString(), InDropPosition, FString(), true, true);
 
 			}),
 			FCanExecuteAction()
