RigVM: Functions to support input variables

diff --git a/Engine/Plugins/Runtime/RigVM/Source/RigVMEditor/Private/Editor/RigVMGraphDetailCustomization.cpp b/Engine/Plugins/Runtime/RigVM/Source/RigVMEditor/Private/Editor/RigVMGraphDetailCustomization.cpp
--- a/Engine/Plugins/Runtime/RigVM/Source/RigVMEditor/Private/Editor/RigVMGraphDetailCustomization.cpp
+++ b/Engine/Plugins/Runtime/RigVM/Source/RigVMEditor/Private/Editor/RigVMGraphDetailCustomization.cpp
@@ -299,6 +299,17 @@ void FRigVMFunctionArgumentGroupLayout::HandleModifiedEvent(ERigVMGraphNotifType
 	URigVMGraph* Graph = GraphPtr.Get();
 	URigVMLibraryNode* LibraryNode = Cast<URigVMLibraryNode>(Graph->GetOuter());
 	if (LibraryNode == nullptr)
+	{
+		if (URigVMNode* Node = Cast<URigVMNode>(InSubject))
+		{
+			LibraryNode = Cast<URigVMLibraryNode>(Node);
+		}
+		if (URigVMPin* Pin = Cast<URigVMPin>(InSubject))
+		{
+			LibraryNode = Cast<URigVMLibraryNode>(Pin->GetNode());
+		}
+	}
+	if (LibraryNode == nullptr)
 	{
 		return;
 	}
@@ -392,6 +403,8 @@ private:
 
 void FRigVMFunctionArgumentLayout::GenerateHeaderRowContent(FDetailWidgetRow& NodeRow)
 {
+	CachedImg_Pin_Disconnected = FAppStyle::GetBrush(TEXT("Graph.Pin.Disconnected"));
+	CachedImg_RefPin_Disconnected = FAppStyle::GetBrush(TEXT("Graph.RefPin.Disconnected"));
 	const UEdGraphSchema* Schema = GetDefault<URigVMEdGraphSchema>();
 
 	ETypeTreeFilter TypeTreeFilter = ETypeTreeFilter::None;
@@ -410,7 +423,6 @@ void FRigVMFunctionArgumentLayout::GenerateHeaderRowContent(FDetailWidgetRow& No
 		.HAlign(HAlign_Fill)
 		[
 			SNew(SHorizontalBox)
-			
 			+ SHorizontalBox::Slot()
 			.FillWidth(1)
 			.HAlign(HAlign_Fill)
@@ -457,12 +469,31 @@ void FRigVMFunctionArgumentLayout::GenerateHeaderRowContent(FDetailWidgetRow& No
 				.OnPinTypeChanged(this, &FRigVMFunctionArgumentLayout::PinInfoChanged)
 				.Schema(Schema)
 				.TypeTreeFilter(TypeTreeFilter)
-				.bAllowArrays(!ShouldPinBeReadOnly())
+				.bAllowArrays(false)
 				.IsEnabled(!ShouldPinBeReadOnly(true))
 				.CustomFilters(CustomPinTypeFilters)
 				.Font(IDetailLayoutBuilder::GetDetailFont())
 			]
 			+ SHorizontalBox::Slot()
+			.VAlign(VAlign_Center)
+			.Padding(0.0f, 0.0f, 4.0f, 0.0f)
+			.AutoWidth()
+			[
+				SNew(SComboButton)
+				.ComboButtonStyle(FAppStyle::Get(),"BlueprintEditor.CompactVariableTypeSelector")
+				.MenuPlacement(EMenuPlacement::MenuPlacement_ComboBoxRight)
+				.OnGetMenuContent(this, &FRigVMFunctionArgumentLayout::GetContainerTypeMenuContent)
+				.ContentPadding(0.0f)
+				.IsEnabled(CanChangeContainerType())
+				.Visibility(EVisibility::Visible)
+				.ButtonContent()
+				[
+					SNew(SImage)
+					.Image(this, &FRigVMFunctionArgumentLayout::GetContainerTypeIconImage)
+					.ColorAndOpacity(this, &FRigVMFunctionArgumentLayout::GetContainerTypeIconColor)
+				]
+			]
+			+ SHorizontalBox::Slot()
 			.HAlign(HAlign_Right)
 			.VAlign(VAlign_Center)
 			.Padding(10, 0, 0, 0)
@@ -612,6 +643,179 @@ void FRigVMFunctionArgumentLayout::OnPrePinInfoChange(const FEdGraphPinType& Pin
 	// not needed for rig vm
 }
 
+const TArray<TSharedPtr<FRigVMFunctionArgumentLayout::TContainerType>> FRigVMFunctionArgumentLayout::GetContainerTypes() const
+{
+		
+	static TArray<TSharedPtr<TContainerType>> ContainerTypes;
+	if (ContainerTypes.IsEmpty())
+	{
+		ContainerTypes.Add(MakeShared<TContainerType>(
+			EPinContainerType::None,
+			false, 
+			LOCTEXT("ContainerTypeSingleNoRef", "Single"), 
+			LOCTEXT("ContainerTypeSingleNoRefTooltip", "Single Value, Copied into the function"),
+			FAppStyle::Get().GetBrush(TEXT("Kismet.VariableList.TypeIcon"))
+		));
+		ContainerTypes.Add(MakeShared<TContainerType>(
+			EPinContainerType::Array,
+			false, 
+			LOCTEXT("ContainerTypeArrayNoRef", "Array"), 
+			LOCTEXT("ContainerTypeArrayNoRefTooltip", "Array Value, Copied into the function"),
+			FAppStyle::Get().GetBrush(TEXT("Kismet.VariableList.ArrayTypeIcon"))
+		));
+		ContainerTypes.Add(MakeShared<TContainerType>(
+			EPinContainerType::None,
+			true, 
+			LOCTEXT("ContainerTypeSingleRef", "Variable"), 
+			LOCTEXT("ContainerTypeSingleRefTooltip", "Single Variable, Referenced into the function"),
+			FAppStyle::Get().GetBrush(TEXT("GraphEditor.RefPinIcon"))
+		));
+		ContainerTypes.Add(MakeShared<TContainerType>(
+			EPinContainerType::Array,
+			true, 
+			LOCTEXT("ContainerTypeArrayRef", "Array Variable"), 
+			LOCTEXT("ContainerTypeArrayRefTooltip", "Array Variable, Referenced into the function"),
+			FAppStyle::Get().GetBrush(TEXT("Kismet.VariableList.ArrayTypeIcon"))
+		));
+	}
+
+	return ContainerTypes;
+}
+
+bool FRigVMFunctionArgumentLayout::CanChangeContainerType() const
+{
+	if (const URigVMPin* Pin = PinPtr.Get())
+	{
+		return !Pin->IsExecuteContext();
+	}
+	return false;
+}
+
+TSharedRef<SWidget> FRigVMFunctionArgumentLayout::GetContainerTypeMenuContent()
+{
+	FMenuBuilder MenuBuilder(true, nullptr, nullptr, /*bCloseSelfOnly*/true);
+
+	bool bRefsEnabled = true;
+	if (const URigVMPin* Pin = PinPtr.Get())
+	{
+		if (Pin->IsExecuteContext())
+		{
+			bRefsEnabled = false;
+		}
+		else if (Pin->GetDirection() != ERigVMPinDirection::Input)
+		{
+			bRefsEnabled = false;
+		}
+		else if (URigVMLibraryNode* LibraryNode = Cast<URigVMLibraryNode>(Pin->GetNode()))
+		{
+			if (!Cast<URigVMFunctionLibrary>(LibraryNode->GetGraph()))
+			{
+				bRefsEnabled = false;
+			}
+		}
+	}
+
+	const TArray<TSharedPtr<TContainerType>>& ContainerTypes = GetContainerTypes();
+	for (auto& Type : ContainerTypes)
+	{
+		check(Type);
+
+		if (Type->Get<1>() && !bRefsEnabled)
+		{
+			continue;
+		}
+		
+		FUIAction Action;
+		Action.ExecuteAction.BindSP(this, &FRigVMFunctionArgumentLayout::OnContainerTypeSelectionChanged, Type);
+
+		TSharedRef<SWidget> Widget =
+			SNew(SHorizontalBox)
+			+ SHorizontalBox::Slot()
+			.AutoWidth()
+			.VAlign(VAlign_Center)
+			[
+				SNew(SImage)
+				.Image(Type->Get<4>())
+				.ToolTipText(Type->Get<3>())
+				.ColorAndOpacity(this, &FRigVMFunctionArgumentLayout::GetContainerTypeIconColor)
+			]
+			+ SHorizontalBox::Slot()
+			.Padding(4.0f, 2.0f)
+			[
+				SNew(STextBlock)
+				.Text(Type->Get<2>())
+				.ToolTipText(Type->Get<3>())
+			];
+
+		MenuBuilder.AddMenuEntry(Action, Widget);
+	}
+
+	return MenuBuilder.MakeWidget();
+}
+
+const FSlateBrush* FRigVMFunctionArgumentLayout::GetContainerTypeIconImage() const
+{
+	if (const URigVMPin* Pin = PinPtr.Get())
+	{
+		if (!Pin->IsDefinedAsInputVariable())
+		{
+			if(!Pin->IsArray())
+			{
+				return GetContainerTypes()[0]->Get<4>();
+			}
+			return GetContainerTypes()[1]->Get<4>();
+		}
+		if(!Pin->IsArray())
+		{
+			return GetContainerTypes()[2]->Get<4>();
+		}
+		return GetContainerTypes()[3]->Get<4>();
+	}
+	return GetContainerTypes()[0]->Get<4>();
+}
+
+FSlateColor FRigVMFunctionArgumentLayout::GetContainerTypeIconColor() const
+{
+	if (const URigVMPin* Pin = PinPtr.Get())
+	{
+		if (!EdGraphSchemaPtr.IsValid())
+		{
+			if (const URigVMGraph* Graph = Pin->GetGraph())
+			{
+				if (UClass* EdGraphSchemaClass = Graph->GetSchema()->GetEdGraphSchemaClass())
+				{
+					EdGraphSchemaPtr = EdGraphSchemaClass->GetDefaultObject<URigVMEdGraphSchema>();
+				}
+			}
+		}
+
+		if (const URigVMEdGraphSchema* EdGraphSchema = EdGraphSchemaPtr.Get())
+		{
+			const FEdGraphPinType PinType = RigVMTypeUtils::PinTypeFromCPPType(*Pin->GetCPPType(), Pin->GetCPPTypeObject());
+			return EdGraphSchema->GetPinTypeColor(PinType);
+		}
+	}
+	return FSlateColor::UseForeground();
+}
+
+void FRigVMFunctionArgumentLayout::OnContainerTypeSelectionChanged(TSharedPtr<TContainerType> InContainerType)
+{
+	check(InContainerType);
+	
+	if (PinPtr.IsValid() && WeakRigVMClientHost.IsValid())
+	{
+		URigVMPin* Pin = PinPtr.Get();
+		IRigVMClientHost* RigVMClientHost = WeakRigVMClientHost.Get();
+		if (URigVMLibraryNode* LibraryNode = Cast<URigVMLibraryNode>(Pin->GetNode()))
+		{
+			if (URigVMController* Controller = RigVMClientHost->GetController(LibraryNode->GetContainedGraph()))
+			{
+				(void)Controller->MarkExposedPinAsInputVariable(Pin->GetFName(), InContainerType->Get<1>());
+			}
+		}
+	}
+}
+
 FRigVMFunctionArgumentDefaultNode::FRigVMFunctionArgumentDefaultNode(
 	const TWeakObjectPtr<URigVMGraph>& InGraph,
 	const TWeakInterfacePtr<IRigVMClientHost>& InClientHost
@@ -663,35 +867,24 @@ void FRigVMFunctionArgumentDefaultNode::GenerateChildContent(IDetailChildrenBuil
 
 	IRigVMClientHost* RigVMClientHost = WeakRigVMClientHost.Get();
 	URigVMGraph* Graph = GraphPtr.Get();
-	URigVMEdGraphNode* RigVMEdGraphNode = nullptr;
-	if (URigVMLibraryNode* LibraryNode = Cast<URigVMLibraryNode>(Graph->GetOuter()))
-	{
-		if (URigVMEdGraph* RigGraph = Cast<URigVMEdGraph>(RigVMClientHost->GetEditorObjectForRigVMGraph(LibraryNode->GetGraph())))
-		{
-			RigVMEdGraphNode = Cast<URigVMEdGraphNode>(RigGraph->FindNodeForModelNodeName(LibraryNode->GetFName()));
-		}
-	}
-
-	if (RigVMEdGraphNode == nullptr)
+	URigVMLibraryNode* LibraryNode = Cast<URigVMLibraryNode>(Graph->GetOuter());
+	if (LibraryNode == nullptr)
 	{
 		return;
 	}
 
+	TSharedPtr<FRigVMMinimalEnvironment> Environment = MakeShared<FRigVMMinimalEnvironment>(LibraryNode->GetImplementingOuter<IRigVMAssetInterface>());
+	URigVMNode* FuncRefNode = Environment->GetController()->AddFunctionReferenceNode(LibraryNode);
+	Environment->SetNode(FuncRefNode);
+
 	ChildrenBuilder.AddCustomRow(FText::GetEmpty())
 	.WholeRowContent()
 	.MaxDesiredWidth(980.f)
 	[
-		SAssignNew(OwnedNodeWidget, SRigVMGraphNode).GraphNodeObj(RigVMEdGraphNode)
+		SAssignNew(PreviewWidget, SRigVMNodePreviewWidget)
+		.Environment(Environment)
+		.NodeEnabled(true)
 	];
-
-	OwnedNodeWidget->SetIsEditable(true);
-	TArray< TSharedRef<SWidget> > Pins;
-	OwnedNodeWidget->GetPins(Pins);
-	for (TSharedRef<SWidget> Pin : Pins)
-	{
-		TSharedRef<SGraphPin> SPin = StaticCastSharedRef<SGraphPin>(Pin);
-		SPin->EnableDragAndDrop(false);
-	}
 }
 
 void FRigVMFunctionArgumentDefaultNode::OnGraphChanged(const FEdGraphEditAction& InAction)
diff --git a/Engine/Plugins/Runtime/RigVM/Source/RigVMEditor/Private/Editor/RigVMMinimalEnvironment.cpp b/Engine/Plugins/Runtime/RigVM/Source/RigVMEditor/Private/Editor/RigVMMinimalEnvironment.cpp
--- a/Engine/Plugins/Runtime/RigVM/Source/RigVMEditor/Private/Editor/RigVMMinimalEnvironment.cpp
+++ b/Engine/Plugins/Runtime/RigVM/Source/RigVMEditor/Private/Editor/RigVMMinimalEnvironment.cpp
@@ -6,16 +6,24 @@
 
 #define LOCTEXT_NAMESPACE "RigVMMinimalEnvironment"
 
-FRigVMMinimalEnvironment::FRigVMMinimalEnvironment(const UClass* InRigVMBlueprintClass)
+FRigVMMinimalEnvironment::FRigVMMinimalEnvironment(const IRigVMAssetInterface* InAssetInterface)
 {
-	ModelController = TStrongObjectPtr<URigVMController>(NewObject<URigVMController>());
+	UObject* Outer = GetTransientPackage();
+	if (InAssetInterface)
+	{
+		if (const UObject* AssetObject = InAssetInterface->GetObject())
+		{
+			Outer = const_cast<UObject*>(AssetObject);
+		}
+	}
+	ModelController = TStrongObjectPtr<URigVMController>(NewObject<URigVMController>(Outer));
 	ModelGraph = TStrongObjectPtr<URigVMGraph>(NewObject<URigVMGraph>(ModelController.Get()));
 	ModelController->SetGraph(ModelGraph.Get());
 
 	EdGraphClass = URigVMEdGraph::StaticClass();
 	EdGraphNodeClass = URigVMEdGraphNode::StaticClass();
 
-	SetSchemata(InRigVMBlueprintClass ? InRigVMBlueprintClass : URigVMBlueprint::StaticClass());
+	SetSchemata(InAssetInterface);
 }
 
 URigVMGraph* FRigVMMinimalEnvironment::GetModel() const
@@ -76,7 +84,10 @@ void FRigVMMinimalEnvironment::SetFunctionNode(const FRigVMGraphFunctionIdentifi
 			{
 				if(Class->ImplementsInterface(URigVMEditorAssetInterface::StaticClass()))
 				{
-					SetSchemata(Class);
+					if (const URigVMNode* Node = GetNode())
+					{
+						SetSchemata(Node->GetImplementingOuter<IRigVMAssetInterface>());
+					}
 				}
 			}
 		}
@@ -107,23 +118,28 @@ URigVMEdGraphNode* FRigVMMinimalEnvironment::GetEdGraphNode() const
 	return nullptr;
 }
 
-void FRigVMMinimalEnvironment::SetSchemata(const UClass* InRigVMBlueprintClass)
+void FRigVMMinimalEnvironment::SetSchemata(const IRigVMAssetInterface* InAssetInterface)
 {
-	check(InRigVMBlueprintClass);
+	if (!InAssetInterface)
+	{
+		return;
+	}
 
-	const FRigVMEditorAssetInterfacePtr CDO = InRigVMBlueprintClass->GetDefaultObject();
-	check(CDO);
+	const TScriptInterface<IRigVMClientHost> ClientHost = InAssetInterface->GetRigVMClientHost();
+	check(ClientHost);
 
-	EdGraphClass = CDO->GetRigVMClientHost()->GetRigVMEdGraphClass();
-	EdGraphNodeClass = CDO->GetRigVMClientHost()->GetRigVMEdGraphNodeClass();
+	EdGraphClass = ClientHost->GetRigVMEdGraphClass();
+	EdGraphNodeClass = ClientHost->GetRigVMEdGraphNodeClass();
 
 	if(!EdGraph.IsValid() || EdGraph->GetClass() != EdGraphClass)
 	{
 		EdGraph = TStrongObjectPtr<URigVMEdGraph>(NewObject<URigVMEdGraph>(ModelGraph.Get(), EdGraphClass));
 	}
 
-	ModelController->SetSchemaClass(CDO->GetRigVMClientHost()->GetRigVMSchemaClass());
-	EdGraph->SetBlueprintClass(InRigVMBlueprintClass);
+	ModelController->SetSchemaClass(ClientHost->GetRigVMSchemaClass());
+	const UObject* AssetObject = InAssetInterface->GetObject();
+	check(AssetObject);
+	EdGraph->SetBlueprintClass(AssetObject->GetClass());
 }
 
 FSimpleDelegate& FRigVMMinimalEnvironment::OnChanged()
@@ -169,6 +185,7 @@ void FRigVMMinimalEnvironment::HandleModified(ERigVMGraphNotifType InNotificatio
 		{
 			break;
 		}
+		case ERigVMGraphNotifType::PinTypeChanged:
 		default:
 		{
 			(void)NumModifications.fetch_add(1);
diff --git a/Engine/Plugins/Runtime/RigVM/Source/RigVMEditor/Private/RigVMEditorModule.cpp b/Engine/Plugins/Runtime/RigVM/Source/RigVMEditor/Private/RigVMEditorModule.cpp
--- a/Engine/Plugins/Runtime/RigVM/Source/RigVMEditor/Private/RigVMEditorModule.cpp
+++ b/Engine/Plugins/Runtime/RigVM/Source/RigVMEditor/Private/RigVMEditorModule.cpp
@@ -626,13 +626,13 @@ void FRigVMEditorModule::GetInstanceActions(FRigVMEditorAssetInterfacePtr RigVMB
 	for (const FRigVMGraphVariableDescription& Variable : Variables)
 	{
 		const FRigVMExternalVariable ExternalVariable = Variable.ToExternalVariable();
-		FText MenuDesc = FText::FromName(ExternalVariable.Name);
-		FText ToolTip = FText::FromString(FString::Printf(TEXT("Get the value of variable %s"), *ExternalVariable.Name.ToString()));
+		FText MenuDesc = FText::FromName(ExternalVariable.GetName());
+		FText ToolTip = FText::FromString(FString::Printf(TEXT("Get the value of variable %s"), *ExternalVariable.GetName().ToString()));
 		TSharedRef<URigVMEdGraphVariableNodeSpawner> GetNodeSpawner = MakeShared<URigVMEdGraphVariableNodeSpawner>(RigVMBlueprint, ExternalVariable, true, MenuDesc, NodeCategory, ToolTip);
 		GetNodeSpawner->SetRelatedBlueprintClass(BlueprintClass);
 		OutGraphNodeSpawners.Add(GetNodeSpawner);
 
-		ToolTip = FText::FromString(FString::Printf(TEXT("Set the value of variable %s"), *ExternalVariable.Name.ToString()));
+		ToolTip = FText::FromString(FString::Printf(TEXT("Set the value of variable %s"), *ExternalVariable.GetName().ToString()));
 		TSharedRef<URigVMEdGraphVariableNodeSpawner> SetNodeSpawner = MakeShared<URigVMEdGraphVariableNodeSpawner>(RigVMBlueprint, ExternalVariable, false, MenuDesc, NodeCategory, ToolTip);
 		SetNodeSpawner->SetRelatedBlueprintClass(BlueprintClass);
 		OutGraphNodeSpawners.Add(SetNodeSpawner);
@@ -690,6 +690,21 @@ void FRigVMEditorModule::GetInstanceActions(FRigVMEditorAssetInterfacePtr RigVMB
 				NodeSpawner->SetRelatedBlueprintClass(BlueprintClass);
 				OutGraphNodeSpawners.Add(NodeSpawner);
 			}
+			
+			static const FText InputVariableCategory = LOCTEXT("InputVariables", "Input Variables");
+			for (const FRigVMGraphVariableDescription& InputVariable : Graph->GetInputVariables())
+			{
+				FText MenuDesc = FText::FromName(InputVariable.Name);
+				FText ToolTip = FText::FromString(FString::Printf(TEXT("Get the value of input variable %s"), *InputVariable.Name.ToString()));
+				TSharedRef<URigVMEdGraphVariableNodeSpawner> NodeSpawner = MakeShared<URigVMEdGraphVariableNodeSpawner>(RigVMBlueprint, Graph, InputVariable, true, MenuDesc, InputVariableCategory, ToolTip);
+				NodeSpawner->SetRelatedBlueprintClass(BlueprintClass);
+				OutGraphNodeSpawners.Add(NodeSpawner);
+
+				ToolTip = FText::FromString(FString::Printf(TEXT("Set the value of input variable %s"), *InputVariable.Name.ToString()));
+				NodeSpawner = MakeShared<URigVMEdGraphVariableNodeSpawner>(RigVMBlueprint, Graph, InputVariable, false, MenuDesc, InputVariableCategory, ToolTip);
+				NodeSpawner->SetRelatedBlueprintClass(BlueprintClass);
+				OutGraphNodeSpawners.Add(NodeSpawner);
+			}			
 		}
 	}
 
@@ -1174,59 +1189,65 @@ void FRigVMEditorModule::GetNodeVariablesContextMenuActions(TScriptInterface<IRi
 
 		if (URigVMFunctionReferenceNode* FunctionReferenceNode = Cast<URigVMFunctionReferenceNode>(EdGraphNode->GetModelNode()))
 		{
-			TSoftObjectPtr<URigVMFunctionReferenceNode> RefPtr(FunctionReferenceNode);
-			if(RefPtr.GetLongPackageName() != FunctionReferenceNode->GetReferencedFunctionHeader().LibraryPointer.GetNodeSoftPath().GetLongPackageName())
+			if(!FunctionReferenceNode->IsFullyRemapped() && FunctionReferenceNode->GetReferencedFunctionHeader().LibraryPointer.GetNodeSoftPath().ResolveObject())
 			{
-				if(!FunctionReferenceNode->IsFullyRemapped() && FunctionReferenceNode->GetReferencedFunctionHeader().LibraryPointer.GetNodeSoftPath().ResolveObject())
-				{
-					FToolMenuSection& VariablesSection = Menu->AddSection("RigVMEditorContextMenuVariables", LOCTEXT("Variables", "Variables"));
-					VariablesSection.AddMenuEntry(
-						"MakeVariablesFromFunctionReferenceNode",
-						LOCTEXT("MakeVariablesFromFunctionReferenceNode", "Create required variables"),
-						LOCTEXT("MakeVariablesFromFunctionReferenceNode_Tooltip", "Creates all required variables for this function and binds them"),
-						FSlateIcon(),
-						FUIAction(FExecuteAction::CreateLambda([Controller, FunctionReferenceNode, RigVMBlueprint]() {
+				FToolMenuSection& VariablesSection = Menu->AddSection("RigVMEditorContextMenuVariables", LOCTEXT("Variables", "Variables"));
+				VariablesSection.AddMenuEntry(
+					"MakeVariablesFromFunctionReferenceNode",
+					LOCTEXT("MakeVariablesFromFunctionReferenceNode", "Create required variables"),
+					LOCTEXT("MakeVariablesFromFunctionReferenceNode_Tooltip", "Creates all required variables for this function and binds them"),
+					FSlateIcon(),
+					FUIAction(FExecuteAction::CreateLambda([Controller, FunctionReferenceNode, RigVMBlueprint]() {
 
-							const TArray<FRigVMExternalVariable> ExternalVariables = FunctionReferenceNode->GetExternalVariables(false);
-							if(!ExternalVariables.IsEmpty())
-							{
-								FScopedTransaction Transaction(LOCTEXT("MakeVariablesFromFunctionReferenceNode", "Create required variables"));
-								RigVMBlueprint->GetObject()->Modify();
+						const TArray<FRigVMExternalVariable> ExternalVariables = FunctionReferenceNode->GetExternalVariables();
+						if(!ExternalVariables.IsEmpty())
+						{
+							FScopedTransaction Transaction(LOCTEXT("MakeVariablesFromFunctionReferenceNode", "Create required variables"));
+							RigVMBlueprint->GetObject()->Modify();
 
-								if (URigVMLibraryNode* LibraryNode = FunctionReferenceNode->LoadReferencedNode())
+							if (URigVMLibraryNode* LibraryNode = FunctionReferenceNode->LoadReferencedNode())
+							{
+								FRigVMAssetInterfacePtr ReferencedBlueprint = IRigVMAssetInterface::GetInterfaceOuter(LibraryNode);
+							   // ReferencedBlueprint != RigVMBlueprint - since only FunctionReferenceNodes from other assets have the potential to be unmapped
+                        
+								for(const FRigVMExternalVariable& ExternalVariable : ExternalVariables)
 								{
-									FRigVMEditorAssetInterfacePtr ReferencedBlueprint = IRigVMEditorAssetInterface::GetInterfaceOuter(LibraryNode);
-								   // ReferencedBlueprint != RigVMBlueprint - since only FunctionReferenceNodes from other assets have the potential to be unmapped
-                            
-								   for(const FRigVMExternalVariable& ExternalVariable : ExternalVariables)
-								   {
-									   FString DefaultValue;
-									   if(ReferencedBlueprint)
-									   {
-										   for(const FRigVMGraphVariableDescription& NewVariable : ReferencedBlueprint->GetAssetVariables())
-										   {
-											   if(NewVariable.Name == ExternalVariable.Name)
-											   {
-												   DefaultValue = NewVariable.DefaultValue;
-												   break;
-											   }
-										   }
-									   }
-                                
-									   FName NewVariableName = RigVMBlueprint->AddHostMemberVariableFromExternal(ExternalVariable, DefaultValue);
-									   if(!NewVariableName.IsNone())
-									   {
-										   Controller->SetRemappedVariable(FunctionReferenceNode, ExternalVariable.Name, NewVariableName);
-									   }
-								   }
-								}
+									URigVMPin* InterfacePin = FunctionReferenceNode->FindRootPinByName(ExternalVariable.GetName());
+									if (!InterfacePin)
+									{
+										continue;
+									}
+									if (!InterfacePin->IsDefinedAsInputVariable())
+									{
+										continue;
+									}
+									
+									FString DefaultValue;
+									if(ReferencedBlueprint)
+									{
+										for(const FRigVMGraphVariableDescription& NewVariable : ReferencedBlueprint->GetAssetVariables())
+										{
+											if(NewVariable.Name == ExternalVariable.GetName())
+											{
+											   DefaultValue = NewVariable.DefaultValue;
+											   break;
+											}
+										}
+									}
 
-								RigVMBlueprint->MarkAssetAsModified();
+									FName NewVariableName = RigVMBlueprint->AddHostMemberVariableFromExternal(ExternalVariable, DefaultValue);
+									if (!NewVariableName.IsNone())
+									{
+										Controller->BindPinToVariable(InterfacePin->GetPinPath(), NewVariableName.ToString());
+									}
+								}
 							}
-                        
-						})
-					));
-				}
+
+							RigVMBlueprint->MarkAssetAsModified();
+						}
+                    
+					})
+				));
 			}
 		}
 	}
diff --git a/Engine/Plugins/Runtime/RigVM/Source/RigVMEditor/Private/Widgets/SRigVMGraphNode.cpp b/Engine/Plugins/Runtime/RigVM/Source/RigVMEditor/Private/Widgets/SRigVMGraphNode.cpp
--- a/Engine/Plugins/Runtime/RigVM/Source/RigVMEditor/Private/Widgets/SRigVMGraphNode.cpp
+++ b/Engine/Plugins/Runtime/RigVM/Source/RigVMEditor/Private/Widgets/SRigVMGraphNode.cpp
@@ -297,6 +297,9 @@ void SRigVMGraphNode::Construct( const FArguments& InArgs )
 	this->GraphNode = InArgs._GraphNodeObj;
 	this->SetCursor( EMouseCursor::CardinalCross );
 
+	SetEnabled(InArgs._Enabled);
+	PinDragDropEnabledAttribute = InArgs._PinDragDropEnabled;
+
  	URigVMEdGraphNode* EdGraphNode = InArgs._GraphNodeObj;
 	ModelNode = EdGraphNode->GetModelNode();
 	if (!ModelNode.IsValid())
@@ -316,7 +319,7 @@ void SRigVMGraphNode::Construct( const FArguments& InArgs )
 
 	bPreviewHereIndicatorHovered = false;
 	
-	SetIsEditable(false);
+	SetIsEditable(IsEnabled());
 
 	Controller->OnModified().AddSP(this, &SRigVMGraphNode::HandleModifiedEvent);
 
@@ -565,7 +568,8 @@ TOptional<FSlateColor> SRigVMGraphNode::GetPinTextColor(const SGraphPin* InGraph
 	}
 
 	const URigVMFunctionInterfaceNode* FunctionInterfaceNode = Cast<URigVMFunctionInterfaceNode>(ModelNode);
-	if (!FunctionInterfaceNode)
+	const URigVMFunctionReferenceNode* FunctionReferenceNode = Cast<URigVMFunctionReferenceNode>(ModelNode);
+	if (!FunctionInterfaceNode && !FunctionReferenceNode)
 	{
 		return Result;
 	}
@@ -578,13 +582,21 @@ TOptional<FSlateColor> SRigVMGraphNode::GetPinTextColor(const SGraphPin* InGraph
 
 	if (const URigVMPin* ModelPin = ModelNode->FindPin(PinPath))
 	{
-		if (!ModelPin->IsLinked() &&
+		if (FunctionInterfaceNode && !ModelPin->IsLinked() &&
 			!FunctionInterfaceNode->IsInterfacePinUsed(ModelPin->GetRootPin()->GetFName()))
 		{
 			// we are other an output argument without any links to it
 			// or an input argument without links and without being bound within the graph
 			return FSlateColor::UseSubduedForeground();
 		}
+
+		if (FunctionReferenceNode && 
+			ModelPin->GetDirection() == ERigVMPinDirection::Input &&
+			ModelPin->IsDefinedAsInputVariable() &&
+			ModelPin->GetBoundVariableName().IsEmpty())
+		{
+			return FSlateColor(FLinearColor::Red);
+		}
 	}
 
 	return Result;
@@ -905,6 +917,9 @@ void SRigVMGraphNode::AddPin(const TSharedRef<SGraphPin>& PinToAdd)
 		{
 			OutputPins.Add(PinToAdd);
 		}
+
+		const bool bPinDragDropEnabled = PinDragDropEnabledAttribute.Get(true);
+		PinToAdd->EnableDragAndDrop(bPinDragDropEnabled);
 	}
 }
 
@@ -912,6 +927,7 @@ void SRigVMGraphNode::CreateStandardPinWidget(UEdGraphPin* CurPin)
 {
 	bool bShowPin = true;
 	bool bIsFixedArray = false;
+	bool bIsInputVariable = false;
 	FString CPPType;
 	FString BoundVariableName;
 	if(const URigVMEdGraphNode* RigGraphNode = Cast<URigVMEdGraphNode>(GraphNode))
@@ -927,6 +943,7 @@ void SRigVMGraphNode::CreateStandardPinWidget(UEdGraphPin* CurPin)
 			CPPType = ModelPin->GetCPPType();
 			BoundVariableName = ModelPin->GetBoundVariableName();
 			bIsFixedArray = ModelPin->IsFixedSizeArray();
+			bIsInputVariable = ModelPin->IsDefinedAsInputVariable();
 		}
 	}
 	
@@ -942,7 +959,8 @@ void SRigVMGraphNode::CreateStandardPinWidget(UEdGraphPin* CurPin)
 			if(PinInfoMetaData.IsValid())
 			{
 				if(PinInfoMetaData->CPPType == CPPType &&
-					PinInfoMetaData->BoundVariableName == BoundVariableName)
+					PinInfoMetaData->BoundVariableName == BoundVariableName &&
+					PinInfoMetaData->bIsInputVariable == bIsInputVariable)
 				{
 					NewPin = RecycledPin;
 				}
@@ -956,7 +974,7 @@ void SRigVMGraphNode::CreateStandardPinWidget(UEdGraphPin* CurPin)
 				(*RecycledPinPtr)->InvalidateGraphData();
 			}
 			NewPin = CreatePinWidget(CurPin);
-			NewPin->AddMetadata(MakeShared<FPinInfoMetaData>(CPPType, BoundVariableName));
+			NewPin->AddMetadata(MakeShared<FPinInfoMetaData>(CPPType, BoundVariableName, bIsInputVariable));
 			check(NewPin.IsValid());
 			PinsToKeep.Remove(CurPin);
 
@@ -1037,6 +1055,10 @@ EVisibility SRigVMGraphNode::GetArrayPlusButtonVisibility(URigVMPin* InModelPin)
 {
 	if(InModelPin)
 	{
+		if (InModelPin->IsDefinedAsInputVariable())
+		{
+			return EVisibility::Hidden;
+		}
 		if (Cast<URigVMFunctionReturnNode>(InModelPin->GetNode()))
 		{
 			return EVisibility::Hidden;
@@ -2256,6 +2278,7 @@ void SRigVMGraphNode::UpdatePinTreeView()
 		{
 			PinInfo.bHasChildren = InModelPin->GetSubPins().Num() > 0;
 			PinInfo.bIsContainer = InModelPin->IsArray();
+			PinInfo.bIsInputVariable = InModelPin->IsDefinedAsInputVariable();
 			PinInfo.bExpanded = InModelPin->IsExpanded();
 			PinInfo.Identifier = InModelPin->GetPinPath();
 			PinInfo.bShowOnlySubPins = InModelPin->ShouldOnlyShowSubPins();
@@ -2274,6 +2297,7 @@ void SRigVMGraphNode::UpdatePinTreeView()
 				}
 			}
 			PinInfo.bIsContainer = false;
+			PinInfo.bIsInputVariable = false;
 			PinInfo.bExpanded = CurrentModelNode->IsPinCategoryExpanded(InCategory);
 			PinInfo.Identifier = InCategory;
 			PinInfo.bShowOnlySubPins = false;
@@ -2291,8 +2315,8 @@ void SRigVMGraphNode::UpdatePinTreeView()
 		if(bIsModelPin)
 		{
 			const bool bAskSchemaForEdition = RigSchema && InModelPin->IsStruct() && !InModelPin->IsBoundToVariable();
-			PinInfo.bHideInputWidget = (!InModelPin->IsBoundToVariable()) && PinInfo.bIsContainer;
-			if (!PinInfo.bHideInputWidget)
+			PinInfo.bHideInputWidget = (!InModelPin->IsBoundToVariable()) && PinInfo.bIsContainer && !PinInfo.bIsInputVariable;
+			if (!PinInfo.bHideInputWidget && !PinInfo.bIsInputVariable)
 			{
 				if (bAskSchemaForEdition && !PinInfo.bHasChildren)
 				{
@@ -2895,50 +2919,6 @@ void SRigVMGraphNode::UpdatePinTreeView()
 		}
 	}
 
-	if(URigVMFunctionReferenceNode* FunctionReferenceNode = Cast<URigVMFunctionReferenceNode>(ModelNode))
-	{
-		TWeakObjectPtr<URigVMFunctionReferenceNode> WeakFunctionReferenceNode = FunctionReferenceNode;
-		TWeakInterfacePtr<IRigVMEditorAssetInterface> WeakRigVMBlueprint = Blueprint;
-
-		// add the entries for the variable remapping
-		for(const TSharedPtr<FRigVMExternalVariable>& ExternalVariable : RigGraphNode->ExternalVariables)
-		{
-			LeftNodeBox->AddSlot()
-			.HAlign(HAlign_Fill)
-			.VAlign(VAlign_Center)
-			.AutoHeight()
-			.MaxHeight(MaxHeight)
-			[
-				SNew(SHorizontalBox)
-	            
-				+SHorizontalBox::Slot()
-				.AutoWidth()
-				.HAlign(HAlign_Left)
-				.VAlign(VAlign_Center)
-				.Padding(PinWidgetSidePadding, TopPadding, PinWidgetSidePadding, 0.f)
-				[
-					SNew(STextBlock)
-					.Text(FText::FromName(ExternalVariable->Name))
-					.TextStyle(FAppStyle::Get(), NAME_DefaultPinLabelStyle)
-					.ColorAndOpacity(this, &SRigVMGraphNode::GetVariableLabelTextColor, WeakFunctionReferenceNode, ExternalVariable->Name)
-					.ToolTipText(this, &SRigVMGraphNode::GetVariableLabelTooltipText, WeakRigVMBlueprint, ExternalVariable->Name)
-				]
-
-				+SHorizontalBox::Slot()
-				.AutoWidth()
-				.HAlign(HAlign_Left)
-				.VAlign(VAlign_Center)
-				.Padding(PinWidgetSidePadding, TopPadding, PinWidgetSidePadding, 0.f)
-				[
-					SNew(SRigVMGraphVariableBinding)
-					.Asset(Blueprint.GetObject())
-					.FunctionReferenceNode(FunctionReferenceNode)
-					.InnerVariableName(ExternalVariable->Name)
-				]
-			];
-		}
-	}
-
 	CreateAddPinButton();
 
 	CreateWorkflowWidgets();
diff --git a/Engine/Plugins/Runtime/RigVM/Source/RigVMEditor/Private/Widgets/SRigVMGraphPinVariableBinding.cpp b/Engine/Plugins/Runtime/RigVM/Source/RigVMEditor/Private/Widgets/SRigVMGraphPinVariableBinding.cpp
--- a/Engine/Plugins/Runtime/RigVM/Source/RigVMEditor/Private/Widgets/SRigVMGraphPinVariableBinding.cpp
+++ b/Engine/Plugins/Runtime/RigVM/Source/RigVMEditor/Private/Widgets/SRigVMGraphPinVariableBinding.cpp
@@ -23,8 +23,6 @@ static const FText RigVMGraphVariableBindingMultipleValues = LOCTEXT("MultipleVa
 void SRigVMGraphVariableBinding::Construct(const FArguments& InArgs)
 {
 	this->ModelPins = InArgs._ModelPins;
-	this->FunctionReferenceNode = InArgs._FunctionReferenceNode;
-	this->InnerVariableName = InArgs._InnerVariableName;
 	this->Blueprint = InArgs._Asset;
 	this->bCanRemoveBinding = InArgs._CanRemoveBinding;
 
@@ -93,14 +91,6 @@ FText SRigVMGraphVariableBinding::GetBindingText() const
 		}
 		return FirstText;
 	}
-	else if(FunctionReferenceNode && !InnerVariableName.IsNone())
-	{
-		const FName BoundVariable = FunctionReferenceNode->GetOuterVariableName(InnerVariableName);
-		if(!BoundVariable.IsNone())
-		{
-			return FText::FromName(BoundVariable);
-		}
-	}
 	return FText();
 }
 
@@ -130,14 +120,6 @@ FLinearColor SRigVMGraphVariableBinding::GetBindingColor() const
 		{
 			BoundVariable = *ModelPins[0]->GetBoundVariableName();
 		}
-		else if(FunctionReferenceNode && !InnerVariableName.IsNone())
-		{
-			BoundVariable = FunctionReferenceNode->GetOuterVariableName(InnerVariableName);
-			if(BoundVariable.IsNone())
-			{
-				return FLinearColor::Red;
-			}
-		}
 
 		for (const FRigVMGraphVariableDescription& VariableDescription : Blueprint->GetAssetVariables())
 		{
@@ -155,6 +137,10 @@ FLinearColor SRigVMGraphVariableBinding::GetBindingColor() const
 			{
 				return  FLinearColor::Red;
 			}
+			if (BoundVariable.IsNone() && ModelPins[0]->IsDefinedAsInputVariable())
+			{
+				return FLinearColor::Red;
+			}
 
 			const TArray<FRigVMGraphVariableDescription>& LocalVariables =  Model->GetLocalVariables(true);
 			for(const FRigVMGraphVariableDescription& LocalVariable : LocalVariables)
@@ -165,7 +151,23 @@ FLinearColor SRigVMGraphVariableBinding::GetBindingColor() const
 					continue;
 				}
 
-				if (ExternalVariable.Name == BoundVariable)
+				if (ExternalVariable.GetName() == BoundVariable)
+				{
+					const FEdGraphPinType PinType = RigVMTypeUtils::PinTypeFromExternalVariable(ExternalVariable);
+					return Schema->GetPinTypeColor(PinType);
+				}
+			}
+
+			const TArray<FRigVMGraphVariableDescription>& InputVariables =  Model->GetInputVariables();
+			for(const FRigVMGraphVariableDescription& InputVariable : InputVariables)
+			{
+				const FRigVMExternalVariable ExternalVariable = InputVariable.ToExternalVariable();
+				if(!ExternalVariable.IsValid(true))
+				{
+					continue;
+				}
+
+				if (ExternalVariable.GetName() == BoundVariable)
 				{
 					const FEdGraphPinType PinType = RigVMTypeUtils::PinTypeFromExternalVariable(ExternalVariable);
 					return Schema->GetPinTypeColor(PinType);
@@ -190,33 +192,6 @@ bool SRigVMGraphVariableBinding::OnCanBindProperty(FProperty* InProperty) const
 		{
 			return ModelPins[0]->CanBeBoundToVariable(ExternalVariable);
 		}
-		else if(FunctionReferenceNode && !InnerVariableName.IsNone())
-		{
-			TArray<FRigVMExternalVariable> InnerVariables = FunctionReferenceNode->GetExternalVariables(false);
-			for(const FRigVMExternalVariable& InnerVariable : InnerVariables)
-			{
-				if(InnerVariable.Name == InnerVariableName)
-				{
-					if(!InnerVariable.bIsReadOnly && ExternalVariable.bIsReadOnly)
-					{
-						return false;
-					}
-					if(InnerVariable.bIsArray != ExternalVariable.bIsArray)
-					{
-						return false;
-					}
-					if(InnerVariable.TypeObject && InnerVariable.TypeObject != ExternalVariable.TypeObject)
-					{
-						return false;
-					}
-					else if(InnerVariable.TypeName != ExternalVariable.TypeName)
-					{
-						return false;
-					}
-					return true;
-				}
-			}
-		}
 	}
 
 	return false;
@@ -252,11 +227,6 @@ void SRigVMGraphVariableBinding::OnAddBinding(FName InPropertyName, const TArray
 				Blueprint->GetController(ModelPin->GetGraph())->BindPinToVariable(ModelPin->GetPinPath(), FString::Join(Parts, TEXT(".")), true /* undo */, true /* python */);
 			}
 		}
-		else if(FunctionReferenceNode && !InnerVariableName.IsNone())
-		{
-			const FName BoundVariableName = *FString::Join(Parts, TEXT("."));
-			Blueprint->GetController(FunctionReferenceNode->GetGraph())->SetRemappedVariable(FunctionReferenceNode, InnerVariableName, BoundVariableName);
-		}
 	}
 }
 
@@ -276,10 +246,6 @@ void SRigVMGraphVariableBinding::OnRemoveBinding(FName InPropertyName)
 				Blueprint->GetController(ModelPin->GetGraph())->UnbindPinFromVariable(ModelPin->GetPinPath(), true /* undo */, true /* python */);
 			}
 		}
-		else if(FunctionReferenceNode && !InnerVariableName.IsNone())
-		{
-			Blueprint->GetController(FunctionReferenceNode->GetGraph())->SetRemappedVariable(FunctionReferenceNode, InnerVariableName, NAME_None);
-		}
 	}
 }
 
@@ -296,8 +262,11 @@ void SRigVMGraphVariableBinding::FillLocalVariableMenu(FMenuBuilder& MenuBuilder
 		return;
 	}
 
-	int32 ValidLocalVariables = 0;
-	const TArray<FRigVMGraphVariableDescription>& LocalVariables =  Model->GetLocalVariables(true);
+	const bool bAllowInputArguments = !ModelPins[0]->IsDefinedAsInputVariable();
+	constexpr bool bAllowInputVariables = true;
+
+	TArray<FRigVMGraphVariableDescription> ValidVariables;
+	const TArray<FRigVMGraphVariableDescription>& LocalVariables = Model->GetLocalVariables(bAllowInputArguments);
 	for(const FRigVMGraphVariableDescription& LocalVariable : LocalVariables)
 	{
 		const FRigVMExternalVariable ExternalVariable = LocalVariable.ToExternalVariable();
@@ -305,42 +274,44 @@ void SRigVMGraphVariableBinding::FillLocalVariableMenu(FMenuBuilder& MenuBuilder
 		{
 			continue;
 		}
-
 		if(!ModelPins[0]->CanBeBoundToVariable(ExternalVariable))
 		{
 			continue;
 		}
-
-		ValidLocalVariables++;
+		ValidVariables.Add(LocalVariable);
+	}
+	const TArray<FRigVMGraphVariableDescription>& InputVariables = Model->GetInputVariables();
+	for(const FRigVMGraphVariableDescription& InputVariable : InputVariables)
+	{
+		const FRigVMExternalVariable ExternalVariable = InputVariable.ToExternalVariable();
+		if(!ExternalVariable.IsValid(true))
+		{
+			continue;
+		}
+		if(!ModelPins[0]->CanBeBoundToVariable(ExternalVariable))
+		{
+			continue;
+		}
+		ValidVariables.Add(InputVariable);
 	}
 
-	if(ValidLocalVariables == 0)
+	if(ValidVariables.IsEmpty())
 	{
 		return;
 	}
 	
-	MenuBuilder.BeginSection("LocalVariables", LOCTEXT("LocalVariables", "Local Variables"));
+	MenuBuilder.BeginSection("Variables", LOCTEXT("Variables", "Variables"));
 	{
 		static const FLazyName PropertyIcon(TEXT("Kismet.VariableList.TypeIcon"));
 		const URigVMEdGraphSchema* Schema = GetDefault<URigVMEdGraphSchema>();
 
-		for(const FRigVMGraphVariableDescription& LocalVariable : LocalVariables)
+		for(const FRigVMGraphVariableDescription& ValidVariable : ValidVariables)
 		{
-			const FRigVMExternalVariable ExternalVariable = LocalVariable.ToExternalVariable();
-			if(!ExternalVariable.IsValid(true))
-			{
-				continue;
-			}
-
-			if(!ModelPins[0]->CanBeBoundToVariable(ExternalVariable))
-			{
-				continue;
-			}
-			
+			const FRigVMExternalVariable ExternalVariable = ValidVariable.ToExternalVariable();
 			const FEdGraphPinType PinType = RigVMTypeUtils::PinTypeFromExternalVariable(ExternalVariable);
 
 			MenuBuilder.AddMenuEntry(
-				FUIAction(FExecuteAction::CreateSP(this, &SRigVMGraphVariableBinding::HandleBindToLocalVariable, LocalVariable)),
+				FUIAction(FExecuteAction::CreateSP(this, &SRigVMGraphVariableBinding::HandleBindToLocalVariable, ValidVariable)),
 				SNew(SHorizontalBox)
 					+SHorizontalBox::Slot()
 					.AutoWidth()
@@ -363,7 +334,7 @@ void SRigVMGraphVariableBinding::FillLocalVariableMenu(FMenuBuilder& MenuBuilder
 					.Padding(4.0f, 0.0f)
 					[
 						SNew(STextBlock)
-						.Text(FText::FromName(LocalVariable.Name))
+						.Text(FText::FromName(ValidVariable.Name))
 					]);
 		}
 	}
diff --git a/Engine/Plugins/Runtime/RigVM/Source/RigVMEditor/Private/Widgets/SRigVMNodePreviewWidget.cpp b/Engine/Plugins/Runtime/RigVM/Source/RigVMEditor/Private/Widgets/SRigVMNodePreviewWidget.cpp
--- a/Engine/Plugins/Runtime/RigVM/Source/RigVMEditor/Private/Widgets/SRigVMNodePreviewWidget.cpp
+++ b/Engine/Plugins/Runtime/RigVM/Source/RigVMEditor/Private/Widgets/SRigVMNodePreviewWidget.cpp
@@ -9,6 +9,8 @@ void SRigVMNodePreviewWidget::Construct(const FArguments& InArgs)
 	Environment = InArgs._Environment;
 	check(Environment);
 	Environment->OnChanged().BindSP(this, &SRigVMNodePreviewWidget::UpdateNodeWidget);
+	NodeEnabledAttribute = InArgs._NodeEnabled;
+	PinDragAndDropEnabledAttribute = InArgs._PinDragAndDropEnabled;
 	
 	ChildSlot
 	.Padding(InArgs._Padding);
@@ -30,6 +32,8 @@ void SRigVMNodePreviewWidget::UpdateNodeWidget()
 	{
 		ChildSlot.AttachWidget(SNew(SRigVMGraphNode)
 			.GraphNodeObj(EdGraphNode)
+			.Enabled(NodeEnabledAttribute)
+			.PinDragDropEnabled(PinDragAndDropEnabledAttribute)
 		);
 	}
 }
diff --git a/Engine/Plugins/Runtime/RigVM/Source/RigVMEditor/Public/Editor/RigVMGraphDetailCustomization.h b/Engine/Plugins/Runtime/RigVM/Source/RigVMEditor/Public/Editor/RigVMGraphDetailCustomization.h
--- a/Engine/Plugins/Runtime/RigVM/Source/RigVMEditor/Public/Editor/RigVMGraphDetailCustomization.h
+++ b/Engine/Plugins/Runtime/RigVM/Source/RigVMEditor/Public/Editor/RigVMGraphDetailCustomization.h
@@ -19,6 +19,7 @@
 #include "SAdvancedTransformInputBox.h"
 #include "Widgets/SRigVMGraphPinNameListValueWidget.h"
 #include "Widgets/SRigVMLogWidget.h"
+#include "Widgets/SRigVMNodePreviewWidget.h"
 #include "HAL/PlatformApplicationMisc.h"
 #include "IDetailDragDropHandler.h"
 
@@ -133,6 +134,8 @@ private:
 
 private:
 
+	typedef TTuple<EPinContainerType, bool, FText, FText, const FSlateBrush*> TContainerType;
+
 	/** Determines if this pin should not be editable */
 	UE_API bool ShouldPinBeReadOnly(bool bIsEditingPinType = false) const;
 
@@ -151,6 +154,13 @@ private:
 	UE_API void PinInfoChanged(const FEdGraphPinType& PinType);
 	UE_API void OnPrePinInfoChange(const FEdGraphPinType& PinType);
 
+	UE_API const TArray<TSharedPtr<TContainerType>> GetContainerTypes() const;
+	UE_API bool CanChangeContainerType() const;
+	UE_API TSharedRef<SWidget> GetContainerTypeMenuContent();
+	UE_API const FSlateBrush* GetContainerTypeIconImage() const;
+	UE_API FSlateColor GetContainerTypeIconColor() const;
+	UE_API void OnContainerTypeSelectionChanged(TSharedPtr<TContainerType> InContainerType);
+
 private:
 
 	/** The argument pin that this layout reflects */
@@ -159,6 +169,9 @@ private:
 	/** The target graph that this argument is on */
 	TWeakObjectPtr<URigVMGraph> GraphPtr;
 
+	/** The ed graph schema that this is using */
+	mutable TWeakObjectPtr<URigVMEdGraphSchema> EdGraphSchemaPtr;
+
 	/** The asset host we are editing */
 	TWeakInterfacePtr<IRigVMClientHost> WeakRigVMClientHost;
 
@@ -170,6 +183,9 @@ private:
 
 	/** The validator to check if a name for an argument is valid */
 	FRigVMLocalVariableNameValidator NameValidator;
+
+	const FSlateBrush* CachedImg_Pin_Disconnected = nullptr;
+	const FSlateBrush* CachedImg_RefPin_Disconnected = nullptr;
 };
 
 class FRigVMFunctionArgumentDefaultNode : public IDetailCustomNodeBuilder, public TSharedFromThis<FRigVMFunctionArgumentDefaultNode>
@@ -200,7 +216,7 @@ private:
 	TWeakObjectPtr<URigVMEdGraph> EdGraphOuterPtr;
 	TWeakInterfacePtr<IRigVMClientHost> WeakRigVMClientHost;
 	FSimpleDelegate OnRebuildChildren;
-	TSharedPtr<SRigVMGraphNode> OwnedNodeWidget;
+	TSharedPtr<SRigVMNodePreviewWidget> PreviewWidget;
 	FDelegateHandle GraphChangedDelegateHandle;
 };
 
diff --git a/Engine/Plugins/Runtime/RigVM/Source/RigVMEditor/Public/Editor/RigVMMinimalEnvironment.h b/Engine/Plugins/Runtime/RigVM/Source/RigVMEditor/Public/Editor/RigVMMinimalEnvironment.h
--- a/Engine/Plugins/Runtime/RigVM/Source/RigVMEditor/Public/Editor/RigVMMinimalEnvironment.h
+++ b/Engine/Plugins/Runtime/RigVM/Source/RigVMEditor/Public/Editor/RigVMMinimalEnvironment.h
@@ -15,7 +15,7 @@
 class FRigVMMinimalEnvironment : public TSharedFromThis<FRigVMMinimalEnvironment>
 {
 public:
-	UE_API FRigVMMinimalEnvironment(const UClass* InRigVMBlueprintClass = nullptr);
+	UE_API FRigVMMinimalEnvironment(const IRigVMAssetInterface* InAssetInterface = nullptr);
 
 	UE_API URigVMGraph* GetModel() const;
 	UE_API URigVMController* GetController() const;
@@ -23,7 +23,7 @@ public:
 	UE_API URigVMEdGraph* GetEdGraph() const;
 	UE_API URigVMEdGraphNode* GetEdGraphNode() const;
 
-	UE_API void SetSchemata(const UClass* InRigVMBlueprintClass);
+	UE_API void SetSchemata(const IRigVMAssetInterface* InAssetInterface);
 	UE_API void SetNode(URigVMNode* InModelNode);
 	UE_API void SetFunctionNode(const FRigVMGraphFunctionIdentifier& InIdentifier);
 	
diff --git a/Engine/Plugins/Runtime/RigVM/Source/RigVMEditor/Public/Widgets/SRigVMGraphNode.h b/Engine/Plugins/Runtime/RigVM/Source/RigVMEditor/Public/Widgets/SRigVMGraphNode.h
--- a/Engine/Plugins/Runtime/RigVM/Source/RigVMEditor/Public/Widgets/SRigVMGraphNode.h
+++ b/Engine/Plugins/Runtime/RigVM/Source/RigVMEditor/Public/Widgets/SRigVMGraphNode.h
@@ -72,9 +72,13 @@ class SRigVMGraphNode : public SGraphNode
 public:
 	SLATE_BEGIN_ARGS(SRigVMGraphNode)
 		: _GraphNodeObj(nullptr)
+		, _Enabled(true)
+		, _PinDragDropEnabled(true)
 		{}
 
 	SLATE_ARGUMENT(URigVMEdGraphNode*, GraphNodeObj)
+	SLATE_ATTRIBUTE(bool, Enabled)
+	SLATE_ATTRIBUTE(bool, PinDragDropEnabled)
 
 	SLATE_END_ARGS()
 
@@ -189,6 +193,7 @@ protected:
 		bool bHasChildren;
 		bool bHideInputWidget;
 		bool bIsContainer;
+		bool bIsInputVariable;
 		int32 Depth;
 		FString Identifier;
 		TSharedPtr<SGraphPin> InputPinWidget;
@@ -207,13 +212,15 @@ protected:
 	public:
 		SLATE_METADATA_TYPE(FPinInfoMetaData, ISlateMetaData)
 
-		FPinInfoMetaData(const FString& InCPPType, const FString& InBoundVariableName)
+		FPinInfoMetaData(const FString& InCPPType, const FString& InBoundVariableName, bool InIsInputVariable)
 		: CPPType(InCPPType)
 		, BoundVariableName(InBoundVariableName)
+		, bIsInputVariable(InIsInputVariable)
 		{}
 
 		FString CPPType;
 		FString BoundVariableName;
+		bool bIsInputVariable;
 	};
 	
 	TArray<FPinInfo> PinInfos;
@@ -227,6 +234,8 @@ protected:
 	// have re-generated the pin list. SRigVMGraphNode has already relinquished them
 	// but we still have a pointer to them in our pin widget.
 	TSet<UEdGraphPin *> PinsToDelete;
+
+	TAttribute<bool> PinDragDropEnabledAttribute;
 };
 
 #undef UE_API
diff --git a/Engine/Plugins/Runtime/RigVM/Source/RigVMEditor/Public/Widgets/SRigVMGraphPinVariableBinding.h b/Engine/Plugins/Runtime/RigVM/Source/RigVMEditor/Public/Widgets/SRigVMGraphPinVariableBinding.h
--- a/Engine/Plugins/Runtime/RigVM/Source/RigVMEditor/Public/Widgets/SRigVMGraphPinVariableBinding.h
+++ b/Engine/Plugins/Runtime/RigVM/Source/RigVMEditor/Public/Widgets/SRigVMGraphPinVariableBinding.h
@@ -19,15 +19,12 @@ public:
 
 	SLATE_BEGIN_ARGS(SRigVMGraphVariableBinding)
 	: _ModelPins()
-    , _FunctionReferenceNode(nullptr)
-    , _InnerVariableName(NAME_None)
     , _Asset(nullptr)
 	, _Blueprint(nullptr)
 	, _CanRemoveBinding(true)
 	{}
 
 		SLATE_ARGUMENT(TArray<URigVMPin*>, ModelPins)
-		SLATE_ARGUMENT(URigVMFunctionReferenceNode*, FunctionReferenceNode)
 		SLATE_ARGUMENT(FName, InnerVariableName)
 		SLATE_ARGUMENT(FRigVMEditorAssetInterfacePtr, Asset)
 		SLATE_ARGUMENT_DEPRECATED(URigVMBlueprint*, Blueprint, 5.7, "Please use Asset instead")
@@ -52,8 +49,6 @@ protected:
 	UE_API void HandleBindToLocalVariable(FRigVMGraphVariableDescription InLocalVariable);
 
 	TArray<URigVMPin*> ModelPins;
-	URigVMFunctionReferenceNode* FunctionReferenceNode;
-	FName InnerVariableName;
 	FRigVMEditorAssetInterfacePtr Blueprint;
 	FPropertyBindingWidgetArgs BindingArgs;
 	bool bCanRemoveBinding;
diff --git a/Engine/Plugins/Runtime/RigVM/Source/RigVMEditor/Public/Widgets/SRigVMNodePreviewWidget.h b/Engine/Plugins/Runtime/RigVM/Source/RigVMEditor/Public/Widgets/SRigVMNodePreviewWidget.h
--- a/Engine/Plugins/Runtime/RigVM/Source/RigVMEditor/Public/Widgets/SRigVMNodePreviewWidget.h
+++ b/Engine/Plugins/Runtime/RigVM/Source/RigVMEditor/Public/Widgets/SRigVMNodePreviewWidget.h
@@ -14,9 +14,13 @@ public:
 	
 	SLATE_BEGIN_ARGS(SRigVMNodePreviewWidget)
 		: _Padding(16.f)
+		, _NodeEnabled(false)
+		, _PinDragAndDropEnabled(false)
 	{}
 	SLATE_ATTRIBUTE(FMargin, Padding)
 	SLATE_ARGUMENT(TSharedPtr<FRigVMMinimalEnvironment>, Environment)
+	SLATE_ATTRIBUTE(bool, NodeEnabled)
+	SLATE_ATTRIBUTE(bool, PinDragAndDropEnabled)
 	SLATE_END_ARGS()
 	
 	UE_API void Construct(const FArguments& InArgs);
@@ -27,6 +31,8 @@ private:
 	UE_API void UpdateNodeWidget();
 	
 	TSharedPtr<FRigVMMinimalEnvironment> Environment;
+	TAttribute<bool> NodeEnabledAttribute;
+	TAttribute<bool> PinDragAndDropEnabledAttribute;
 };
 
 #undef UE_API
