Add new rigVM math node "Get Rotation from X Vector"

#Synced-CL 0

diff --git a/Engine/Plugins/Runtime/RigVM/Source/RigVM/Private/RigVMFunctions/Math/RigVMFunction_MathQuaternion.cpp b/Engine/Plugins/Runtime/RigVM/Source/RigVM/Private/RigVMFunctions/Math/RigVMFunction_MathQuaternion.cpp
--- a/Engine/Plugins/Runtime/RigVM/Source/RigVM/Private/RigVMFunctions/Math/RigVMFunction_MathQuaternion.cpp
+++ b/Engine/Plugins/Runtime/RigVM/Source/RigVM/Private/RigVMFunctions/Math/RigVMFunction_MathQuaternion.cpp
@@ -252,3 +252,7 @@ FRigVMFunction_MathQuaternionMirrorTransform_Execute()
 	Result = Transform.GetRotation();
 }
 
+FRigVMFunction_MathQuaternionFromXVector_Execute()
+{
+	Result = Value.ToOrientationQuat();
+}
diff --git a/Engine/Plugins/Runtime/RigVM/Source/RigVM/Public/RigVMFunctions/Math/RigVMFunction_MathQuaternion.h b/Engine/Plugins/Runtime/RigVM/Source/RigVM/Public/RigVMFunctions/Math/RigVMFunction_MathQuaternion.h
--- a/Engine/Plugins/Runtime/RigVM/Source/RigVM/Public/RigVMFunctions/Math/RigVMFunction_MathQuaternion.h
+++ b/Engine/Plugins/Runtime/RigVM/Source/RigVM/Public/RigVMFunctions/Math/RigVMFunction_MathQuaternion.h
@@ -858,4 +858,31 @@ struct FRigVMFunction_MathQuaternionMirrorTransform : public FRigVMFunction_Math
 	FQuat Result;
 };
 
+/**
+ * Return the Quaternion orientation corresponding to the direction in which the vector points.
+ * Similar to the blueprint "Rotation From X Vector" version, returns a result without roll such that it preserves the up vector.
+ *
+ * @note If you don't care about preserving the up vector and just want the most direct rotation, you can use the faster "From Two Vectors"
+ */
+USTRUCT(meta=(DisplayName="Rotation From X Vector", Keywords = "Orientation, From Axis"))
+struct FRigVMFunction_MathQuaternionFromXVector : public FRigVMFunction_MathQuaternionBase
+{
+	GENERATED_BODY()
+
+	FRigVMFunction_MathQuaternionFromXVector()
+	{
+	}
+	
+	RIGVM_METHOD()
+	UE_API virtual void Execute() override;
+
+	// The input direction vector
+	UPROPERTY(meta=(Input))
+	FVector Value = FVector::XAxisVector;
+
+	// The resulting rotation
+	UPROPERTY(meta=(Output))
+	FQuat Result = FQuat::Identity;
+};
+
 #undef UE_API
\ No newline at end of file
