RigVM: Add more instrumentation to compiler checks

#rnx

#Synced-CL 48723628

diff --git a/Engine/Plugins/Runtime/RigVM/Source/RigVMDeveloper/Private/RigVMCompiler/RigVMCompiler.cpp b/Engine/Plugins/Runtime/RigVM/Source/RigVMDeveloper/Private/RigVMCompiler/RigVMCompiler.cpp
--- a/Engine/Plugins/Runtime/RigVM/Source/RigVMDeveloper/Private/RigVMCompiler/RigVMCompiler.cpp
+++ b/Engine/Plugins/Runtime/RigVM/Source/RigVMDeveloper/Private/RigVMCompiler/RigVMCompiler.cpp
@@ -2798,8 +2798,8 @@ bool URigVMCompiler::TraverseInlineFunction(const FRigVMInlineFunctionExprAST* I
 			if(Instruction.OpCode == ERigVMOpCode::SetupTraits)
 			{
 				const FRigVMSetupTraitsOp& Op = FunctionByteCode.GetOpAt<FRigVMSetupTraitsOp>(Instruction);
-				check(Op.Arg.GetMemoryType() == ERigVMMemoryType::Literal);
-				check(FunctionCompilationData->LiteralPropertyDescriptions.IsValidIndex(Op.Arg.GetRegisterIndex()));
+				checkf(Op.Arg.GetMemoryType() == ERigVMMemoryType::Literal, TEXT("Hit a non-literal arg for a trait, function %s in package %s"), *FunctionReferenceNode->GetFunctionIdentifier().GetLibraryNodePath(), *GetPackage()->GetPathName());
+				checkf(FunctionCompilationData->LiteralPropertyDescriptions.IsValidIndex(Op.Arg.GetRegisterIndex()), TEXT("Hit an invalid register for trait arg, function %s in package %s"), *FunctionReferenceNode->GetFunctionIdentifier().GetLibraryNodePath(), *GetPackage()->GetPathName());
 				const FString& DefaultValue = FunctionCompilationData->LiteralPropertyDescriptions[Op.Arg.GetRegisterIndex()].DefaultValue;
 				if(!DefaultValue.IsEmpty() && DefaultValue != TEXT("()"))
 				{
@@ -2969,6 +2969,7 @@ bool URigVMCompiler::TraverseInlineFunction(const FRigVMInlineFunctionExprAST* I
 
 				if(MemoryType == ERigVMMemoryType::Literal && LiteralValueToSetupTraitArg.Contains(PropertyIndex))
 				{
+					checkf(LiteralValueToSetupTraitArg.Contains(PropertyIndex), TEXT("Trait property index out of range, function %s in package %s"), *FunctionReferenceNode->GetFunctionIdentifier().GetLibraryNodePath(), *GetPackage()->GetPathName());
 					LiteralValueToSetupTraitArg.FindChecked(PropertyIndex) = Operand;
 				}
 			}
@@ -3002,7 +3003,7 @@ bool URigVMCompiler::TraverseInlineFunction(const FRigVMInlineFunctionExprAST* I
 					const int32 OriginalIndex = FCString::Atoi(*IndexStrings[PartIndex]);
 					FRigVMCompilerWorkData::FFunctionRegisterData Data = {FunctionReferenceNode, ERigVMMemoryType::Work, OriginalIndex};
 					const FRigVMOperand& NewOperand = WorkData.FunctionRegisterToOperand.FindChecked(Data);
-					check(NewOperand.GetMemoryType() == ERigVMMemoryType::Work);
+					checkf(NewOperand.GetMemoryType() == ERigVMMemoryType::Work, TEXT("New work operand for trait is not work state, function %s in package %s"), *FunctionReferenceNode->GetFunctionIdentifier().GetLibraryNodePath(), *GetPackage()->GetPathName());
 					IndexStrings[PartIndex] = FString::FromInt(NewOperand.GetRegisterIndex());
 				}
 
@@ -3037,7 +3038,7 @@ bool URigVMCompiler::TraverseInlineFunction(const FRigVMInlineFunctionExprAST* I
 			{
 				break;
 			}
-			check (SourceVarExpr);
+			checkf(SourceVarExpr, TEXT("Hit nullptr SourceVarExpr, function %s in package %s"), *FunctionReferenceNode->GetFunctionIdentifier().GetLibraryNodePath(), *GetPackage()->GetPathName());
 			
 			if(SourceVarExpr->IsExecuteContext())
 			{
@@ -3049,11 +3050,11 @@ bool URigVMCompiler::TraverseInlineFunction(const FRigVMInlineFunctionExprAST* I
 
 			if ( const FString* SegmentPath = InExpr->SegmentPathForChild.Find(Pin->GetPinIndex()))
 			{
-				check(!SegmentPath->IsEmpty());
+				checkf(!SegmentPath->IsEmpty(), TEXT("Hit empty segment path, function %s in package %s"), *FunctionReferenceNode->GetFunctionIdentifier().GetLibraryNodePath(), *GetPackage()->GetPathName());
 				
 				const FString HeadCPPType = SourceVarExpr->GetPin()->GetCPPType();
 				const int32 RegisterOffsetIndex = WorkData.FindOrAddPropertyPath(InterfaceOperands.Last(), HeadCPPType, *SegmentPath);
-				check(RegisterOffsetIndex != INDEX_NONE);
+				checkf(RegisterOffsetIndex != INDEX_NONE, TEXT("Hit invalid register offset, function %s in package %s"), *FunctionReferenceNode->GetFunctionIdentifier().GetLibraryNodePath(), *GetPackage()->GetPathName());
 				InterfaceOperands.Last().RegisterOffset = static_cast<uint16>(RegisterOffsetIndex);
 			}
 		}
@@ -3101,14 +3102,14 @@ bool URigVMCompiler::TraverseInlineFunction(const FRigVMInlineFunctionExprAST* I
 			const FRigVMInstruction& Instruction = Instructions[i];
 			const FRigVMOperandArray OperandArray = ByteCode.GetOperandsForOp(Instruction);
 			const FRigVMOperandArray FunctionOperandArray = FunctionByteCode.GetOperandsForOp(FunctionInstructions[FunctionInstructionIndex]);
-			check(FunctionOperandArray.Num() == OperandArray.Num());
+			checkf(FunctionOperandArray.Num() == OperandArray.Num(), TEXT("Hit a function with a mismatch of function operands, function %s in package %s"), *FunctionReferenceNode->GetFunctionIdentifier().GetLibraryNodePath(), *GetPackage()->GetPathName());
 			
 			for (int32 j=0; j<OperandArray.Num(); ++j)
 			{
 				FRigVMOperand* Operand = const_cast<FRigVMOperand*>(&OperandArray[j]);
 
 				const FRigVMOperand OriginalOperand = *Operand;
-				check(OriginalOperand == FunctionOperandArray[j]);
+				checkf(OriginalOperand == FunctionOperandArray[j], TEXT("Hit a function with an invalid / non-matching function operand, function %s in package %s"), *FunctionReferenceNode->GetFunctionIdentifier().GetLibraryNodePath(), *GetPackage()->GetPathName());
 				
 				FRigVMOperand InterfaceOperand;
 
@@ -3147,7 +3148,10 @@ bool URigVMCompiler::TraverseInlineFunction(const FRigVMInlineFunctionExprAST* I
 					}
 					else
 					{
-						const FRigVMOperand& OuterOperand = WorkData.PinPathToOperand->FindChecked(FString::Printf(TEXT("Variable::%s"), *OuterVariablePath));
+						const FString OuterOperandKey = FString::Printf(TEXT("Variable::%s"), *OuterVariablePath);
+						checkf(WorkData.PinPathToOperand->Contains(OuterOperandKey), TEXT("hit a non-existing operand '%s', function %s in package %s"), *OuterOperandKey, *FunctionReferenceNode->GetFunctionIdentifier().GetLibraryNodePath(), *GetPackage()->GetPathName());
+
+						const FRigVMOperand& OuterOperand = WorkData.PinPathToOperand->FindChecked(OuterOperandKey);
 						Operand->MemoryType = OuterOperand.MemoryType;
 						Operand->RegisterIndex = OuterOperand.RegisterIndex;
 					}
@@ -3164,6 +3168,7 @@ bool URigVMCompiler::TraverseInlineFunction(const FRigVMInlineFunctionExprAST* I
 					// Operand is internal
 					// Replace with added Operand
 					FRigVMCompilerWorkData::FFunctionRegisterData Data = {FunctionReferenceNode, Operand->GetMemoryType(), Operand->GetRegisterIndex()};
+					checkf(WorkData.FunctionRegisterToOperand.Contains(Data), TEXT("hit a missing function operand, function %s in package %s"), *FunctionReferenceNode->GetFunctionIdentifier().GetLibraryNodePath(), *GetPackage()->GetPathName());
 					FRigVMOperand NewOperand = WorkData.FunctionRegisterToOperand.FindChecked(Data);
 					Operand->MemoryType = NewOperand.MemoryType;
 					Operand->RegisterIndex = NewOperand.RegisterIndex;
@@ -3223,7 +3228,7 @@ bool URigVMCompiler::TraverseInlineFunction(const FRigVMInlineFunctionExprAST* I
 				}
 
 				// Store correspondence between Function operand and New operands
-				check(!FunctionOperandToNewOperand.Contains(OriginalOperand) || FunctionOperandToNewOperand.FindChecked(OriginalOperand) == *Operand);
+				checkf(!FunctionOperandToNewOperand.Contains(OriginalOperand) || FunctionOperandToNewOperand.FindChecked(OriginalOperand) == *Operand, TEXT("hit a correspondence issue with a function operand, function %s in package %s"), *FunctionReferenceNode->GetFunctionIdentifier().GetLibraryNodePath(), *GetPackage()->GetPathName());
 				FunctionOperandToNewOperand.Add(OriginalOperand, *Operand);
 
 				// Also find the function debug operands and store correspondence to new debug operands
@@ -3280,6 +3285,7 @@ bool URigVMCompiler::TraverseInlineFunction(const FRigVMInlineFunctionExprAST* I
 			{
 				for (const FRigVMOperand& Operand : FunctionByteCode.InputOperandsPerInstruction[FunctionInstructionIndex])
 				{
+					checkf(FunctionOperandToNewOperand.Contains(Operand), TEXT("hit a missing function input operand during remapping, function %s in package %s"), *FunctionReferenceNode->GetFunctionIdentifier().GetLibraryNodePath(), *GetPackage()->GetPathName());
 					InputOperands.Add(FunctionOperandToNewOperand.FindChecked(Operand));
 				}
 			}
@@ -3287,6 +3293,7 @@ bool URigVMCompiler::TraverseInlineFunction(const FRigVMInlineFunctionExprAST* I
 			{
 				for (const FRigVMOperand& Operand : FunctionByteCode.OutputOperandsPerInstruction[FunctionInstructionIndex])
 				{
+					checkf(FunctionOperandToNewOperand.Contains(Operand), TEXT("hit a missing function output operand during remapping, function %s in package %s"), *FunctionReferenceNode->GetFunctionIdentifier().GetLibraryNodePath(), *GetPackage()->GetPathName());
 					OutputOperands.Add(FunctionOperandToNewOperand.FindChecked(Operand));
 				}
 			}
