RigVM: Allow local functions access to all variables

#rnx

#Synced-CL 48723628

diff --git a/Engine/Plugins/Runtime/RigVM/Source/RigVMDeveloper/Private/RigVMCompiler/RigVMCompiler.cpp b/Engine/Plugins/Runtime/RigVM/Source/RigVMDeveloper/Private/RigVMCompiler/RigVMCompiler.cpp
--- a/Engine/Plugins/Runtime/RigVM/Source/RigVMDeveloper/Private/RigVMCompiler/RigVMCompiler.cpp
+++ b/Engine/Plugins/Runtime/RigVM/Source/RigVMDeveloper/Private/RigVMCompiler/RigVMCompiler.cpp
@@ -590,6 +590,12 @@ bool URigVMCompiler::Compile(const FRigVMCompileSettings& InSettings, TArray<URi
 											}
 										}
 									}
+									else
+									{
+										// if this function doesn't require remapping it's a function used within
+										// the same package. in that case we'll give access to all external variables.
+										FunctionVariables = InExternalVariables;
+									}
 
 									bSuccessfullCompilation = CompileFunction(WorkData.Settings, LibraryNode, FunctionController, FunctionVariables, &FunctionData->CompilationData, OutVMContext);
 								}
