RigVM: A few optimizations to reduce overheads when generating graphs for UAF bindings

Ideally we can build an AST for this in the end, but right now we have to make the operations in the controller a bit faster.

#Synced-CL 48734215

diff --git a/Engine/Plugins/Runtime/RigVM/Source/RigVM/Private/RigVMCore/RigVMDispatchFactory.cpp b/Engine/Plugins/Runtime/RigVM/Source/RigVM/Private/RigVMCore/RigVMDispatchFactory.cpp
--- a/Engine/Plugins/Runtime/RigVM/Source/RigVM/Private/RigVMCore/RigVMDispatchFactory.cpp
+++ b/Engine/Plugins/Runtime/RigVM/Source/RigVM/Private/RigVMCore/RigVMDispatchFactory.cpp
@@ -76,11 +76,6 @@ FText FRigVMDispatchFactory::GetNodeTooltip(const FRigVMTemplateTypeMap& InTypes
 
 FString FRigVMDispatchFactory::GetArgumentDefaultValue(const FName& InArgumentName, TRigVMTypeIndex InTypeIndex) const
 {
-	if(FRigVMRegistry_RWLock::Get().IsArrayType(InTypeIndex))
-	{
-		static const FString EmptyArrayString = TEXT("()");
-		return EmptyArrayString;
-	}
 	if(InTypeIndex == RigVMTypeUtils::TypeIndex::Bool)
 	{
 		static const FString FalseString = TEXT("False");
@@ -100,35 +95,87 @@ FString FRigVMDispatchFactory::GetArgumentDefaultValue(const FName& InArgumentNa
 	{
 		return FName(NAME_None).ToString();
 	}
-	if(InTypeIndex == FRigVMRegistry_RWLock::Get().GetTypeIndex<FVector2D>())
-	{
-		static FString DefaultValueString = GetDefaultValueForStruct(FVector2D::ZeroVector);
-		return DefaultValueString;
-	}
-	if(InTypeIndex == FRigVMRegistry_RWLock::Get().GetTypeIndex<FVector>())
-	{
-		static FString DefaultValueString = GetDefaultValueForStruct(FVector::ZeroVector); 
-		return DefaultValueString;
-	}
-	if(InTypeIndex == FRigVMRegistry_RWLock::Get().GetTypeIndex<FRotator>())
-	{
-		static FString DefaultValueString = GetDefaultValueForStruct(FRotator::ZeroRotator); 
-		return DefaultValueString;
-	}
-	if(InTypeIndex == FRigVMRegistry_RWLock::Get().GetTypeIndex<FQuat>())
+
+	FRigVMRegistry_RWLock::FConditionalReadScopeLock Lock(FRigVMRegistry_RWLock::Get());
+	FRigVMRegistry_NoLock& Registry = Lock.GetRegistry();
+	if(Registry.IsArrayType_NoLock(InTypeIndex))
 	{
-		static FString DefaultValueString = GetDefaultValueForStruct(FQuat::Identity); 
-		return DefaultValueString;
+		static const FString EmptyArrayString = TEXT("()");
+		return EmptyArrayString;
 	}
-	if(InTypeIndex == FRigVMRegistry_RWLock::Get().GetTypeIndex<FTransform>())
+	
+	if (Registry.IsGlobalRegistry())
 	{
-		static FString DefaultValueString = GetDefaultValueForStruct(FTransform::Identity); 
-		return DefaultValueString;
+		// GetTypeIndex_NoLock checks are expensive, so we can use a cached type index for the global registry here
+		static const TRigVMTypeIndex GlobalVector2DIndex = Registry.GetTypeIndex_NoLock<FVector2D>();
+		if(InTypeIndex == GlobalVector2DIndex)
+		{
+			static FString DefaultValueString = GetDefaultValueForStruct(FVector2D::ZeroVector);
+			return DefaultValueString;
+		}
+		static const TRigVMTypeIndex GlobalVectorIndex = Registry.GetTypeIndex_NoLock<FVector>();
+		if(InTypeIndex == GlobalVectorIndex)
+		{
+			static FString DefaultValueString = GetDefaultValueForStruct(FVector::ZeroVector); 
+			return DefaultValueString;
+		}
+		static const TRigVMTypeIndex GlobalRotatorIndex = Registry.GetTypeIndex_NoLock<FRotator>();
+		if(InTypeIndex == GlobalRotatorIndex)
+		{
+			static FString DefaultValueString = GetDefaultValueForStruct(FRotator::ZeroRotator); 
+			return DefaultValueString;
+		}
+		static const TRigVMTypeIndex GlobalQuatIndex = Registry.GetTypeIndex_NoLock<FQuat>();
+		if(InTypeIndex == GlobalQuatIndex)
+		{
+			static FString DefaultValueString = GetDefaultValueForStruct(FQuat::Identity); 
+			return DefaultValueString;
+		}
+		static const TRigVMTypeIndex GlobalTransformIndex = Registry.GetTypeIndex_NoLock<FTransform>();
+		if(InTypeIndex == GlobalTransformIndex)
+		{
+			static FString DefaultValueString = GetDefaultValueForStruct(FTransform::Identity); 
+			return DefaultValueString;
+		}
+		static const TRigVMTypeIndex GlobalLinearColorIndex = Registry.GetTypeIndex_NoLock<FLinearColor>();
+		if(InTypeIndex == GlobalLinearColorIndex)
+		{
+			static FString DefaultValueString = GetDefaultValueForStruct(FLinearColor::White); 
+			return DefaultValueString;
+		}
 	}
-	if(InTypeIndex == FRigVMRegistry_RWLock::Get().GetTypeIndex<FLinearColor>())
+	else
 	{
-		static FString DefaultValueString = GetDefaultValueForStruct(FLinearColor::White); 
-		return DefaultValueString;
+		if(InTypeIndex == Registry.GetTypeIndex_NoLock<FVector2D>())
+		{
+			static FString DefaultValueString = GetDefaultValueForStruct(FVector2D::ZeroVector);
+			return DefaultValueString;
+		}
+		if(InTypeIndex == Registry.GetTypeIndex_NoLock<FVector>())
+		{
+			static FString DefaultValueString = GetDefaultValueForStruct(FVector::ZeroVector); 
+			return DefaultValueString;
+		}
+		if(InTypeIndex == Registry.GetTypeIndex_NoLock<FRotator>())
+		{
+			static FString DefaultValueString = GetDefaultValueForStruct(FRotator::ZeroRotator); 
+			return DefaultValueString;
+		}
+		if(InTypeIndex == Registry.GetTypeIndex_NoLock<FQuat>())
+		{
+			static FString DefaultValueString = GetDefaultValueForStruct(FQuat::Identity); 
+			return DefaultValueString;
+		}
+		if(InTypeIndex == Registry.GetTypeIndex_NoLock<FTransform>())
+		{
+			static FString DefaultValueString = GetDefaultValueForStruct(FTransform::Identity); 
+			return DefaultValueString;
+		}
+		if(InTypeIndex == Registry.GetTypeIndex_NoLock<FLinearColor>())
+		{
+			static FString DefaultValueString = GetDefaultValueForStruct(FLinearColor::White); 
+			return DefaultValueString;
+		}
 	}
 	return FString();
 }
diff --git a/Engine/Plugins/Runtime/RigVM/Source/RigVM/Private/RigVMCore/RigVMTemplate.cpp b/Engine/Plugins/Runtime/RigVM/Source/RigVM/Private/RigVMCore/RigVMTemplate.cpp
--- a/Engine/Plugins/Runtime/RigVM/Source/RigVM/Private/RigVMCore/RigVMTemplate.cpp
+++ b/Engine/Plugins/Runtime/RigVM/Source/RigVM/Private/RigVMCore/RigVMTemplate.cpp
@@ -628,7 +628,7 @@ TArray<TRigVMTypeIndex> FRigVMTemplateArgument::GetSupportedTypeIndices_NoLock(c
 	};
 #endif
 
-	TArray<TRigVMTypeIndex> SupportedTypes;
+	TSet<TRigVMTypeIndex> SupportedTypes;
 	if(InPermutationIndices.IsEmpty())
 	{
 		ForEachType([&](const TRigVMTypeIndex TypeIndex)
@@ -637,9 +637,9 @@ TArray<TRigVMTypeIndex> FRigVMTemplateArgument::GetSupportedTypeIndices_NoLock(c
 			if (TypeIndex != INDEX_NONE)
 			{
 #if UE_RIGVM_DEBUG_TYPEINDEX
-				SupportedTypes.AddUnique(UpdateTypeIndex(TypeIndex));
+				SupportedTypes.Add(UpdateTypeIndex(TypeIndex));
 #else
-				SupportedTypes.AddUnique(TypeIndex);
+				SupportedTypes.Add(TypeIndex);
 #endif
 			}
 			return true;
@@ -647,6 +647,7 @@ TArray<TRigVMTypeIndex> FRigVMTemplateArgument::GetSupportedTypeIndices_NoLock(c
 	}
 	else
 	{
+		SupportedTypes.Reserve(InPermutationIndices.Num());
 		for(const int32 PermutationIndex : InPermutationIndices)
 		{
 			// INDEX_NONE indicates deleted permutation
@@ -654,14 +655,14 @@ TArray<TRigVMTypeIndex> FRigVMTemplateArgument::GetSupportedTypeIndices_NoLock(c
 			if (TypeIndex != INDEX_NONE)
 			{
 #if UE_RIGVM_DEBUG_TYPEINDEX
-				SupportedTypes.AddUnique(UpdateTypeIndex(TypeIndex));
+				SupportedTypes.Add(UpdateTypeIndex(TypeIndex));
 #else
-				SupportedTypes.AddUnique(TypeIndex);
+				SupportedTypes.Add(TypeIndex);
 #endif
 			}
 		}
 	}
-	return SupportedTypes;
+	return SupportedTypes.Array();
 }
 
 #if WITH_EDITOR
@@ -1585,7 +1586,12 @@ const FRigVMExecuteArgument* FRigVMTemplate::GetExecuteArgument(int32 InIndex, c
 
 const FRigVMExecuteArgument* FRigVMTemplate::FindExecuteArgument(const FName& InArgumentName, const FRigVMDispatchContext& InContext) const
 {
-	const TArray<FRigVMExecuteArgument>& Args = GetExecuteArguments(InContext);
+	return FindExecuteArgument_NoLock(InArgumentName, InContext, FRigVMRegistryReadLock());
+}
+
+const FRigVMExecuteArgument* FRigVMTemplate::FindExecuteArgument_NoLock(const FName& InArgumentName, const FRigVMDispatchContext& InContext, const FRigVMRegistryHandle& InRegistry) const
+{
+	const TArray<FRigVMExecuteArgument>& Args = GetExecuteArguments_NoLock(InContext, InRegistry);
 	return Args.FindByPredicate([InArgumentName](const FRigVMExecuteArgument& Arg) -> bool
 	{
 		return Arg.Name == InArgumentName;
@@ -1930,11 +1936,12 @@ bool FRigVMTemplate::Resolve_NoLock(FTypeMap& InOutTypes, TArray<int32>& OutPerm
 			// Using a map to collect all permutations that we can keep/remove
 			// instead of removing them one by one, which can be costly
 			TMap<int32, bool> PermutationsToKeep;
-			PermutationsToKeep.Reserve(AllTypes.Num());
+			const int32 NumAllTypes = AllTypes.Num();
+			PermutationsToKeep.Reserve(NumAllTypes);
 
-			for (int32 PermutationIndex = 0; PermutationIndex < AllTypes.Num(); PermutationIndex++)
+			for (int32 PermutationIndex = 0; PermutationIndex < NumAllTypes; PermutationIndex++)
 			{
-				const TRigVMTypeIndex Type = AllTypes[PermutationIndex];
+				const TRigVMTypeIndex& Type = AllTypes[PermutationIndex];
 				if(!InRegistry->CanMatchTypes_NoLock(Type, *InputType, bAllowFloatingPointCasts))
 				{
 					PermutationsToKeep.Add(PermutationIndex, false);
diff --git a/Engine/Plugins/Runtime/RigVM/Source/RigVM/Private/RigVMStringUtils.cpp b/Engine/Plugins/Runtime/RigVM/Source/RigVM/Private/RigVMStringUtils.cpp
--- a/Engine/Plugins/Runtime/RigVM/Source/RigVM/Private/RigVMStringUtils.cpp
+++ b/Engine/Plugins/Runtime/RigVM/Source/RigVM/Private/RigVMStringUtils.cpp
@@ -2,6 +2,8 @@
 
 #include "RigVMStringUtils.h"
 
+#include "String/Split.h"
+
 FString RigVMStringUtils::JoinStrings(const FString& InStringA, const FString& InStringB, const TCHAR* InSeparator, const TCHAR* InPrefix, const TCHAR* InSuffix)
 {
 	const int32 SeparatorLen = FCString::Strlen(InSeparator);
@@ -158,6 +160,11 @@ bool RigVMStringUtils::SplitPinPathAtStart(const FString& InPinPath, FString& Le
 	return InPinPath.Split(TEXT("."), &LeftMost, &Right, ESearchCase::IgnoreCase, ESearchDir::FromStart);
 }
 
+bool RigVMStringUtils::SplitPinPathAtStart(FStringView InPinPath, FStringView& LeftMost, FStringView& Right)
+{
+	return UE::String::SplitFirst(InPinPath, TEXT("."), LeftMost, Right, ESearchCase::IgnoreCase);
+}
+
 bool RigVMStringUtils::SplitPinPathAtEnd(const FString& InPinPath, FString& Left, FString& RightMost)
 {
 	return InPinPath.Split(TEXT("."), &Left, &RightMost, ESearchCase::IgnoreCase, ESearchDir::FromEnd);
diff --git a/Engine/Plugins/Runtime/RigVM/Source/RigVM/Public/RigVMCore/RigVMTemplate.h b/Engine/Plugins/Runtime/RigVM/Source/RigVM/Public/RigVMCore/RigVMTemplate.h
--- a/Engine/Plugins/Runtime/RigVM/Source/RigVM/Public/RigVMCore/RigVMTemplate.h
+++ b/Engine/Plugins/Runtime/RigVM/Source/RigVM/Public/RigVMCore/RigVMTemplate.h
@@ -199,7 +199,9 @@ struct FRigVMTemplateArgument
 		ETypeCategory_SingleObjectValue,
 		ETypeCategory_ArrayObjectValue,
 		ETypeCategory_ArrayArrayObjectValue,
-		ETypeCategory_Invalid
+		ETypeCategory_Invalid,
+		
+		ETypeCategory_Count,
 	};
 
 	// default constructor
@@ -369,7 +371,7 @@ private:
 		}
 		
 	private:
-		TArray<TArrayView<const TRigVMTypeIndex>> Types;
+		TArray<TArrayView<const TRigVMTypeIndex>, TInlineAllocator<(uint32)ETypeCategory::ETypeCategory_Count>> Types;
 	};
 };
 
@@ -462,6 +464,7 @@ public:
 
 	// returns an argument given a name (or nullptr)
 	RIGVM_API const FRigVMExecuteArgument* FindExecuteArgument(const FName& InArgumentName, const FRigVMDispatchContext& InContext) const;
+	RIGVM_API const FRigVMExecuteArgument* FindExecuteArgument_NoLock(const FName& InArgumentName, const FRigVMDispatchContext& InContext, const FRigVMRegistryHandle& InRegistry) const;
 
 	// returns the top level execute context struct this template uses
 	RIGVM_API const UScriptStruct* GetExecuteContextStruct() const;
diff --git a/Engine/Plugins/Runtime/RigVM/Source/RigVM/Public/RigVMStringUtils.h b/Engine/Plugins/Runtime/RigVM/Source/RigVM/Public/RigVMStringUtils.h
--- a/Engine/Plugins/Runtime/RigVM/Source/RigVM/Public/RigVMStringUtils.h
+++ b/Engine/Plugins/Runtime/RigVM/Source/RigVM/Public/RigVMStringUtils.h
@@ -37,6 +37,7 @@ namespace RigVMStringUtils
 
 	// Splits a PinPath at the start, so for example "Node.Color.R" becomes "Node" and "Color.R"
 	RIGVM_API bool SplitPinPathAtStart(const FString& InPinPath, FString& LeftMost, FString& Right);
+	RIGVM_API bool SplitPinPathAtStart(FStringView InPinPath, FStringView& LeftMost, FStringView& Right);
 
 	// Splits a PinPath at the start, so for example "Node.Color.R" becomes "Node.Color" and "R"
 	RIGVM_API bool SplitPinPathAtEnd(const FString& InPinPath, FString& Left, FString& RightMost);
diff --git a/Engine/Plugins/Runtime/RigVM/Source/RigVMDeveloper/Private/RigVMModel/Nodes/RigVMTemplateNode.cpp b/Engine/Plugins/Runtime/RigVM/Source/RigVMDeveloper/Private/RigVMModel/Nodes/RigVMTemplateNode.cpp
--- a/Engine/Plugins/Runtime/RigVM/Source/RigVMDeveloper/Private/RigVMModel/Nodes/RigVMTemplateNode.cpp
+++ b/Engine/Plugins/Runtime/RigVM/Source/RigVMDeveloper/Private/RigVMModel/Nodes/RigVMTemplateNode.cpp
@@ -667,7 +667,7 @@ FString URigVMTemplateNode::GetInitialDefaultValueForPin(const FName& InRootPinN
 	return DefaultValue;
 }
 
-FName URigVMTemplateNode::GetDisplayNameForPin(const URigVMPin* InPin) const
+FName URigVMTemplateNode::GetDisplayNameForPin(const URigVMPin* InPin, const TArray<int32>& InPermutationIndices) const
 {
 #if WITH_EDITOR
 	if(const FRigVMTemplate* Template = GetTemplate())
@@ -691,9 +691,7 @@ FName URigVMTemplateNode::GetDisplayNameForPin(const URigVMPin* InPin) const
 			}
 		
 			const FName RootPinName = InPin->GetFName();
-			const TArray<int32> AllPermutations = GetResolvedPermutationIndices(false);
-
-			const FText DisplayNameText = Template->GetDisplayNameForArgument(RootPinName, AllPermutations);
+			const FText DisplayNameText = Template->GetDisplayNameForArgument(RootPinName, InPermutationIndices.Num() > 0 ? InPermutationIndices : GetResolvedPermutationIndices(false));
 			if(DisplayNameText.IsEmpty())
 			{
 				return NAME_None;
@@ -718,7 +716,17 @@ FName URigVMTemplateNode::GetDisplayNameForPin(const URigVMPin* InPin) const
 	return Super::GetDisplayNameForPin(InPin);
 }
 
+FName URigVMTemplateNode::GetDisplayNameForPin(const URigVMPin* InPin) const
+{
+	return GetDisplayNameForPin(InPin, TArray<int32>());
+}
+
 TRigVMTypeIndex URigVMTemplateNode::TryReduceTypesToSingle(const TArray<TRigVMTypeIndex>& InTypes, const TRigVMTypeIndex PreferredType) const
+{
+	return TryReduceTypesToSingle_NoLock(InTypes, FRigVMRegistryReadLock(), PreferredType);
+}
+
+TRigVMTypeIndex URigVMTemplateNode::TryReduceTypesToSingle_NoLock(const TArray<TRigVMTypeIndex>& InTypes, const FRigVMRegistryHandle& InRegistry, const TRigVMTypeIndex PreferredType) const
 {
 	if (InTypes.IsEmpty())
 	{
@@ -732,7 +740,7 @@ TRigVMTypeIndex URigVMTemplateNode::TryReduceTypesToSingle(const TArray<TRigVMTy
 	
 	for (int32 i=1; i<InTypes.Num(); ++i)
 	{
-		if (!FRigVMRegistry::Get().CanMatchTypes(InTypes[0], InTypes[i], true))
+		if (!InRegistry->CanMatchTypes_NoLock(InTypes[0], InTypes[i], true))
 		{
 			return INDEX_NONE;
 		}
diff --git a/Engine/Plugins/Runtime/RigVM/Source/RigVMDeveloper/Private/RigVMModel/RigVMController.cpp b/Engine/Plugins/Runtime/RigVM/Source/RigVMDeveloper/Private/RigVMModel/RigVMController.cpp
--- a/Engine/Plugins/Runtime/RigVM/Source/RigVMDeveloper/Private/RigVMModel/RigVMController.cpp
+++ b/Engine/Plugins/Runtime/RigVM/Source/RigVMDeveloper/Private/RigVMModel/RigVMController.cpp
@@ -20556,26 +20556,25 @@ void URigVMController::OverrideDefaultValueMember(const FString& InMemberName, c
 	InOutDefaultValue = FString::Printf(TEXT("(%s)"), *FString::Join(DefaultValues, TEXT(",")));	
 }
 
-void URigVMController::ResolveTemplateNodeMetaData(URigVMTemplateNode* InNode, bool bSetupUndoRedo)
+void URigVMController::ResolveTemplateNodeMetaData(URigVMTemplateNode* InNode, const TArray<int32>& InFilteredPermutationIndices, bool bSetupUndoRedo)
 {
-#if WITH_EDITOR
 	check(InNode);
-	
-	TArray<int32> FilteredPermutationIndices = InNode->GetResolvedPermutationIndices(false);
+
 	if(InNode->IsA<URigVMUnitNode>())
 	{
 		const FLinearColor PreviousColor = InNode->NodeColor;
-		InNode->NodeColor = InNode->GetTemplate()->GetColor(FilteredPermutationIndices);
+#if WITH_EDITOR
+		InNode->NodeColor = InNode->GetTemplate()->GetColor(InFilteredPermutationIndices);
+#endif
 		if(!InNode->NodeColor.Equals(PreviousColor, 0.01f))
 		{
 			Notify(ERigVMGraphNotifType::NodeColorChanged, InNode);
 		}
 	}
-#endif
 
 	for(URigVMPin* Pin : InNode->GetPins())
 	{
-		const FName DisplayName = InNode->GetDisplayNameForPin(Pin);
+		const FName DisplayName = InNode->GetDisplayNameForPin(Pin, InFilteredPermutationIndices);
 		if(Pin->DisplayName != DisplayName)
 		{
 			Pin->DisplayName = DisplayName;
@@ -20593,7 +20592,7 @@ void URigVMController::ResolveTemplateNodeMetaData(URigVMTemplateNode* InNode, b
 			}
 			if(!Pin->IsValidDefaultValue(Pin->GetDefaultValue()))
 			{
-				const FString NewDefaultValue = InNode->GetInitialDefaultValueForPin(Pin->GetFName(), FilteredPermutationIndices);
+				const FString NewDefaultValue = InNode->GetInitialDefaultValueForPin(Pin->GetFName(), InFilteredPermutationIndices);
 				if(!NewDefaultValue.IsEmpty())
 				{
 					SetPinDefaultValue(Pin, NewDefaultValue, true, bSetupUndoRedo, false);
@@ -20603,6 +20602,12 @@ void URigVMController::ResolveTemplateNodeMetaData(URigVMTemplateNode* InNode, b
 	}
 }
 
+void URigVMController::ResolveTemplateNodeMetaData(URigVMTemplateNode* InNode, bool bSetupUndoRedo)
+{
+	TArray<int32> FilteredPermutationIndices = InNode->GetResolvedPermutationIndices(false);
+	ResolveTemplateNodeMetaData(InNode, FilteredPermutationIndices, bSetupUndoRedo);
+}
+
 bool URigVMController::FullyResolveTemplateNode(URigVMTemplateNode* InNode, int32 InPermutationIndex, bool bSetupUndoRedo)
 {
 	if(bIsFullyResolvingTemplateNode)
@@ -20626,9 +20631,10 @@ bool URigVMController::FullyResolveTemplateNode(URigVMTemplateNode* InNode, int3
 	InNode->ResolvedPermutation = InPermutationIndex;
 	
 	// Figure out the permutation index from the pin types
+	TArray<int32> Permutations;
 	if (InPermutationIndex == INDEX_NONE)
 	{
-		TArray<int32> Permutations = InNode->GetResolvedPermutationIndices(false);
+		Permutations = InNode->GetResolvedPermutationIndices(false);
 
 		// If some float/double pin needs to change type, permutations might be empty
 		// Try running again allowing that change
@@ -21108,27 +21114,31 @@ bool URigVMController::UpdateTemplateNodePinTypes(URigVMTemplateNode* InNode, bo
 	TArray<int32> ResolvedPermutations = InNode->GetResolvedPermutationIndices(true);
 	InNode->ResolvedPermutation = ResolvedPermutations.Num() == 1 ? ResolvedPermutations[0] : INDEX_NONE;
 
-	TArray<URigVMPin*> Pins = InNode->GetPins();
+	const TArray<URigVMPin*>& Pins = InNode->GetPins();
 	TArray<const FRigVMTemplateArgument*> Arguments;
-	Arguments.SetNumUninitialized(Pins.Num());
-	for(int32 PinIndex=0; PinIndex < Pins.Num(); ++PinIndex)
+	const int32 NumPins = Pins.Num();
+	Arguments.SetNumUninitialized(NumPins);
+	for(int32 PinIndex=0; PinIndex < NumPins; ++PinIndex)
 	{
 		URigVMPin* Pin = Pins[PinIndex];
 		Arguments[PinIndex] = Template->FindArgument(Pin->GetFName());
 	}
 
 	// Remove invalid permutations
-	ResolvedPermutations.RemoveAll([&](const int32& Permutation)
 	{
-		for (const FRigVMTemplateArgument* Argument : Arguments)
+		FRigVMRegistryReadLock RegistryReadLock(Registry);
+		ResolvedPermutations.RemoveAllSwap([&](const int32& Permutation)
 		{
-			if (Argument && Argument->GetTypeIndex(Permutation) == INDEX_NONE)
+			for (const FRigVMTemplateArgument* Argument : Arguments)
 			{
-				return true;
+				if (Argument && Argument->GetTypeIndex_NoLock(Permutation, RegistryReadLock) == INDEX_NONE)
+				{
+					return true;
+				}
 			}
-		}
-		return false;
-	});
+			return false;
+		}, EAllowShrinking::No);
+	}
 
 	FRigVMDispatchContext DispatchContext;
 	if(const URigVMDispatchNode* DispatchNode = Cast<URigVMDispatchNode>(InNode))
@@ -21138,9 +21148,10 @@ bool URigVMController::UpdateTemplateNodePinTypes(URigVMTemplateNode* InNode, bo
 
 	// Find the types for each possible permutation
 	TMap<int32, TArray<TRigVMTypeIndex>> PinTypes; // permutation to pin types
+	PinTypes.Reserve(ResolvedPermutations.Num() * NumPins);
 	for (int32 ResolvedPermutation : ResolvedPermutations)
 	{
-		for(int32 PinIndex=0; PinIndex < Pins.Num(); ++PinIndex)
+		for(int32 PinIndex=0; PinIndex < NumPins; ++PinIndex)
 		{
 			URigVMPin* Pin = Pins[PinIndex];
 			if (Pin->GetDirection() == ERigVMPinDirection::Hidden)
@@ -21176,7 +21187,7 @@ bool URigVMController::UpdateTemplateNodePinTypes(URigVMTemplateNode* InNode, bo
 	// If reduction happens for multiple pins, we need to make sure that they reduce to the same permutation
 	// If possible, we want to respect the current pin type if its not wildcard
 	TArray<bool> WasReduced;
-	WasReduced.SetNumZeroed(Pins.Num());
+	WasReduced.SetNumZeroed(NumPins);
 
 	TMap<int32, TArray<TRigVMTypeIndex>> ReducedTypes = PinTypes;
 	TArray<TRigVMTypeIndex> FinalPinTypes;
@@ -21184,7 +21195,7 @@ bool URigVMController::UpdateTemplateNodePinTypes(URigVMTemplateNode* InNode, bo
 	
 	// Reduce compatible types, try to find the permutation in the reduced types (from the proposed types)
 	FRigVMTemplateTypeMap AlreadyResolvedTypeMap;
-	for(int32 PinIndex=0; PinIndex < Pins.Num(); ++PinIndex)
+	for(int32 PinIndex=0; PinIndex < NumPins; ++PinIndex)
 	{
 		if (WasReduced[PinIndex])
 		{
@@ -21229,8 +21240,8 @@ bool URigVMController::UpdateTemplateNodePinTypes(URigVMTemplateNode* InNode, bo
 	{
 		ensureMsgf(!ReducedTypes.IsEmpty(), TEXT("Found incompatible preferred types"));
 	}
-	
-	for(int32 PinIndex=0; PinIndex < Pins.Num(); ++PinIndex)
+
+	for(int32 PinIndex=0; PinIndex < NumPins; ++PinIndex)
 	{
 		URigVMPin* Pin = Pins[PinIndex];
 		if (Pin->GetDirection() == ERigVMPinDirection::Hidden)
@@ -21283,6 +21294,7 @@ bool URigVMController::UpdateTemplateNodePinTypes(URigVMTemplateNode* InNode, bo
 
 		// Remove reduced types which do not match this type
 		TArray<int32> PermutationToRemove;
+		PermutationToRemove.Reserve(ReducedTypes.Num());
 		for (TPair<int32, TArray<TRigVMTypeIndex>>& Pair : ReducedTypes)
 		{
 			if (Pair.Value[PinIndex] != FinalPinTypes[PinIndex])
@@ -21299,7 +21311,7 @@ bool URigVMController::UpdateTemplateNodePinTypes(URigVMTemplateNode* InNode, bo
 	// First unresolve any pins that need unresolving, then resolve everything else
 	for (int32 UnresolveResolve=0; UnresolveResolve<2; ++UnresolveResolve)
 	{
-		for(int32 PinIndex=0; PinIndex < Pins.Num(); ++PinIndex)
+		for(int32 PinIndex=0; PinIndex < NumPins; ++PinIndex)
 		{
 			URigVMPin* Pin = Pins[PinIndex];
 			if (UnresolveResolve == 0 && Pin->IsWildCard())
diff --git a/Engine/Plugins/Runtime/RigVM/Source/RigVMDeveloper/Private/RigVMModel/RigVMNode.cpp b/Engine/Plugins/Runtime/RigVM/Source/RigVMDeveloper/Private/RigVMModel/RigVMNode.cpp
--- a/Engine/Plugins/Runtime/RigVM/Source/RigVMDeveloper/Private/RigVMModel/RigVMNode.cpp
+++ b/Engine/Plugins/Runtime/RigVM/Source/RigVMDeveloper/Private/RigVMModel/RigVMNode.cpp
@@ -524,7 +524,7 @@ ERigVMNodeDefaultValueOverrideState::Type URigVMNode::GetPinDefaultValueOverride
 
 URigVMPin* URigVMNode::FindPin(const FString& InPinPath) const
 {
-	FString Left, Right;
+	FStringView Left, Right;
 	if (!URigVMPin::SplitPinPathAtStart(InPinPath, Left, Right))
 	{
 		Left = InPinPath;
@@ -542,7 +542,7 @@ URigVMPin* URigVMNode::FindPin(const FString& InPinPath) const
 			{
 				return Pin;
 			}
-			return Pin->FindSubPin(Right);
+			return Pin->FindSubPinView(Right);
 		}
 	}
 
@@ -564,7 +564,7 @@ URigVMPin* URigVMNode::FindPin(const FString& InPinPath) const
 				{
 					return Pin;
 				}
-				return Pin->FindSubPin(Right);
+				return Pin->FindSubPinView(Right);
 			}
 		}
 	}
diff --git a/Engine/Plugins/Runtime/RigVM/Source/RigVMDeveloper/Private/RigVMModel/RigVMPin.cpp b/Engine/Plugins/Runtime/RigVM/Source/RigVMDeveloper/Private/RigVMModel/RigVMPin.cpp
--- a/Engine/Plugins/Runtime/RigVM/Source/RigVMDeveloper/Private/RigVMModel/RigVMPin.cpp
+++ b/Engine/Plugins/Runtime/RigVM/Source/RigVMDeveloper/Private/RigVMModel/RigVMPin.cpp
@@ -70,6 +70,11 @@ bool URigVMPin::SplitPinPathAtStart(const FString& InPinPath, FString& LeftMost,
 	return RigVMStringUtils::SplitPinPathAtStart(InPinPath, LeftMost, Right);
 }
 
+bool URigVMPin::SplitPinPathAtStart(FStringView InPinPath, FStringView& LeftMost, FStringView& Right)
+{
+	return RigVMStringUtils::SplitPinPathAtStart(InPinPath, LeftMost, Right);
+}
+
 bool URigVMPin::SplitPinPathAtEnd(const FString& InPinPath, FString& Left, FString& RightMost)
 {
 	return RigVMStringUtils::SplitPinPathAtEnd(InPinPath, Left, RightMost);
@@ -144,7 +149,7 @@ URigVMPin::URigVMPin()
 {
 }
 
-bool URigVMPin::NameEquals(const FString& InName, bool bFollowCoreRedirectors) const
+bool URigVMPin::NameEquals(FStringView InName, bool bFollowCoreRedirectors) const
 {
 	if(InName.Equals(GetName(), ESearchCase::IgnoreCase))
 	{
@@ -180,7 +185,7 @@ bool URigVMPin::NameEquals(const FString& InName, bool bFollowCoreRedirectors) c
 				return NameEquals(*RedirectedNamePtr->ToString(), false);
 			}
 
-			const FCoreRedirectObjectName OldObjectName(*InName, Struct->GetFName(), *Struct->GetOutermost()->GetPathName());
+			const FCoreRedirectObjectName OldObjectName(FName(InName), Struct->GetFName(), *Struct->GetOutermost()->GetPathName());
 			const FCoreRedirectObjectName NewObjectName = FCoreRedirects::GetRedirectedName(ECoreRedirectFlags::Type_Property, OldObjectName);
 			if (OldObjectName != NewObjectName)
 			{
@@ -2145,9 +2150,9 @@ TArray<URigVMPin*> URigVMPin::GetAllSubPinsRecursively() const
 	return AllSubPins;
 }
 
-URigVMPin* URigVMPin::FindSubPin(const FString& InPinPath) const
+URigVMPin* URigVMPin::FindSubPinView(FStringView InPinPath) const
 {
-	FString Left, Right;
+	FStringView Left, Right;
 	if (!URigVMPin::SplitPinPathAtStart(InPinPath, Left, Right))
 	{
 		Left = InPinPath;
@@ -2161,12 +2166,17 @@ URigVMPin* URigVMPin::FindSubPin(const FString& InPinPath) const
 			{
 				return Pin;
 			}
-			return Pin->FindSubPin(Right);
+			return Pin->FindSubPinView(Right);
 		}
 	}
 	return nullptr;
 }
 
+URigVMPin* URigVMPin::FindSubPin(const FString& InPinPath) const
+{
+	return FindSubPinView(InPinPath);
+}
+
 bool URigVMPin::IsLinkedTo(const URigVMPin* InPin) const
 {
 	for (const URigVMLink* Link : Links)
diff --git a/Engine/Plugins/Runtime/RigVM/Source/RigVMDeveloper/Public/RigVMModel/Nodes/RigVMTemplateNode.h b/Engine/Plugins/Runtime/RigVM/Source/RigVMDeveloper/Public/RigVMModel/Nodes/RigVMTemplateNode.h
--- a/Engine/Plugins/Runtime/RigVM/Source/RigVMDeveloper/Public/RigVMModel/Nodes/RigVMTemplateNode.h
+++ b/Engine/Plugins/Runtime/RigVM/Source/RigVMDeveloper/Public/RigVMModel/Nodes/RigVMTemplateNode.h
@@ -136,10 +136,12 @@ public:
 
 	// returns the display name for a pin
 	UE_API virtual FName GetDisplayNameForPin(const URigVMPin* InPin) const override;
+	UE_API FName GetDisplayNameForPin(const URigVMPin* InPin, const TArray<int32>& InPermutationIndices) const;
 
 	// Tries to reduce the input types to a single type, if all are compatible
 	// Will prioritize the InPreferredType if available
 	UE_API TRigVMTypeIndex TryReduceTypesToSingle(const TArray<TRigVMTypeIndex>& InTypes, const TRigVMTypeIndex PreferredType = TRigVMTypeIndex()) const;
+	UE_API TRigVMTypeIndex TryReduceTypesToSingle_NoLock(const TArray<TRigVMTypeIndex>& InTypes, const FRigVMRegistryHandle& InRegistry, const TRigVMTypeIndex PreferredType = TRigVMTypeIndex()) const;
 
 	UE_API virtual uint32 GetStructureHash() const override;
 
diff --git a/Engine/Plugins/Runtime/RigVM/Source/RigVMDeveloper/Public/RigVMModel/RigVMController.h b/Engine/Plugins/Runtime/RigVM/Source/RigVMDeveloper/Public/RigVMModel/RigVMController.h
--- a/Engine/Plugins/Runtime/RigVM/Source/RigVMDeveloper/Public/RigVMModel/RigVMController.h
+++ b/Engine/Plugins/Runtime/RigVM/Source/RigVMDeveloper/Public/RigVMModel/RigVMController.h
@@ -1440,6 +1440,7 @@ private:
 	static UE_API void OverrideDefaultValueMember(const FString& InMemberName, const FString& InMemberValue, FString& InOutDefaultValue);
 
 	UE_API void ResolveTemplateNodeMetaData(URigVMTemplateNode* InNode, bool bSetupUndoRedo);
+	UE_API void ResolveTemplateNodeMetaData(URigVMTemplateNode* InNode, const TArray<int32>& InFilteredPermutationIndices, bool bSetupUndoRedo);
 
 	// Changes Pin types if filtered types of a pin are unique
 	UE_API bool UpdateTemplateNodePinTypes(URigVMTemplateNode* InNode, bool bSetupUndoRedo, bool bInitializeDefaultValue = true, TMap<URigVMPin*, TArray<TRigVMTypeIndex>> ProposedTypes = TMap<URigVMPin*, TArray<TRigVMTypeIndex>>());
diff --git a/Engine/Plugins/Runtime/RigVM/Source/RigVMDeveloper/Public/RigVMModel/RigVMPin.h b/Engine/Plugins/Runtime/RigVM/Source/RigVMDeveloper/Public/RigVMModel/RigVMPin.h
--- a/Engine/Plugins/Runtime/RigVM/Source/RigVMDeveloper/Public/RigVMModel/RigVMPin.h
+++ b/Engine/Plugins/Runtime/RigVM/Source/RigVMDeveloper/Public/RigVMModel/RigVMPin.h
@@ -138,6 +138,7 @@ public:
 
 	// Splits a PinPath at the start, so for example "Node.Color.R" becomes "Node" and "Color.R"
 	static UE_API bool SplitPinPathAtStart(const FString& InPinPath, FString& LeftMost, FString& Right);
+	static UE_API bool SplitPinPathAtStart(FStringView InPinPath, FStringView& LeftMost, FStringView& Right);
 
 	// Splits a PinPath at the start, so for example "Node.Color.R" becomes "Node.Color" and "R"
 	static UE_API bool SplitPinPathAtEnd(const FString& InPinPath, FString& Left, FString& RightMost);
@@ -160,7 +161,7 @@ public:
 	UE_API URigVMPin();
 
 	// returns true if the name of this pin matches a given name
-	UE_API bool NameEquals(const FString& InName, bool bFollowCoreRedirectors = false) const;
+	UE_API bool NameEquals(FStringView InName, bool bFollowCoreRedirectors = false) const;
 
 	// Returns a . separated path containing all names of the pin and its owners,
 	// this includes the node name, for example "Node.Color.R"
@@ -441,6 +442,7 @@ public:
 	// Returns a SubPin given a name / path or nullptr.
 	UFUNCTION(BlueprintCallable, Category = RigVMPin)
 	UE_API URigVMPin* FindSubPin(const FString& InPinPath) const;
+	UE_API URigVMPin* FindSubPinView(FStringView InPinPath) const;
 
 	// Returns true if this Pin is linked to another Pin
 	UFUNCTION(BlueprintCallable, Category = RigVMPin)
