Control Rig: Fix crash when opening AncientGame
#rnx

#Synced-CL 48734632

diff --git a/Engine/Plugins/Runtime/RigVM/Source/RigVMDeveloper/Private/RigVMEditorAsset.cpp b/Engine/Plugins/Runtime/RigVM/Source/RigVMDeveloper/Private/RigVMEditorAsset.cpp
--- a/Engine/Plugins/Runtime/RigVM/Source/RigVMDeveloper/Private/RigVMEditorAsset.cpp
+++ b/Engine/Plugins/Runtime/RigVM/Source/RigVMDeveloper/Private/RigVMEditorAsset.cpp
@@ -47,7 +47,6 @@
 #include "RigVMEditorModule.h"
 #include "ScopedTransaction.h"
 #include "UnrealEdGlobals.h"
-#include "IPackageAutoSaver.h"
 #if !WITH_RIGVMLEGACYEDITOR
 #include "RigVMEditor/Private/Editor/Kismet/RigVMBlueprintCompilationManager.h"
 #endif
@@ -1470,7 +1469,7 @@ void IRigVMEditorAssetInterface::SerializeImpl(FArchive& Ar)
 	PostSerialize(Ar);
 
 #if WITH_EDITOR
-	if (GUnrealEd && GUnrealEd->GetPackageAutoSaver().IsAutoSaving())
+	if (GUnrealEd && GUnrealEd->IsAutosaving())
 	{
 		TWeakObjectPtr<UObject> WeakThis(GetObject());
 
