RigVM: Rely on fuzzy string match in context menu

#rnx

#Synced-CL 48757204

diff --git a/Engine/Plugins/Runtime/RigVM/Source/RigVMDeveloper/Private/EdGraph/RigVMEdGraphSchema.cpp b/Engine/Plugins/Runtime/RigVM/Source/RigVMDeveloper/Private/EdGraph/RigVMEdGraphSchema.cpp
--- a/Engine/Plugins/Runtime/RigVM/Source/RigVMDeveloper/Private/EdGraph/RigVMEdGraphSchema.cpp
+++ b/Engine/Plugins/Runtime/RigVM/Source/RigVMDeveloper/Private/EdGraph/RigVMEdGraphSchema.cpp
@@ -30,6 +30,7 @@
 #include "RigVMModel/Nodes/RigVMDispatchNode.h"
 #include "RigVMEditorAsset.h"
 #include "RigVMDeveloperTypeUtils.h"
+#include "Algo/LevenshteinDistance.h"
 #if WITH_RIGVMLEGACYEDITOR
 #include "BlueprintEditor.h"
 #else
@@ -1066,6 +1067,113 @@ bool URigVMEdGraphSchema::SupportsPreviewFromHereOnNode(const UEdGraphNode* InNo
 	return false;
 }
 
+float URigVMEdGraphSchema::GetActionFilteredWeight(const FEdGraphSchemaAction& InCurrentAction, const TArray<FString>& InFilterTerms,
+	const TArray<FString>& InSanitizedFilterTerms, const TArray<UEdGraphPin*>& DraggedFromPins) const
+{
+	TArray<FString> SanitizedFilterTerms = InSanitizedFilterTerms;
+	for (FString& SanitizedFilterTerm : SanitizedFilterTerms)
+	{
+		SanitizedFilterTerm.ToLowerInline();
+	}
+	TArray<FString> MenuDescriptions = InCurrentAction.GetMenuDescriptionArray();
+	for (FString& MenuDescription : MenuDescriptions)
+	{
+		MenuDescription.ToLowerInline();
+	}
+	TArray<FString> MenuDescriptionsReversed = MenuDescriptions;
+	Algo::Reverse(MenuDescriptionsReversed);
+	
+	const TArray<FString>& Keywords = InCurrentAction.GetSearchKeywordsArray();
+
+	auto LevenshteinSingle = [](const FStringView& A, const FStringView& B) -> float
+	{
+		float WorstCase = static_cast<float>(A.Len() + B.Len());
+		WorstCase = WorstCase < 1.0f ? 1.0f : WorstCase;
+		return 1.0f - (static_cast<float>(Algo::LevenshteinDistance(A, B)) / WorstCase);
+	};
+
+	// returns the maximum weight for a given clause within an array of candidates
+	auto LevenshteinSingleMaxInArray = [LevenshteinSingle](const FString& A, const TArrayView<const FString>& B) -> float
+	{
+		float Weight = 0.f;
+		for (const FString& ElementInB : B)
+		{
+			Weight = FMath::Max(Weight, LevenshteinSingle(A, ElementInB));
+		}
+		return Weight;
+	};
+
+	// returns the maximum weight for an array of clauses within an array of candidates
+	auto LevenshteinMaxInArray = [LevenshteinSingleMaxInArray](const TArrayView<const FString>& A, const TArrayView<const FString>& B) -> float
+	{
+		float Weight = 0.f;
+		for (const FString& ElementInA : A)
+		{
+			Weight = FMath::Max(Weight, LevenshteinSingleMaxInArray(ElementInA, B));
+		}
+		return Weight;
+	};
+
+	// returns the maximum weight for all clauses for a matching array of candidates
+	auto LevenshteinSameSizeArrays = [LevenshteinSingle](const TArrayView<const FString>& A, const TArrayView<const FString>& B) -> float
+	{
+		if (A.Num() != B.Num() || A.IsEmpty() || B.IsEmpty())
+		{
+			return 0.f;
+		}
+		float Weight = 0.f;
+		for (int32 Index = 0; Index < A.Num(); Index++)
+		{
+			Weight += LevenshteinSingle(A[Index], B[Index]);
+		}
+		return Weight / static_cast<float>(A.Num());
+	};
+
+	// returns the maximum weight of an array of clauses for any matching subarrays within a larger array of candidates
+	auto LevenshteinArbitraryArrays = [LevenshteinSameSizeArrays](const TArrayView<const FString>& A, const TArrayView<const FString>& B) -> float
+	{
+		if (A.IsEmpty() || B.IsEmpty() || A.Num() > B.Num())
+		{
+			return 0.f;
+		}
+		float Weight = 0.f;
+		for (int32 Offset = 0; Offset <= B.Num() - A.Num(); Offset++)
+		{
+			const TArrayView<const FString> PartOfB = B.Slice(Offset, A.Num());
+			Weight = FMath::Max(Weight, LevenshteinSameSizeArrays(A, PartOfB));
+		}
+		return Weight;
+	};
+
+	const FString JoinedSanitizedFilterTerms = RigVMStringUtils::JoinStrings(SanitizedFilterTerms, TEXT(""));
+	const FString JoinedMenuDescriptions = RigVMStringUtils::JoinStrings(MenuDescriptions, TEXT(""));
+
+	// build up the score with a set of priorities
+	float Score = 0.f;
+	
+	// exact match start has the highest weight
+	if (JoinedMenuDescriptions.StartsWith(JoinedSanitizedFilterTerms, ESearchCase::IgnoreCase))
+	{
+		Score += 2048.f;
+	}
+	// string contained somewhere second highest weight
+	if (JoinedMenuDescriptions.Contains(JoinedSanitizedFilterTerms, ESearchCase::IgnoreCase))
+	{
+		Score += 1024.f;
+	}
+	
+	// fuzzy match with the joined string
+	Score += 512.f * LevenshteinSingle(JoinedSanitizedFilterTerms, JoinedMenuDescriptions);
+	// fuzzy match within the menu descriptions
+	Score += 256.f * FMath::Max(LevenshteinArbitraryArrays(SanitizedFilterTerms, MenuDescriptions), LevenshteinArbitraryArrays(SanitizedFilterTerms, MenuDescriptionsReversed));
+	// fuzzy match within the keywords
+	Score += 64.f * LevenshteinMaxInArray(SanitizedFilterTerms, Keywords);
+	// additional weight based on the distance in the number of terms and the number of menu descriptions (get transform vs get transform array)
+	Score += 32.f * (1.f - (FMath::Max(0.f, static_cast<float>(MenuDescriptions.Num() - SanitizedFilterTerms.Num())) / static_cast<float>(FMath::Max3(MenuDescriptions.Num(), SanitizedFilterTerms.Num(), 1))));
+
+	return Score;
+}
+
 bool URigVMEdGraphSchema::SupportsPinType(const UScriptStruct* ScriptStruct) const
 {
 	if(!ScriptStruct)
diff --git a/Engine/Plugins/Runtime/RigVM/Source/RigVMDeveloper/Public/EdGraph/RigVMEdGraphSchema.h b/Engine/Plugins/Runtime/RigVM/Source/RigVMDeveloper/Public/EdGraph/RigVMEdGraphSchema.h
--- a/Engine/Plugins/Runtime/RigVM/Source/RigVMDeveloper/Public/EdGraph/RigVMEdGraphSchema.h
+++ b/Engine/Plugins/Runtime/RigVM/Source/RigVMDeveloper/Public/EdGraph/RigVMEdGraphSchema.h
@@ -365,6 +365,8 @@ public:
 	UE_API virtual TSharedPtr<INameValidatorInterface> GetNameValidator(const FRigVMEditorAssetInterfacePtr Asset, const FName& OriginalName, const UStruct* ValidationScope) const;
 	UE_API bool CanPromotePinToVariable(const UEdGraphPin& Pin, bool bCond) const;
 	UE_API virtual bool SupportsPreviewFromHereOnNode(const UEdGraphNode* InNode) const;
+	virtual bool AllowsFuzzyActionFiltering(const FEdGraphSchemaAction& InCurrentAction) const override { return true; }
+	UE_API virtual float GetActionFilteredWeight(const FEdGraphSchemaAction& InCurrentAction, const TArray<FString>& InFilterTerms, const TArray<FString>& InSanitizedFilterTerms, const TArray<UEdGraphPin*>& DraggedFromPins) const override;
 
 	/** Returns true if the schema supports the script type */
 	UE_API bool SupportsPinType(const UScriptStruct* ScriptStruct) const;
diff --git a/Engine/Source/Editor/GraphEditor/Private/SGraphActionMenu.cpp b/Engine/Source/Editor/GraphEditor/Private/SGraphActionMenu.cpp
--- a/Engine/Source/Editor/GraphEditor/Private/SGraphActionMenu.cpp
+++ b/Engine/Source/Editor/GraphEditor/Private/SGraphActionMenu.cpp
@@ -1320,16 +1320,20 @@ void SGraphActionMenu::ScoreAndAddActions(int32 StartingIndex)
 		bool bShowAction = true;
 		float EachWeight = TNumericLimits<float>::Lowest();
 
-		const FString& SearchText = CurrentAction->GetFullSearchText();
-		for (int32 FilterIndex = 0; (FilterIndex < FilterTerms.Num()) && bShowAction; ++FilterIndex)
+		// only filter out per strict rule if the schema doesn't rely on fuzzy filter weights
+		if (!ActionSchema->AllowsFuzzyActionFiltering(*CurrentAction))
 		{
-			const bool bMatchesTerm = (SearchText.Contains(FilterTerms[FilterIndex], ESearchCase::CaseSensitive) || (SearchText.Contains(SanitizedFilterTerms[FilterIndex], ESearchCase::CaseSensitive) == true));
-			bShowAction = bMatchesTerm;
-		}
+			const FString& SearchText = CurrentAction->GetFullSearchText();
+			for (int32 FilterIndex = 0; (FilterIndex < FilterTerms.Num()) && bShowAction; ++FilterIndex)
+			{
+				const bool bMatchesTerm = (SearchText.Contains(FilterTerms[FilterIndex], ESearchCase::CaseSensitive) || (SearchText.Contains(SanitizedFilterTerms[FilterIndex], ESearchCase::CaseSensitive) == true));
+				bShowAction = bMatchesTerm;
+			}
 
-		if (!bShowAction)
-		{
-			continue;
+			if (!bShowAction)
+			{
+				continue;
+			}
 		}
 
 		bool bAddedAndSelected = false;
diff --git a/Engine/Source/Runtime/Engine/Classes/EdGraph/EdGraphSchema.h b/Engine/Source/Runtime/Engine/Classes/EdGraph/EdGraphSchema.h
--- a/Engine/Source/Runtime/Engine/Classes/EdGraph/EdGraphSchema.h
+++ b/Engine/Source/Runtime/Engine/Classes/EdGraph/EdGraphSchema.h
@@ -976,6 +976,13 @@ class UEdGraphSchema : public UObject
 	/** Get the name to show in the editor */
 	ENGINE_API virtual FText GetPinDisplayName(const UEdGraphPin* Pin) const;
 
+	/**
+	 * Returns true if this schema is using fuzzy string matching for searches 
+	 * 
+	 * @param InCurrentAction				The current graph action to calculate the weight of.
+	 */
+	virtual bool AllowsFuzzyActionFiltering(const FEdGraphSchemaAction& InCurrentAction) const { return false; }
+
 	/**
 	 * Calculate the weight priority of a given action for the context menu. 
 	 * 
