UAF/RigVM: More compilation optimizations

Add flag to suspend AST link checks and use it during programatic graph generation. Halves the compilation time of assets of moderate complexity (e.g. UAF_SYS_Main)
Remove redundant FName -> FString conversions when finding valid node names
Remove redundant nested registry locks in GetInitialDefaultValueForPin

#Synced-CL 48758736

diff --git a/Engine/Plugins/Experimental/UAF/UAF/Source/UAFUncookedOnly/Private/AnimNextRigVMAssetEditorData.cpp b/Engine/Plugins/Experimental/UAF/UAF/Source/UAFUncookedOnly/Private/AnimNextRigVMAssetEditorData.cpp
--- a/Engine/Plugins/Experimental/UAF/UAF/Source/UAFUncookedOnly/Private/AnimNextRigVMAssetEditorData.cpp
+++ b/Engine/Plugins/Experimental/UAF/UAF/Source/UAFUncookedOnly/Private/AnimNextRigVMAssetEditorData.cpp
@@ -891,6 +891,7 @@ void UAnimNextRigVMAssetEditorData::BuildFunctionWrapperEvents(FAnimNextRigVMAss
 	URigVMGraph* WrapperGraph = NewObject<URigVMGraph>(this, NAME_None, RF_Transient);
 	URigVMController* Controller = VMClient->GetOrCreateController(WrapperGraph);
 	FRigVMControllerNotifGuard NotifGuard(Controller);
+	FRigVMControllerASTLinkCheckGuard LinkCheckGuard(Controller);
 	bool bAddedWrapperEvent = true;
 	TArray<FRigVMExternalVariable> ExternalVariables = Asset->GetExternalVariables();
 
diff --git a/Engine/Plugins/Experimental/UAF/UAF/Source/UAFUncookedOnly/Private/UncookedOnlyUtils.cpp b/Engine/Plugins/Experimental/UAF/UAF/Source/UAFUncookedOnly/Private/UncookedOnlyUtils.cpp
--- a/Engine/Plugins/Experimental/UAF/UAF/Source/UAFUncookedOnly/Private/UncookedOnlyUtils.cpp
+++ b/Engine/Plugins/Experimental/UAF/UAF/Source/UAFUncookedOnly/Private/UncookedOnlyUtils.cpp
@@ -406,6 +406,9 @@ void FUtils::CompileVariableBindingsInternal(const FRigVMCompileSettings& InSett
 
 	FRigVMClient* VMClient = EditorData->GetRigVMClient();
 	URigVMController* Controller = VMClient->GetOrCreateController(BindingGraph);
+
+	FRigVMControllerASTLinkCheckGuard LinkCheckGuard(Controller);
+
 	UScriptStruct* BindingsNodeType = bInThreadSafe ? FRigUnit_AnimNextExecuteBindings_WT::StaticStruct() : FRigUnit_AnimNextExecuteBindings_GT::StaticStruct();
 	URigVMNode* ExecuteBindingsNode = Controller->AddUnitNode(BindingsNodeType, FRigVMStruct::ExecuteName, FVector2D::ZeroVector, FString(), false);
 	if(ExecuteBindingsNode == nullptr)
diff --git a/Engine/Plugins/Runtime/RigVM/Source/RigVMDeveloper/Private/RigVMModel/Nodes/RigVMTemplateNode.cpp b/Engine/Plugins/Runtime/RigVM/Source/RigVMDeveloper/Private/RigVMModel/Nodes/RigVMTemplateNode.cpp
--- a/Engine/Plugins/Runtime/RigVM/Source/RigVMDeveloper/Private/RigVMModel/Nodes/RigVMTemplateNode.cpp
+++ b/Engine/Plugins/Runtime/RigVM/Source/RigVMDeveloper/Private/RigVMModel/Nodes/RigVMTemplateNode.cpp
@@ -556,13 +556,16 @@ FString URigVMTemplateNode::GetInitialDefaultValueForPin(const FName& InRootPinN
 	{
 		return FString();
 	}
-	
+
 	TArray<int32> PermutationIndices = InPermutationIndices;
 	if(PermutationIndices.IsEmpty())
 	{
 		PermutationIndices = GetResolvedPermutationIndices(false);
 	}
 
+	const FRigVMRegistry& Registry = FRigVMRegistry::Get();
+	FRigVMRegistryReadLock RegistryReadLock(Registry);
+	
 	FString DefaultValue;
 	const FRigVMDispatchFactory* Factory = Template->GetDispatchFactory(); 
 	const FRigVMTemplateArgument* Argument = Template->FindArgument(InRootPinName);
@@ -573,7 +576,7 @@ FString URigVMTemplateNode::GetInitialDefaultValueForPin(const FName& InRootPinN
 		{
 			FString NewDefaultValue;
 
-			const TRigVMTypeIndex TypeIndex = Argument->GetTypeIndex(PermutationIndex);
+			const TRigVMTypeIndex TypeIndex = Argument->GetTypeIndex_NoLock(PermutationIndex, RegistryReadLock);
 
 			// INDEX_NONE indicates deleted permutation
 			if (TypeIndex == INDEX_NONE)
@@ -585,7 +588,7 @@ FString URigVMTemplateNode::GetInitialDefaultValueForPin(const FName& InRootPinN
 			{
 				NewDefaultValue = Factory->GetArgumentDefaultValue(Argument->GetName(), TypeIndex);
 			}
-			else if(const FRigVMFunction* Permutation = Template->GetPermutation(PermutationIndex))
+			else if(const FRigVMFunction* Permutation = Template->GetPermutation_NoLock(PermutationIndex, RegistryReadLock))
 			{
 				const TSharedPtr<FStructOnScope> StructOnScope = MakeShareable(new FStructOnScope(Permutation->Struct));
 				const FRigVMStruct* DefaultStruct = (const FRigVMStruct*)StructOnScope->GetStructMemory();
@@ -596,13 +599,13 @@ FString URigVMTemplateNode::GetInitialDefaultValueForPin(const FName& InRootPinN
 			}
 			else
 			{
-				if (FRigVMRegistry::Get().IsArrayType(TypeIndex))
+				if (RegistryReadLock->IsArrayType_NoLock(TypeIndex))
 				{
 					NewDefaultValue = TEXT("()");
 				}
 				else
 				{
-					const FRigVMTemplateArgumentType& Type = FRigVMRegistry::Get().GetType(TypeIndex);
+					const FRigVMTemplateArgumentType& Type = RegistryReadLock->GetType_NoLock(TypeIndex);
 					if (UScriptStruct* ScriptStruct = Cast<UScriptStruct>(Type.CPPTypeObject))
 					{
 						TArray<uint8, TAlignedHeapAllocator<16>> TempBuffer;
diff --git a/Engine/Plugins/Runtime/RigVM/Source/RigVMDeveloper/Private/RigVMModel/RigVMController.cpp b/Engine/Plugins/Runtime/RigVM/Source/RigVMDeveloper/Private/RigVMModel/RigVMController.cpp
--- a/Engine/Plugins/Runtime/RigVM/Source/RigVMDeveloper/Private/RigVMModel/RigVMController.cpp
+++ b/Engine/Plugins/Runtime/RigVM/Source/RigVMDeveloper/Private/RigVMModel/RigVMController.cpp
@@ -101,6 +101,7 @@ URigVMController::URigVMController()
 	, UserLinkDirection(ERigVMPinDirection::Invalid)
 	, bEnableTypeCasting(true)
 	, bAllowPrivateFunctions(false)
+	, bSuspendASTLinkChecks(false)
 	, bIsTransacting(false)
 	, bIsRunningUnitTest(false)
 	, bIsFullyResolvingTemplateNode(false)
@@ -124,6 +125,7 @@ URigVMController::URigVMController(const FObjectInitializer& ObjectInitializer)
 	, UserLinkDirection(ERigVMPinDirection::Invalid)
 	, bEnableTypeCasting(true)
 	, bAllowPrivateFunctions(false)
+	, bSuspendASTLinkChecks(false)
 	, bIsTransacting(false)
 	, bIsRunningUnitTest(false)
 	, bIsFullyResolvingTemplateNode(false)
@@ -12813,7 +12815,7 @@ bool URigVMController::AddLink(URigVMPin* OutputPin, URigVMPin* InputPin, bool b
 		FString FailureReason;
 
 		bool bCanLink = true;
-		if(bIsRestoringLinks)
+		if(bIsRestoringLinks || bSuspendASTLinkChecks)
 		{
 			bCanLink = URigVMPin::CanLink(OutputPin, InputPin, &FailureReason, nullptr, UserLinkDirection, bCreateCastNode);
 		}
@@ -21149,35 +21151,39 @@ bool URigVMController::UpdateTemplateNodePinTypes(URigVMTemplateNode* InNode, bo
 	// Find the types for each possible permutation
 	TMap<int32, TArray<TRigVMTypeIndex>> PinTypes; // permutation to pin types
 	PinTypes.Reserve(ResolvedPermutations.Num() * NumPins);
-	for (int32 ResolvedPermutation : ResolvedPermutations)
+
 	{
-		for(int32 PinIndex=0; PinIndex < NumPins; ++PinIndex)
+		FRigVMRegistryReadLock RegistryReadLock(Registry);
+		for (int32 ResolvedPermutation : ResolvedPermutations)
 		{
-			URigVMPin* Pin = Pins[PinIndex];
-			if (Pin->GetDirection() == ERigVMPinDirection::Hidden)
+			for(int32 PinIndex=0; PinIndex < NumPins; ++PinIndex)
 			{
-				PinTypes.FindOrAdd(ResolvedPermutation).Add(INDEX_NONE);
-				continue;
-			}
+				URigVMPin* Pin = Pins[PinIndex];
+				if (Pin->GetDirection() == ERigVMPinDirection::Hidden)
+				{
+					PinTypes.FindOrAdd(ResolvedPermutation).Add(INDEX_NONE);
+					continue;
+				}
 
-			bool bAddedType = false;
-			TArray<TRigVMTypeIndex>& Types = PinTypes.FindOrAdd(ResolvedPermutation);
-			if (const FRigVMTemplateArgument* Argument = Arguments[PinIndex])
-			{
-				Types.Add(Argument->GetTypeIndex(ResolvedPermutation));
-				bAddedType = true;
-			}
-			else if (const FRigVMExecuteArgument* ExecuteArgument = Template->FindExecuteArgument(Pin->GetFName(), DispatchContext))
-			{
-				Types.Add(ExecuteArgument->TypeIndex);
-				bAddedType = true;
-			}
-			
-			if (!bAddedType)
-			{
-				// This is a pin with no argument
-				// Its marked as invalid, and will show wildcard type
-				Types.Add(INDEX_NONE);
+				bool bAddedType = false;
+				TArray<TRigVMTypeIndex>& Types = PinTypes.FindOrAdd(ResolvedPermutation);
+				if (const FRigVMTemplateArgument* Argument = Arguments[PinIndex])
+				{
+					Types.Add(Argument->GetTypeIndex_NoLock(ResolvedPermutation, RegistryReadLock));
+					bAddedType = true;
+				}
+				else if (const FRigVMExecuteArgument* ExecuteArgument = Template->FindExecuteArgument_NoLock(Pin->GetFName(), DispatchContext, RegistryReadLock))
+				{
+					Types.Add(ExecuteArgument->TypeIndex);
+					bAddedType = true;
+				}
+				
+				if (!bAddedType)
+				{
+					// This is a pin with no argument
+					// Its marked as invalid, and will show wildcard type
+					Types.Add(INDEX_NONE);
+				}
 			}
 		}
 	}
diff --git a/Engine/Plugins/Runtime/RigVM/Source/RigVMDeveloper/Private/RigVMModel/RigVMGraph.cpp b/Engine/Plugins/Runtime/RigVM/Source/RigVMDeveloper/Private/RigVMModel/RigVMGraph.cpp
--- a/Engine/Plugins/Runtime/RigVM/Source/RigVMDeveloper/Private/RigVMModel/RigVMGraph.cpp
+++ b/Engine/Plugins/Runtime/RigVM/Source/RigVMDeveloper/Private/RigVMModel/RigVMGraph.cpp
@@ -651,6 +651,18 @@ bool URigVMGraph::IsNameAvailable(const FString& InName) const
 	return true;
 }
 
+bool URigVMGraph::IsNameAvailable(FName InName) const
+{
+	for (const URigVMNode* Node : Nodes)
+	{
+		if (Node->GetFName() == InName)
+		{
+			return false;
+		}
+	}
+	return true;
+}
+
 TSharedPtr<FScribbleGraphData> URigVMGraph::SetupScribbleGraph(bool bCreateIfMissing)
 {
 	if (!ScribbleGraph.IsValid())
diff --git a/Engine/Plugins/Runtime/RigVM/Source/RigVMDeveloper/Private/RigVMModel/RigVMSchema.cpp b/Engine/Plugins/Runtime/RigVM/Source/RigVMDeveloper/Private/RigVMModel/RigVMSchema.cpp
--- a/Engine/Plugins/Runtime/RigVM/Source/RigVMDeveloper/Private/RigVMModel/RigVMSchema.cpp
+++ b/Engine/Plugins/Runtime/RigVM/Source/RigVMDeveloper/Private/RigVMModel/RigVMSchema.cpp
@@ -916,7 +916,7 @@ FString URigVMSchema::GetValidNodeName(const URigVMGraph* InGraph, const FString
 	check(InGraph);
 
 	return GetUniqueName(*InPrefix, [&](const FName& InName) {
-		return InGraph->IsNameAvailable(InName.ToString());
+		return InGraph->IsNameAvailable(InName);
 	}, false, true).ToString();
 }
 
diff --git a/Engine/Plugins/Runtime/RigVM/Source/RigVMDeveloper/Public/RigVMModel/RigVMController.h b/Engine/Plugins/Runtime/RigVM/Source/RigVMDeveloper/Public/RigVMModel/RigVMController.h
--- a/Engine/Plugins/Runtime/RigVM/Source/RigVMDeveloper/Public/RigVMModel/RigVMController.h
+++ b/Engine/Plugins/Runtime/RigVM/Source/RigVMDeveloper/Public/RigVMModel/RigVMController.h
@@ -1651,6 +1651,7 @@ private:
 	ERigVMPinDirection UserLinkDirection;
 	bool bEnableTypeCasting;
 	bool bAllowPrivateFunctions;
+	bool bSuspendASTLinkChecks;
 	TOptional<ERigVMPinDefaultValueType> OptionalDefaultValueType;
 
 	// temporary maps used for pin redirection
@@ -1731,6 +1732,7 @@ private:
 	friend struct FRigVMPinInfoArray;
 	friend class FRigVMControllerNotifGuard;
 	friend class FRigVMDefaultValueTypeGuard;
+	friend class FRigVMControllerASTLinkCheckGuard;
 	friend struct FRigVMClient;
 	friend struct FRigVMActionWrapper;
 	friend class URigVMSchema;
@@ -1785,6 +1787,29 @@ private:
 	TOptional<ERigVMPinDefaultValueType> bPreviousDefaultValueType;
 };
 
+class FRigVMControllerASTLinkCheckGuard
+{
+public:
+
+	explicit FRigVMControllerASTLinkCheckGuard(URigVMController* InController, bool bInSuspendASTLinkChecks = true)
+		: Controller(InController)
+	{
+		bPreviousSuspendASTLinkChecks = Controller->bSuspendASTLinkChecks;
+		Controller->bSuspendASTLinkChecks = bInSuspendASTLinkChecks;
+	}
+
+	~FRigVMControllerASTLinkCheckGuard()
+	{
+		Controller->bSuspendASTLinkChecks = bPreviousSuspendASTLinkChecks;
+	}
+
+private:
+
+	URigVMController* Controller;
+	bool bPreviousSuspendASTLinkChecks;
+};
+
+
 USTRUCT()
 struct FRigVMController_CommonTypePerTemplate
 {
diff --git a/Engine/Plugins/Runtime/RigVM/Source/RigVMDeveloper/Public/RigVMModel/RigVMGraph.h b/Engine/Plugins/Runtime/RigVM/Source/RigVMDeveloper/Public/RigVMModel/RigVMGraph.h
--- a/Engine/Plugins/Runtime/RigVM/Source/RigVMDeveloper/Public/RigVMModel/RigVMGraph.h
+++ b/Engine/Plugins/Runtime/RigVM/Source/RigVMDeveloper/Public/RigVMModel/RigVMGraph.h
@@ -249,6 +249,7 @@ private:
 	TArray<FRigVMGraphSection> SectionsMatchingTheSelection;
 
 	UE_API bool IsNameAvailable(const FString& InName) const;
+	UE_API bool IsNameAvailable(FName InName) const;
 
 	UPROPERTY()
 	FScribbleGraph ScribbleGraph;
