RigVM: Node renderer to generate images for each node for the docs

#rnx

#Synced-CL 48802464

diff --git a/Engine/Plugins/Runtime/RigVM/Source/RigVM/Public/RigVMRuntimeAsset.h b/Engine/Plugins/Runtime/RigVM/Source/RigVM/Public/RigVMRuntimeAsset.h
--- a/Engine/Plugins/Runtime/RigVM/Source/RigVM/Public/RigVMRuntimeAsset.h
+++ b/Engine/Plugins/Runtime/RigVM/Source/RigVM/Public/RigVMRuntimeAsset.h
@@ -115,7 +115,7 @@ public:
 
 };
 
-UCLASS(MinimalAPI, Abstract)
+UCLASS(MinimalAPI)
 class URigVMRuntimeAsset : public UObject, public IRigVMRuntimeAssetInterface, public IRigVMGraphFunctionHost
 {
 	GENERATED_BODY()
diff --git a/Engine/Plugins/Runtime/RigVM/Source/RigVMDeveloper/Private/RigVMModel/RigVMSchema.cpp b/Engine/Plugins/Runtime/RigVM/Source/RigVMDeveloper/Private/RigVMModel/RigVMSchema.cpp
--- a/Engine/Plugins/Runtime/RigVM/Source/RigVMDeveloper/Private/RigVMModel/RigVMSchema.cpp
+++ b/Engine/Plugins/Runtime/RigVM/Source/RigVMDeveloper/Private/RigVMModel/RigVMSchema.cpp
@@ -31,6 +31,16 @@ UClass* URigVMSchema::GetEdGraphSchemaClass() const
 	return URigVMEdGraphSchema::StaticClass();
 }
 
+UClass* URigVMSchema::GetRuntimeAssetClass() const
+{
+	return URigVMRuntimeAsset::StaticClass();
+}
+
+UClass* URigVMSchema::GetEditorAssetClass() const
+{
+	return URigVMEditorAsset::StaticClass();
+}
+
 bool URigVMSchema::SupportsType(URigVMController* InController, TRigVMTypeIndex InTypeIndex) const
 {
 	check(Registry);
diff --git a/Engine/Plugins/Runtime/RigVM/Source/RigVMDeveloper/Public/RigVMModel/RigVMSchema.h b/Engine/Plugins/Runtime/RigVM/Source/RigVMDeveloper/Public/RigVMModel/RigVMSchema.h
--- a/Engine/Plugins/Runtime/RigVM/Source/RigVMDeveloper/Public/RigVMModel/RigVMSchema.h
+++ b/Engine/Plugins/Runtime/RigVM/Source/RigVMDeveloper/Public/RigVMModel/RigVMSchema.h
@@ -46,6 +46,12 @@ public:
 	// Returns the corresponding Ed Graph Schema class
 	UE_API virtual UClass* GetEdGraphSchemaClass() const;
 
+	// Returns the corresponding runtime asset class to use for this schema
+	UE_API virtual UClass* GetRuntimeAssetClass() const;
+
+	// Returns the corresponding editor asset class to use for this schema
+	UE_API virtual UClass* GetEditorAssetClass() const;
+
 	// Returns true if a graph supports a given type
 	UE_API bool SupportsType(URigVMController* InController, TRigVMTypeIndex InTypeIndex) const;
 	UE_API virtual bool SupportsType_NoLock(URigVMController* InController, TRigVMTypeIndex InTypeIndex, FRigVMRegistryHandle& InRegistry) const;
diff --git a/Engine/Plugins/Runtime/RigVM/Source/RigVMEditor/Private/Editor/RigVMMinimalEnvironment.cpp b/Engine/Plugins/Runtime/RigVM/Source/RigVMEditor/Private/Editor/RigVMMinimalEnvironment.cpp
--- a/Engine/Plugins/Runtime/RigVM/Source/RigVMEditor/Private/Editor/RigVMMinimalEnvironment.cpp
+++ b/Engine/Plugins/Runtime/RigVM/Source/RigVMEditor/Private/Editor/RigVMMinimalEnvironment.cpp
@@ -68,6 +68,9 @@ void FRigVMMinimalEnvironment::SetNode(URigVMNode* InModelNode)
 	}
 
 	HandleModified(ERigVMGraphNotifType::NodeAdded, GetModel(), GetNode());
+
+	// process the queued up events
+	Tick_GameThead(0.f);
 }
 
 void FRigVMMinimalEnvironment::SetFunctionNode(const FRigVMGraphFunctionIdentifier& InIdentifier)
diff --git a/Engine/Plugins/Runtime/RigVM/Source/RigVMEditor/Private/RigVMEditorBlueprintLibrary.cpp b/Engine/Plugins/Runtime/RigVM/Source/RigVMEditor/Private/RigVMEditorBlueprintLibrary.cpp
--- a/Engine/Plugins/Runtime/RigVM/Source/RigVMEditor/Private/RigVMEditorBlueprintLibrary.cpp
+++ b/Engine/Plugins/Runtime/RigVM/Source/RigVMEditor/Private/RigVMEditorBlueprintLibrary.cpp
@@ -5,7 +5,15 @@
 #include "AssetRegistry/AssetRegistryModule.h"
 #include "Misc/ScopedSlowTask.h"
 #include "UObject/UObjectIterator.h"
-
+#include "ImageUtils.h"
+#include "TextureResource.h"
+#include "Engine/TextureRenderTarget2D.h"
+#include "Slate/WidgetRenderer.h"
+#include "RenderingThread.h"
+#include "Misc/FileHelper.h"
+#include "HAL/IConsoleManager.h"
+#include "Editor/RigVMMinimalEnvironment.h"
+#include "Widgets/SRigVMGraphNode.h"
 
 #include UE_INLINE_GENERATED_CPP_BY_NAME(RigVMEditorBlueprintLibrary)
 
@@ -361,4 +369,282 @@ TArray<FAssetData> URigVMEditorBlueprintLibrary::GetAssetsWithFilter(
 	return AssetDataList;
 }
 
+FAutoConsoleCommand FCmdRigVMJsonPrintGraph
+(
+	TEXT("RigVM.JSON.RenderNodeToPNG"),
+	TEXT("Renders a node to a PNG image."),
+	FConsoleCommandWithArgsDelegate::CreateLambda([](const TArray<FString>& Args)
+	{
+		auto PrintHelp = []()
+		{
+			UE_LOG(LogRigVMDeveloper, Warning, TEXT("Valid Arguments:\n"
+				"subject=/Script/RigVM.RigVMFunction_MathBoolOr - the path name of the subject (struct or graph function)\n"
+				"filepath=c:/temp/NodePreview.png\n"
+			));
+			return;					
+		};
+		if (Args.Num() != 2)
+		{
+			PrintHelp();
+			return;
+		}
+		UObject* Subject = nullptr;
+		FRigVMGraphFunctionIdentifier FunctionIdentifier;
+		FString FilePath;
+
+		for (const FString& Arg : Args)
+		{
+			if (Arg.StartsWith(TEXT("subject="), ESearchCase::IgnoreCase))
+			{
+				const FString SubjectPathName = Arg.Mid(8);
+				Subject = FindObject<UObject>(nullptr, *SubjectPathName);
+				if (!Subject)
+				{
+					UE_LOG(LogRigVMDeveloper, Error, TEXT("Invalid Subject '%s'"), *SubjectPathName);
+					return;
+				}
+			}
+			else if (Arg.StartsWith(TEXT("filepath="), ESearchCase::IgnoreCase))
+			{
+				FilePath = Arg.Mid(9);
+			}
+			else
+			{
+				PrintHelp();
+				return;
+			}
+		}
+
+		if (Subject)
+		{
+			if (URigVMEditorBlueprintLibrary::RenderRigVMSubjectToPNG(Subject, FilePath))
+			{
+				UE_LOG(LogRigVMDeveloper, Log, TEXT("Successfully rendered image of '%s' to file '%s'."), *Subject->GetPathName(), *FilePath);
+			}
+			else
+			{
+				UE_LOG(LogRigVMDeveloper, Warning, TEXT("Error when rendering node of '%s' to file '%s'."), *Subject->GetPathName(), *FilePath);
+			}
+		}
+		else
+		{
+			PrintHelp();
+		}
+	})
+);
+
+bool URigVMEditorBlueprintLibrary::RenderRigVMSubjectToPNG(UObject* InSubject, const FString& OutFileName)
+{
+	URigVMEditorAsset* EditorAsset = Cast<URigVMEditorAsset>(MakeRigVMAssetForSubject(InSubject));
+	if (!EditorAsset)
+	{
+		return false;
+	}
+
+	URigVMGraph* Model = EditorAsset->GetRigVMClient()->GetDefaultModel();
+	if (!Model)
+	{
+		return false;
+	}
+
+	URigVMController* Controller = EditorAsset->GetRigVMClient()->GetOrCreateController(Model);
+	if (!Controller)
+	{
+		return false;
+	}
+
+	URigVMNode* ModelNode = nullptr;
+	if (UScriptStruct* ScriptStruct = Cast<UScriptStruct>(InSubject))
+	{
+		if (ScriptStruct->IsChildOf(FRigVMDispatchFactory::StaticStruct()))
+		{
+			const FName FactoryName = FRigVMDispatchFactory::GetFactoryName(ScriptStruct);
+			FRigVMRegistryReadLock Registry;
+			if (const FRigVMDispatchFactory* Factory = Registry->FindDispatchFactory_NoLock(FactoryName))
+			{
+				ModelNode = Controller->AddTemplateNode(Factory->GetTemplateNotation_NoLock(Registry));
+			}
+		}
+		else if (ScriptStruct->IsChildOf(FRigVMStruct::StaticStruct()))
+		{
+			ModelNode = Controller->AddUnitNode(const_cast<UScriptStruct*>(ScriptStruct));
+		}
+	}
+	else if (URigVMLibraryNode* LibraryNode = Cast<URigVMLibraryNode>(InSubject))
+	{
+		ModelNode = Controller->AddFunctionReferenceNode(LibraryNode);
+	}
+
+	return RenderModelNodeToPNG(EditorAsset, ModelNode, OutFileName);
+}
+
+bool URigVMEditorBlueprintLibrary::RenderModelNodeToPNG(URigVMEditorAsset* EditorAsset, URigVMNode* InModelNode, const FString& OutFileName)
+{
+	if (!InModelNode)
+	{
+		return false;
+	}
+
+	//create directory if it doesn't exist
+	if (!IFileManager::Get().MakeDirectory(*FPaths::GetPath(OutFileName), true))
+	{
+		UE_LOG(LogRigVMDeveloper, Warning, TEXT("Cannot create directory for file '%s'."), *OutFileName);
+		return false;
+	}
+
+	// we only allow png extensions
+	if (!OutFileName.EndsWith(TEXT(".png"), ESearchCase::IgnoreCase))
+	{
+		UE_LOG(LogRigVMDeveloper, Warning, TEXT("Image file '%s' is not a PNG."), *OutFileName);
+		return false;
+	}
+
+	TSharedPtr<FRigVMMinimalEnvironment> MinimalEnvironment = MakeShared<FRigVMMinimalEnvironment>(EditorAsset);
+	MinimalEnvironment->SetNode(InModelNode);
+
+	URigVMEdGraphNode* EdGraphNode = MinimalEnvironment->GetEdGraphNode();
+	if (!EdGraphNode)
+	{
+		return false;
+	}
+
+	TSharedRef<SRigVMGraphNode> GraphNode =
+		SNew(SRigVMGraphNode)
+		.Enabled(true)
+		.GraphNodeObj(EdGraphNode);
+	
+	GraphNode->SlatePrepass();
+
+	const FVector2D DesiredSize = GraphNode->GetDesiredSize();
+	const uint32 SizeX = static_cast<uint32>(DesiredSize.X);
+	const uint32 SizeY = static_cast<uint32>(DesiredSize.Y);
+	if (SizeX <= 0 || SizeY <= 0)
+	{
+		return false;
+	}
+
+	// create a render target to render to
+	UTextureRenderTarget2D* RenderTarget = NewObject<UTextureRenderTarget2D>(EditorAsset);
+	RenderTarget->ClearColor = FLinearColor::Transparent;
+	RenderTarget->InitCustomFormat(SizeX, SizeY, EPixelFormat::PF_B8G8R8A8, false);
+
+	// render the widget into the render target
+	FWidgetRenderer* WidgetRenderer = new FWidgetRenderer(true, false);
+	check(WidgetRenderer);
+	WidgetRenderer->DrawWidget(RenderTarget, GraphNode, GraphNode->GetCachedGeometry().Scale, DesiredSize, 0.f);
+	FlushRenderingCommands();
+
+	// read the pixels of the render target into a color array
+	TArray<FColor> PixelColorData;
+	ENQUEUE_RENDER_COMMAND(ReadScreenshotRTCmd)(
+	[RenderTarget, &PixelColorData, SizeX, SizeY](FRHICommandListImmediate& RHICmdList)
+		{
+			FTextureRenderTarget2DResource* RTResource =
+				static_cast<FTextureRenderTarget2DResource*>(RenderTarget->GetRenderTargetResource());
+			RHICmdList.ReadSurfaceData(
+				RTResource->GetTextureRHI(),
+				FIntRect(0, 0, static_cast<int32>(SizeX), static_cast<int32>(SizeY)),
+				PixelColorData,
+				FReadSurfaceDataFlags());
+		}
+	);
+	FlushRenderingCommands();
+	BeginCleanup(WidgetRenderer);
+
+	// ReadPixels can return fewer pixels than expected
+	if (PixelColorData.Num() != SizeX * SizeY)
+	{
+		return false;
+	}
+	
+	// save the color array to png
+	TArray64<uint8> CompressedBitmap;
+	FImageUtils::PNGCompressImageArray(RenderTarget->SizeX, RenderTarget->SizeY, TArrayView64<const FColor>(PixelColorData.GetData(), PixelColorData.Num()), CompressedBitmap);
+	return FFileHelper::SaveArrayToFile(CompressedBitmap, *OutFileName);
+}
+
+UObject* URigVMEditorBlueprintLibrary::MakeRigVMAssetForSubject(UObject* InUnitOrDispatchStructOrLibraryNode)
+{
+	if (!InUnitOrDispatchStructOrLibraryNode)
+	{
+		return  nullptr;
+	}
+
+	const URigVMSchema* SchemaToUse = nullptr;
+
+	if (UScriptStruct* ScriptStruct = Cast<UScriptStruct>(InUnitOrDispatchStructOrLibraryNode))
+	{
+		// find the schema that allows this struct
+		for (TObjectIterator<UClass> ClassIterator; ClassIterator; ++ClassIterator)
+		{
+			if (!ClassIterator->IsChildOf(URigVMSchema::StaticClass()))
+			{
+				continue;
+			}
+
+			const URigVMSchema* Schema = ClassIterator->GetDefaultObject<URigVMSchema>();
+			if (!Schema)
+			{
+				continue;
+			}
+
+			if (ScriptStruct->IsChildOf(FRigVMDispatchFactory::StaticStruct()))
+			{
+				const FName FactoryName = FRigVMDispatchFactory::GetFactoryName(ScriptStruct);
+				FRigVMRegistryReadLock Registry;
+				if (const FRigVMDispatchFactory* Factory = Registry->FindDispatchFactory_NoLock(FactoryName))
+				{
+					if (Schema->SupportsDispatchFactory_NoLock(nullptr, Factory, Registry))
+					{
+						SchemaToUse = Schema;
+						break;
+					}
+				}
+			}
+			else if (ScriptStruct->IsChildOf(FRigVMStruct::StaticStruct()))
+			{
+				FRigVMRegistryReadLock Registry;
+				if (const FRigVMFunction* Function = Registry->FindFunction_NoLock(ScriptStruct, *FRigVMStruct::ExecuteName.ToString()))
+				{
+					if (Schema->SupportsUnitFunction_NoLock(nullptr, Function, Registry))
+					{
+						SchemaToUse = Schema;
+						break;
+					}
+				}
+			}		
+		}
+	}
+	else if (URigVMLibraryNode* LibraryNode = Cast<URigVMLibraryNode>(InUnitOrDispatchStructOrLibraryNode))
+	{
+		if (URigVMGraph* Graph = LibraryNode->GetGraph())
+		{
+			SchemaToUse = Graph->GetSchema();
+		}
+	}
+
+	if (!SchemaToUse)
+	{
+		return nullptr;
+	}
+
+	UClass* RuntimeAssetClass = SchemaToUse->GetRuntimeAssetClass();
+	if (!RuntimeAssetClass)
+	{
+		return nullptr;
+	}
+
+	UClass* EditorAssetClass = SchemaToUse->GetEditorAssetClass();
+	if (!EditorAssetClass)
+	{
+		return nullptr;
+	}
+
+	URigVMRuntimeAsset* RuntimeAsset = NewObject<URigVMRuntimeAsset>(GetTransientPackage(), RuntimeAssetClass);
+	RuntimeAsset->Initialize(EditorAssetClass);
+	URigVMEditorAsset* EditorAsset = CastChecked<URigVMEditorAsset>(RuntimeAsset->GetEditorOnlyData());
+	
+	return EditorAsset; 
+}
+
 #undef LOCTEXT_NAMESPACE
diff --git a/Engine/Plugins/Runtime/RigVM/Source/RigVMEditor/Public/RigVMEditorBlueprintLibrary.h b/Engine/Plugins/Runtime/RigVM/Source/RigVMEditor/Public/RigVMEditorBlueprintLibrary.h
--- a/Engine/Plugins/Runtime/RigVM/Source/RigVMEditor/Public/RigVMEditorBlueprintLibrary.h
+++ b/Engine/Plugins/Runtime/RigVM/Source/RigVMEditor/Public/RigVMEditorBlueprintLibrary.h
@@ -1,20 +1,19 @@
 // Copyright Epic Games, Inc. All Rights Reserved.
 
-/**
- *
- * This thumbnail renderer displays a given Control Rig
- */
-
 #pragma once
 
 #include "CoreMinimal.h"
 #include "RigVMBlueprintLegacy.h"
 #include "AssetRegistry/AssetData.h"
 #include "Kismet/BlueprintFunctionLibrary.h"
+#include "RigVMCore/RigVMGraphFunctionDefinition.h"
 #include "RigVMEditorBlueprintLibrary.generated.h"
 
 #define UE_API RIGVMEDITOR_API
 
+class URigVMNode;
+class URigVMEditorAsset;
+
 UENUM(BlueprintType)
 enum ERigVMBlueprintLoadLogSeverity : uint8
 {
@@ -133,6 +132,12 @@ public:
 	static UE_API TArray<FAssetData> GetAssetsWithFilter(
 		UClass* InClass,
 		FRigVMAssetDataFilter InAssetDataFilter);
+
+	UFUNCTION(BlueprintCallable, Category = "RigVM Blueprint")
+	static UE_API bool RenderRigVMSubjectToPNG(UObject* InUnitOrDispatchStructOrLibraryNode, const FString& OutFileName);
+
+	static UE_API bool RenderModelNodeToPNG(URigVMEditorAsset* EditorAsset, URigVMNode* InModelNode, const FString& OutFileName);
+	static UE_API UObject* MakeRigVMAssetForSubject(UObject* InUnitOrDispatchStructOrLibraryNode);
 };
 
 #undef UE_API
diff --git a/Engine/Plugins/Runtime/RigVM/Source/RigVMEditor/RigVMEditor.Build.cs b/Engine/Plugins/Runtime/RigVM/Source/RigVMEditor/RigVMEditor.Build.cs
--- a/Engine/Plugins/Runtime/RigVM/Source/RigVMEditor/RigVMEditor.Build.cs
+++ b/Engine/Plugins/Runtime/RigVM/Source/RigVMEditor/RigVMEditor.Build.cs
@@ -21,6 +21,8 @@ public class RigVMEditor : ModuleRules
                 "UnrealEd",
                 "Slate",
                 "SlateCore",
+                "RenderCore",
+                "RHI",
                 "UMG",
                 "InputCore",
                 "ToolWidgets",
