RigVM / Build fix: rename a CVar which had a duplicate name.

#rnx

#Synced-CL 48803475

diff --git a/Engine/Plugins/Runtime/RigVM/Source/RigVMEditor/Private/RigVMEditorBlueprintLibrary.cpp b/Engine/Plugins/Runtime/RigVM/Source/RigVMEditor/Private/RigVMEditorBlueprintLibrary.cpp
--- a/Engine/Plugins/Runtime/RigVM/Source/RigVMEditor/Private/RigVMEditorBlueprintLibrary.cpp
+++ b/Engine/Plugins/Runtime/RigVM/Source/RigVMEditor/Private/RigVMEditorBlueprintLibrary.cpp
@@ -369,7 +369,7 @@ TArray<FAssetData> URigVMEditorBlueprintLibrary::GetAssetsWithFilter(
 	return AssetDataList;
 }
 
-FAutoConsoleCommand FCmdRigVMJsonPrintGraph
+FAutoConsoleCommand FCmdRigVMJsonRenderNodeToPNG
 (
 	TEXT("RigVM.JSON.RenderNodeToPNG"),
 	TEXT("Renders a node to a PNG image."),
