RigVM: Node to access absolute time

#rnx
#Synced-CL 48802451

diff --git a/Engine/Plugins/Runtime/RigVM/Source/RigVM/Private/RigVMFunctions/Animation/RigVMFunction_GetDeltaTime.cpp b/Engine/Plugins/Runtime/RigVM/Source/RigVM/Private/RigVMFunctions/Animation/RigVMFunction_GetDeltaTime.cpp
--- a/Engine/Plugins/Runtime/RigVM/Source/RigVM/Private/RigVMFunctions/Animation/RigVMFunction_GetDeltaTime.cpp
+++ b/Engine/Plugins/Runtime/RigVM/Source/RigVM/Private/RigVMFunctions/Animation/RigVMFunction_GetDeltaTime.cpp
@@ -20,3 +20,19 @@ IMPLEMENT_RIGVMSTRUCT_AUTOMATION_TEST(FRigVMFunction_GetDeltaTime)
 	return true;
 }
 #endif
+
+FRigVMFunction_GetAbsoluteTime_Execute()
+{
+	Result = ExecuteContext.GetAbsoluteTime<float>();
+}
+
+#if WITH_DEV_AUTOMATION_TESTS
+
+IMPLEMENT_RIGVMSTRUCT_AUTOMATION_TEST(FRigVMFunction_GetAbsoluteTime)
+{
+	ExecuteContext.SetAbsoluteTime(7.0f);
+	Execute();
+	AddErrorIfFalse(FMath::IsNearlyEqual(Unit.Result, 7.0f), TEXT("unexpected absolute time"));
+	return true;
+}
+#endif
\ No newline at end of file
diff --git a/Engine/Plugins/Runtime/RigVM/Source/RigVM/Public/RigVMFunctions/Animation/RigVMFunction_GetDeltaTime.h b/Engine/Plugins/Runtime/RigVM/Source/RigVM/Public/RigVMFunctions/Animation/RigVMFunction_GetDeltaTime.h
--- a/Engine/Plugins/Runtime/RigVM/Source/RigVM/Public/RigVMFunctions/Animation/RigVMFunction_GetDeltaTime.h
+++ b/Engine/Plugins/Runtime/RigVM/Source/RigVM/Public/RigVMFunctions/Animation/RigVMFunction_GetDeltaTime.h
@@ -26,3 +26,23 @@ struct FRigVMFunction_GetDeltaTime : public FRigVMFunction_AnimBase
 	float Result;
 };
 
+/**
+ * Returns the time gone by since the start of the world / the time within sequencer 
+ */
+USTRUCT(meta=(DisplayName="Absolute Time", Varying))
+struct FRigVMFunction_GetAbsoluteTime : public FRigVMFunction_AnimBase
+{
+	GENERATED_BODY()
+	
+	FRigVMFunction_GetAbsoluteTime()
+	{
+		Result = 0.f;
+	}
+
+	RIGVM_METHOD()
+	RIGVM_API virtual void Execute() override;
+
+	// The current absolute time in seconds
+	UPROPERTY(meta=(Output))
+	float Result;
+};
