RigVM: Store versions along side Object Archive

#rnx

Full cook is here https://horde.devtools.epicgames.com/job/692db5656656b85302bb4baa

#Synced-CL 48802464

diff --git a/Engine/Plugins/Runtime/RigVM/Source/RigVM/Private/RigVMCore/RigVMGraphFunctionDefinition.cpp b/Engine/Plugins/Runtime/RigVM/Source/RigVM/Private/RigVMCore/RigVMGraphFunctionDefinition.cpp
--- a/Engine/Plugins/Runtime/RigVM/Source/RigVM/Private/RigVMCore/RigVMGraphFunctionDefinition.cpp
+++ b/Engine/Plugins/Runtime/RigVM/Source/RigVM/Private/RigVMCore/RigVMGraphFunctionDefinition.cpp
@@ -392,6 +392,10 @@ FArchive& operator<<(FArchive& Ar, FRigVMGraphFunctionData& Data)
 	}
 
 	Ar << Data.CollapseNodeArchive;
+	if (Ar.IsLoading())
+	{
+		Data.CollapseNodeArchive.SetVersionsFromArchive(Ar);
+	}
 	UE_RIGVM_ARCHIVETRACE_ENTRY(Ar, TEXT("CollapseNodeArchive"));
 
 #if UE_BUILD_SHIPPING
diff --git a/Engine/Plugins/Runtime/RigVM/Source/RigVM/Private/RigVMCore/RigVMObjectArchive.cpp b/Engine/Plugins/Runtime/RigVM/Source/RigVM/Private/RigVMCore/RigVMObjectArchive.cpp
--- a/Engine/Plugins/Runtime/RigVM/Source/RigVM/Private/RigVMCore/RigVMObjectArchive.cpp
+++ b/Engine/Plugins/Runtime/RigVM/Source/RigVM/Private/RigVMCore/RigVMObjectArchive.cpp
@@ -83,6 +83,19 @@ void FRigVMObjectArchive::Decompress()
 	}
 }
 
+void FRigVMObjectArchive::SetVersionsFromArchive(const FArchive& Ar)
+{
+	const FCustomVersionContainer& VersionContainer = Ar.GetCustomVersions();
+	const FCustomVersionArray& CustomVersions = VersionContainer.GetAllVersions();
+
+	Versions.Reset();
+	Versions.Reserve(CustomVersions.Num());
+	for (const FCustomVersion& CustomVersion : CustomVersions)
+	{
+		Versions.FindOrAdd(CustomVersion.Key, 0) = CustomVersion.Version;
+	}
+}
+
 FRigVMObjectArchiveWriter::FRigVMObjectArchiveWriter(FRigVMObjectArchive& InArchive, const UObject* InRoot)
 : Archive( InArchive )
 , Offset(0)
@@ -96,6 +109,11 @@ FRigVMObjectArchiveWriter::FRigVMObjectArchiveWriter(FRigVMObjectArchive& InArch
 	ArIgnoreOuterRef = 1;
 }
 
+FRigVMObjectArchiveWriter::~FRigVMObjectArchiveWriter()
+{
+	Archive.SetVersionsFromArchive(*this);
+}
+
 void FRigVMObjectArchiveWriter::Serialize(void* V, int64 Length)
 {
 	if(Length == 0)
@@ -244,6 +262,8 @@ FRigVMObjectArchiveReader::FRigVMObjectArchiveReader(FRigVMObjectArchive& InArch
 	SetIsSaving(false);
 	SetIsLoading(true);
 	Archive.Decompress();
+
+	ReadVersions();
 }
 
 void FRigVMObjectArchiveReader::Serialize(void* V, int64 Length)
@@ -463,8 +483,31 @@ FRigVMObjectArchiveReader::FObjectHeader FRigVMObjectArchiveReader::GetRootObjec
 	return FObjectHeader();
 }
 
+void FRigVMObjectArchiveReader::ReadVersions()
+{
+	if (Archive.Versions.IsEmpty())
+	{
+		return;
+	}
+
+	FCustomVersionContainer VersionContainer = GetCustomVersions();
+	for (const TPair<FGuid,int32>& Pair : Archive.Versions)
+	{
+		VersionContainer.SetVersion(Pair.Key, Pair.Value, VersionContainer.GetFriendlyName(Pair.Key));
+	}
+	SetCustomVersions(VersionContainer);
+}
+
 FArchive& operator<<(FArchive& Ar, FRigVMObjectArchive& Data)
 {
+	if (!Ar.IsLoading() || Ar.CustomVer(FRigVMObjectVersion::GUID) >= FRigVMObjectVersion::ObjectArchiveVersionContainerSerialization)
+	{
+		Ar << Data.Versions;
+	}
+	else
+	{
+		Data.SetVersionsFromArchive(Ar);
+	}
 	Ar << Data.Buffer;
 	Ar << Data.UncompressedSize;
 	Ar << Data.CompressedSize;
diff --git a/Engine/Plugins/Runtime/RigVM/Source/RigVM/Public/RigVMCore/RigVMObjectArchive.h b/Engine/Plugins/Runtime/RigVM/Source/RigVM/Public/RigVMCore/RigVMObjectArchive.h
--- a/Engine/Plugins/Runtime/RigVM/Source/RigVM/Public/RigVMCore/RigVMObjectArchive.h
+++ b/Engine/Plugins/Runtime/RigVM/Source/RigVM/Public/RigVMCore/RigVMObjectArchive.h
@@ -32,8 +32,13 @@ public:
 
 	RIGVM_API friend FArchive& operator<<(FArchive& Ar, FRigVMObjectArchive& Data);
 
+	RIGVM_API void SetVersionsFromArchive(const FArchive& Ar);
+
 private:
 
+	UPROPERTY()
+	TMap<FGuid,int32> Versions;
+	
 	UPROPERTY()
 	TArray<uint8> Buffer;
 
@@ -54,6 +59,7 @@ class FRigVMObjectArchiveWriter : public FArchiveUObject
 {
 public:
 	UE_API FRigVMObjectArchiveWriter( FRigVMObjectArchive& InOutArchive, const UObject* InRoot );
+	UE_API virtual ~FRigVMObjectArchiveWriter() override;
 	UE_API virtual void Serialize( void* V, int64 Length ) override;
 	UE_API virtual int64 Tell() override;
 	UE_API virtual int64 TotalSize() override;
@@ -115,6 +121,8 @@ public:
 
 protected:
 
+	void ReadVersions();
+
 	TMap<FName, UClass*> ReadClasses;
 	TMap<FName, UObject*> ReadObjects;
 	TArray<UObject*> DeserializedObjects;
diff --git a/Engine/Plugins/Runtime/RigVM/Source/RigVM/Public/RigVMObjectVersion.h b/Engine/Plugins/Runtime/RigVM/Source/RigVM/Public/RigVMObjectVersion.h
--- a/Engine/Plugins/Runtime/RigVM/Source/RigVM/Public/RigVMObjectVersion.h
+++ b/Engine/Plugins/Runtime/RigVM/Source/RigVM/Public/RigVMObjectVersion.h
@@ -78,6 +78,9 @@ struct FRigVMObjectVersion
 		// Function arguments can now represent an input variable (an external variable passed into a funciton)
 		FunctionArgumentCanRepresentInputVariable,
 
+		// Object archive is now storing the version container
+		ObjectArchiveVersionContainerSerialization,
+
 		// -----<new versions can be added above this line>-------------------------------------------------
 		VersionPlusOne,
 		LatestVersion = VersionPlusOne - 1
