RigVM: Correctly forward events between minimal preview environment and actual function node

#rnx

#Synced-CL 48944701

diff --git a/Engine/Plugins/Runtime/RigVM/Source/RigVMDeveloper/Private/EdGraph/RigVMEdGraph.cpp b/Engine/Plugins/Runtime/RigVM/Source/RigVMDeveloper/Private/EdGraph/RigVMEdGraph.cpp
--- a/Engine/Plugins/Runtime/RigVM/Source/RigVMDeveloper/Private/EdGraph/RigVMEdGraph.cpp
+++ b/Engine/Plugins/Runtime/RigVM/Source/RigVMDeveloper/Private/EdGraph/RigVMEdGraph.cpp
@@ -1053,9 +1053,11 @@ URigVMGraph* URigVMEdGraph::GetModel() const
 	
 	if (const FRigVMClient* Client = GetRigVMClient())
 	{
-		URigVMGraph* Model = Client->GetModel(this);
-		CachedModelGraph = Model;
-		return Model;
+		if (URigVMGraph* Model = Client->GetModel(this))
+		{
+			CachedModelGraph = Model;
+			return Model;
+		}
 	}
 
 	// for preview scenarios we'll nest the edgraph under the model graph
@@ -1072,7 +1074,10 @@ URigVMController* URigVMEdGraph::GetController() const
 {
 	if (FRigVMClient* Client = GetRigVMClient())
 	{
-		return Client->GetOrCreateController(GetModel());
+		if (URigVMController* Controller = Client->GetOrCreateController(GetModel()))
+		{
+			return Controller;
+		}
 	}
 	
 	// for preview scenarios we'll nest the edgraph under the model graph
diff --git a/Engine/Plugins/Runtime/RigVM/Source/RigVMEditor/Private/Editor/RigVMGraphDetailCustomization.cpp b/Engine/Plugins/Runtime/RigVM/Source/RigVMEditor/Private/Editor/RigVMGraphDetailCustomization.cpp
--- a/Engine/Plugins/Runtime/RigVM/Source/RigVMEditor/Private/Editor/RigVMGraphDetailCustomization.cpp
+++ b/Engine/Plugins/Runtime/RigVM/Source/RigVMEditor/Private/Editor/RigVMGraphDetailCustomization.cpp
@@ -822,11 +822,12 @@ FRigVMFunctionArgumentDefaultNode::FRigVMFunctionArgumentDefaultNode(
 )
 	: GraphPtr(InGraph)
 	, WeakRigVMClientHost(InClientHost)
+	, bIsHandlingModifiedEvent(false)
 {
 	if (GraphPtr.IsValid() && WeakRigVMClientHost.IsValid())
 	{
 		IRigVMClientHost* RigVMClientHost = WeakRigVMClientHost.Get();
-		RigVMClientHost->OnModified().AddRaw(this, &FRigVMFunctionArgumentDefaultNode::HandleModifiedEvent);
+		RigVMClientHost->OnModified().AddRaw(this, &FRigVMFunctionArgumentDefaultNode::HandleModifiedEvent, false);
 
 		if (URigVMLibraryNode* LibraryNode = Cast<URigVMLibraryNode>(GraphPtr->GetOuter()))
 		{
@@ -873,16 +874,21 @@ void FRigVMFunctionArgumentDefaultNode::GenerateChildContent(IDetailChildrenBuil
 		return;
 	}
 
-	TSharedPtr<FRigVMMinimalEnvironment> Environment = MakeShared<FRigVMMinimalEnvironment>(LibraryNode->GetImplementingOuter<IRigVMAssetInterface>());
-	URigVMNode* FuncRefNode = Environment->GetController()->AddFunctionReferenceNode(LibraryNode);
-	Environment->SetNode(FuncRefNode);
+	MinimalEnvironment = MakeShared<FRigVMMinimalEnvironment>(LibraryNode->GetImplementingOuter<IRigVMAssetInterface>());
+	URigVMNode* FuncRefNode = MinimalEnvironment->GetController()->AddFunctionReferenceNode(LibraryNode, FVector2D::ZeroVector, LibraryNode->GetName());
+	MinimalEnvironment->SetNode(FuncRefNode);
+
+	if (URigVMGraph* Model = MinimalEnvironment->GetModel())
+	{
+		Model->OnModified().AddSP(this, &FRigVMFunctionArgumentDefaultNode::HandleModifiedEvent, true);
+	}
 
 	ChildrenBuilder.AddCustomRow(FText::GetEmpty())
 	.WholeRowContent()
 	.MaxDesiredWidth(980.f)
 	[
 		SAssignNew(PreviewWidget, SRigVMNodePreviewWidget)
-		.Environment(Environment)
+		.Environment(MinimalEnvironment)
 		.NodeEnabled(true)
 	];
 }
@@ -895,46 +901,92 @@ void FRigVMFunctionArgumentDefaultNode::OnGraphChanged(const FEdGraphEditAction&
 	}
 }
 
-void FRigVMFunctionArgumentDefaultNode::HandleModifiedEvent(ERigVMGraphNotifType InNotifType, URigVMGraph* InGraph, UObject* InSubject)
+void FRigVMFunctionArgumentDefaultNode::HandleModifiedEvent(ERigVMGraphNotifType InNotifType, URigVMGraph* InGraph, UObject* InSubject, bool bIsPreview)
 {
-	if (!GraphPtr.IsValid())
+	if (bIsHandlingModifiedEvent)
+	{
+		return;
+	}
+	const TGuardValue<bool> GuardReEntry(bIsHandlingModifiedEvent, true);
+	
+	if (!MinimalEnvironment)
 	{
 		return;
 	}
 
-	URigVMGraph* Graph = GraphPtr.Get();
-	URigVMLibraryNode* LibraryNode = Cast<URigVMLibraryNode>(Graph->GetOuter());
-	if (LibraryNode == nullptr)
+	if (MinimalEnvironment->GetModel() != InGraph)
 	{
 		return;
 	}
-	if (LibraryNode->GetGraph() != InGraph)
+
+	const URigVMFunctionReferenceNode* FunctionReferenceNode = Cast<URigVMFunctionReferenceNode>(MinimalEnvironment->GetNode());
+	if (!FunctionReferenceNode)
 	{
 		return;
 	}
 
-	switch (InNotifType)
+	if (!WeakRigVMClientHost.IsValid())
 	{
-		case ERigVMGraphNotifType::PinAdded:
-		case ERigVMGraphNotifType::PinRemoved:
-		case ERigVMGraphNotifType::PinTypeChanged:
-		case ERigVMGraphNotifType::PinIndexChanged:
-		case ERigVMGraphNotifType::PinRenamed:
+		return;
+	}
+
+	URigVMFunctionLibrary* FunctionLibrary = WeakRigVMClientHost->GetLocalFunctionLibrary();
+	if (!FunctionLibrary)
+	{
+		return;
+	}
+
+	URigVMLibraryNode* FunctionDefinition = FunctionLibrary->FindFunction(FunctionReferenceNode->GetFunctionIdentifier().GetFunctionFName());
+	if (!FunctionDefinition)
+	{
+		return;
+	}
+
+	// we run the same handler on the notifications coming from the
+	// minimal preview environment - as well as the graph in the asset.
+	// we sync the pin states between these two environments.
+	// reentry is guarded against at the top.
+	URigVMController* Controller = nullptr;
+	if (bIsPreview)
+	{
+		if (FunctionReferenceNode->GetGraph() != InGraph)
 		{
-			URigVMPin* Pin = CastChecked<URigVMPin>(InSubject);
-			if (Pin->GetNode() == LibraryNode)
-			{
-				OnRebuildChildren.ExecuteIfBound();
-			}
-			break;
+			return;
 		}
-		case ERigVMGraphNotifType::NodeRenamed:
-		case ERigVMGraphNotifType::NodeColorChanged:
+		Controller = WeakRigVMClientHost->GetController(FunctionLibrary);
+	}
+	else
+	{
+		URigVMLibraryNode* LibraryNode = Cast<URigVMLibraryNode>(InGraph->GetOuter());
+		if (LibraryNode != FunctionDefinition)
 		{
-			URigVMNode* Node = CastChecked<URigVMNode>(InSubject);
-			if (Node == LibraryNode)
+			return;
+		}
+		Controller = MinimalEnvironment->GetController();
+	}
+	if (!Controller)
+	{
+		return;
+	}
+
+	switch (InNotifType)
+	{
+		case ERigVMGraphNotifType::PinDefaultValueChanged:
+		case ERigVMGraphNotifType::PinExpansionChanged:
+		case ERigVMGraphNotifType::PinArraySizeChanged:
+		{
+			if (const URigVMPin* PinOnReference = Cast<URigVMPin>(InSubject))
 			{
-				OnRebuildChildren.ExecuteIfBound();
+				bool bSetupUndoRedo = !GIsTransacting;
+				TSharedPtr<FScopedTransaction> Transaction;
+				if (bSetupUndoRedo)
+				{
+					Transaction = MakeShared<FScopedTransaction>(LOCTEXT("UpdateFunctionArgument", "Update Function Argument"));
+				}
+				const FString SegmentPath = PinOnReference->GetSegmentPath(true);
+				const FString PinPath = URigVMPin::JoinPinPath(FunctionDefinition->GetName(), SegmentPath);
+				(void)Controller->SetPinExpansion(PinPath, PinOnReference->IsExpanded(), bSetupUndoRedo);
+				(void)Controller->SetPinDefaultValue(PinPath, PinOnReference->GetDefaultValue(), true, bSetupUndoRedo);
 			}
 			break;
 		}
diff --git a/Engine/Plugins/Runtime/RigVM/Source/RigVMEditor/Private/Editor/RigVMMinimalEnvironment.cpp b/Engine/Plugins/Runtime/RigVM/Source/RigVMEditor/Private/Editor/RigVMMinimalEnvironment.cpp
--- a/Engine/Plugins/Runtime/RigVM/Source/RigVMEditor/Private/Editor/RigVMMinimalEnvironment.cpp
+++ b/Engine/Plugins/Runtime/RigVM/Source/RigVMEditor/Private/Editor/RigVMMinimalEnvironment.cpp
@@ -17,7 +17,7 @@ FRigVMMinimalEnvironment::FRigVMMinimalEnvironment(const IRigVMAssetInterface* I
 		}
 	}
 	ModelController = TStrongObjectPtr<URigVMController>(NewObject<URigVMController>(Outer));
-	ModelGraph = TStrongObjectPtr<URigVMGraph>(NewObject<URigVMGraph>(ModelController.Get()));
+	ModelGraph = TStrongObjectPtr<URigVMGraph>(NewObject<URigVMGraph>(ModelController.Get(), TEXT("PreviewGraph")));
 	ModelController->SetGraph(ModelGraph.Get());
 
 	EdGraphClass = URigVMEdGraph::StaticClass();
@@ -63,6 +63,7 @@ void FRigVMMinimalEnvironment::SetNode(URigVMNode* InModelNode)
 		if(URigVMEdGraph* MyEdGraph = GetEdGraph())
 		{
 			EdGraphNode = NewObject<URigVMEdGraphNode>(MyEdGraph, EdGraphNodeClass);
+			EdGraphNode->SyncGraphNodeNameWithModelNodeName(InModelNode);
 			MyEdGraph->Nodes.Add(EdGraphNode.Get());
 		}
 	}
@@ -137,6 +138,7 @@ void FRigVMMinimalEnvironment::SetSchemata(const IRigVMAssetInterface* InAssetIn
 	if(!EdGraph.IsValid() || EdGraph->GetClass() != EdGraphClass)
 	{
 		EdGraph = TStrongObjectPtr<URigVMEdGraph>(NewObject<URigVMEdGraph>(ModelGraph.Get(), EdGraphClass));
+		EdGraph->ModelNodePath = ModelGraph->GetNodePath();
 	}
 
 	ModelController->SetSchemaClass(ClientHost->GetRigVMSchemaClass());
diff --git a/Engine/Plugins/Runtime/RigVM/Source/RigVMEditor/Public/Editor/RigVMGraphDetailCustomization.h b/Engine/Plugins/Runtime/RigVM/Source/RigVMEditor/Public/Editor/RigVMGraphDetailCustomization.h
--- a/Engine/Plugins/Runtime/RigVM/Source/RigVMEditor/Public/Editor/RigVMGraphDetailCustomization.h
+++ b/Engine/Plugins/Runtime/RigVM/Source/RigVMEditor/Public/Editor/RigVMGraphDetailCustomization.h
@@ -210,7 +210,7 @@ private:
 private:
 
 	UE_API void OnGraphChanged(const FEdGraphEditAction& InAction);
-	UE_API void HandleModifiedEvent(ERigVMGraphNotifType InNotifType, URigVMGraph* InGraph, UObject* InSubject);
+	UE_API void HandleModifiedEvent(ERigVMGraphNotifType InNotifType, URigVMGraph* InGraph, UObject* InSubject, bool bIsPreview);
 
 	TWeakObjectPtr<URigVMGraph> GraphPtr;
 	TWeakObjectPtr<URigVMEdGraph> EdGraphOuterPtr;
@@ -218,6 +218,8 @@ private:
 	FSimpleDelegate OnRebuildChildren;
 	TSharedPtr<SRigVMNodePreviewWidget> PreviewWidget;
 	FDelegateHandle GraphChangedDelegateHandle;
+	TSharedPtr<FRigVMMinimalEnvironment> MinimalEnvironment;
+	bool bIsHandlingModifiedEvent;
 };
 
 /** Customization for editing rig vm graphs */
