RigVM: Node rendering facility for Docs

#rnx

#Synced-CL 49029065

diff --git a/Engine/Plugins/Runtime/RigVM/Source/RigVMDeveloper/Private/RigVMJsonUtils.cpp b/Engine/Plugins/Runtime/RigVM/Source/RigVMDeveloper/Private/RigVMJsonUtils.cpp
--- a/Engine/Plugins/Runtime/RigVM/Source/RigVMDeveloper/Private/RigVMJsonUtils.cpp
+++ b/Engine/Plugins/Runtime/RigVM/Source/RigVMDeveloper/Private/RigVMJsonUtils.cpp
@@ -18,6 +18,9 @@
 #include "RigVMCore/RigVMVariableDescription.h"
 #include "Serialization/JsonSerializer.h"
 #include "UObject/UObjectIterator.h"
+#include "Misc/FileHelper.h"
+#include "Misc/ScopedSlowTask.h"
+#include "Engine/Engine.h"
 
 #define LOCTEXT_NAMESPACE "RigVMJsonUtils"
 
@@ -31,147 +34,197 @@ FAutoConsoleCommand FCmdRigVMJsonPrintAvailableNodes
 	TEXT("Prints out a json string containing information about all available nodes."),
 	FConsoleCommandWithArgsDelegate::CreateLambda([](const TArray<FString>& Args)
 	{
-		FRigVMRegistryReadLock Registry;
-		RigVMJsonUtils::FIntrospectionSettings Settings(Registry);
-		
-		for (const FString& Arg : Args)
+		TArray<FString> ConsecutiveCommands;
 		{
-			if (Arg.Equals(TEXT("exclude-deprecated"), ESearchCase::IgnoreCase))
-			{
-				Settings.bExcludeDeprecated = true;
-			}
-			else if (Arg.Equals(TEXT("exclude-functions"), ESearchCase::IgnoreCase))
-			{
-				Settings.bExcludeFunctions = true;
-			}
-			else if (Arg.Equals(TEXT("exclude-dispatches"), ESearchCase::IgnoreCase))
-			{
-				Settings.bExcludeDispatches = true;
-			}
-			else if (Arg.Equals(TEXT("merge-templates"), ESearchCase::IgnoreCase))
-			{
-				Settings.bMergeTemplates = true;
-			}
-			else if(Arg.Len() > 8 && Arg.StartsWith(TEXT("context="), ESearchCase::IgnoreCase))
+			FRigVMRegistryReadLock Registry;
+			RigVMJsonUtils::FIntrospectionSettings Settings(Registry);
+
+			for (const FString& Arg : Args)
 			{
-				const FString ArgContextName = Arg.Mid(8);
-				UClass* ArgSchemaClass = nullptr;
-				for (TObjectIterator<UClass> ClassIterator; ClassIterator; ++ClassIterator)
+				if (Arg.Equals(TEXT("exclude-deprecated"), ESearchCase::IgnoreCase))
+				{
+					Settings.bExcludeDeprecated = true;
+				}
+				else if (Arg.Equals(TEXT("exclude-functions"), ESearchCase::IgnoreCase))
+				{
+					Settings.bExcludeFunctions = true;
+				}
+				else if (Arg.Equals(TEXT("exclude-dispatches"), ESearchCase::IgnoreCase))
+				{
+					Settings.bExcludeDispatches = true;
+				}
+				else if (Arg.Equals(TEXT("create-images"), ESearchCase::IgnoreCase))
+				{
+					Settings.bCreateImages = true;
+				}
+				else if (Arg.Equals(TEXT("merge-templates"), ESearchCase::IgnoreCase))
 				{
-					if (*ClassIterator == URigVMSchema::StaticClass())
+					Settings.bMergeTemplates = true;
+				}
+				else if(Arg.Len() > 8 && Arg.StartsWith(TEXT("context="), ESearchCase::IgnoreCase))
+				{
+					const FString ArgContextName = Arg.Mid(8);
+					UClass* ArgSchemaClass = nullptr;
+					for (TObjectIterator<UClass> ClassIterator; ClassIterator; ++ClassIterator)
 					{
-						continue;
+						if (*ClassIterator == URigVMSchema::StaticClass())
+						{
+							continue;
+						}
+						if (ClassIterator->IsChildOf(URigVMSchema::StaticClass()))
+						{
+							const FString ClassName = ClassIterator->GetName();
+							const FString ContextName = ClassName.Replace(TEXT("Schema"), TEXT(""));
+							if (ArgContextName.Equals(ContextName, ESearchCase::IgnoreCase))
+							{
+								ArgSchemaClass = *ClassIterator;
+								break;
+							}
+						}
+					}
+
+					if (ArgSchemaClass == nullptr)
+					{
+						UE_LOG(LogRigVMDeveloper, Error, TEXT("Cannot find RigVMSchema class for context '%s'"), *ArgContextName);
+						return;
 					}
-					if (ClassIterator->IsChildOf(URigVMSchema::StaticClass()))
+
+					Settings.Schema = ArgSchemaClass->GetDefaultObject<URigVMSchema>();
+					check(Settings.Schema->GetEdGraphSchemaClass());
+					Settings.EdGraphSchema = Settings.Schema->GetEdGraphSchemaClass()->GetDefaultObject<URigVMEdGraphSchema>();
+				}
+				else if(Arg.Len() > 8 && Arg.StartsWith(TEXT("filepath="), ESearchCase::IgnoreCase))
+				{
+					Settings.FilePath = Arg.Mid(9);
+				}
+				else
+				{
+					FString ContextNamesJoined;
+					TArray<FString> ContextNames;
+				
+					for (TObjectIterator<UClass> ClassIterator; ClassIterator; ++ClassIterator)
 					{
+						if (*ClassIterator == URigVMSchema::StaticClass())
+						{
+							continue;
+						}
+						if (!ClassIterator->IsChildOf(URigVMSchema::StaticClass()))
+						{
+							continue;
+						}
 						const FString ClassName = ClassIterator->GetName();
 						const FString ContextName = ClassName.Replace(TEXT("Schema"), TEXT(""));
-						if (ArgContextName.Equals(ContextName, ESearchCase::IgnoreCase))
+
+						static const TArray<FString> StringsToFilter = {
+							TEXT("EngineTest"),
+							TEXT("TestPlugin"),
+							TEXT("RigVMExample"),
+						};
+
+						bool bValidContextName = true;
+						for (const FString& StringToFilter : StringsToFilter)
 						{
-							ArgSchemaClass = *ClassIterator;
-							break;
+							if (ContextName.Contains(StringToFilter, ESearchCase::IgnoreCase))
+							{
+								bValidContextName = false;
+								break;
+							}
 						}
+						if (!bValidContextName)
+						{
+							continue;
+						}
+						ContextNames.Add(ContextName);
 					}
-				}
 
-				if (ArgSchemaClass == nullptr)
-				{
-					UE_LOG(LogRigVMDeveloper, Error, TEXT("Cannot find RigVMSchema class for context '%s'"), *ArgContextName);
+					if (!ContextNames.IsEmpty())
+					{
+						ContextNamesJoined = RigVMStringUtils::JoinStrings(ContextNames, TEXT(", "));
+					}
+
+					UE_LOG(LogRigVMDeveloper, Warning, TEXT("Valid Arguments:\n"
+						"context=MyContextName - filter the nodes by context in question (%s)\n"
+						"exclude-deprecated - only includes the current / non-deprecated nodes\n"
+						"exclude-functions - excludes RigVM struct based functions\n"
+						"exclude-dispatches - excludes RigVM dispatch factories\n"
+						"filepath=MyFilePath - the [optional] file path to export the content to"
+					), *ContextNamesJoined);
 					return;
 				}
+			}
 
-				Settings.Schema = ArgSchemaClass->GetDefaultObject<URigVMSchema>();
-				check(Settings.Schema->GetEdGraphSchemaClass());
-				Settings.EdGraphSchema = Settings.Schema->GetEdGraphSchemaClass()->GetDefaultObject<URigVMEdGraphSchema>();
+			if (Settings.bCreateImages && Settings.FilePath.IsEmpty())
+			{
+				UE_LOG(LogRigVMDeveloper, Display, TEXT("Cannot create node images - filepath argument not specified."));
+				Settings.bCreateImages = false;
 			}
-			else
+			
+			TSharedPtr<FJsonObject> JsonObject = RigVMJsonUtils::AvailableNodesToJson(Settings, [](const UScriptStruct* InStruct) -> bool
 			{
-				FString ContextNamesJoined;
-				TArray<FString> ContextNames;
-				
-				for (TObjectIterator<UClass> ClassIterator; ClassIterator; ++ClassIterator)
+				check(InStruct);
+				const UPackage* Package = InStruct->GetPackage();
+
+				static const TArray<FString> StringsToFilter = {
+					TEXT("EngineTest"),
+					TEXT("TestPlugin"),
+					TEXT("RigVMExample"),
+				};
+				for (const FString& StringToFilter : StringsToFilter)
 				{
-					if (*ClassIterator == URigVMSchema::StaticClass())
+					if (Package->GetName().Contains(StringToFilter))
 					{
-						continue;
+						return false;
 					}
-					if (!ClassIterator->IsChildOf(URigVMSchema::StaticClass()))
+				}
+				return true;
+			});
+			
+			if (JsonObject)
+			{
+				FString Text;
+				RigVMJsonUtils::FJsonWriterRef Writer = RigVMJsonUtils::MakeJsonWriter(Text);
+				if (FJsonSerializer::Serialize(JsonObject, Writer))
+				{
+	#if WITH_EDITOR
+					if (Settings.FilePath.IsEmpty())
 					{
-						continue;
+						UE_LOG(LogRigVMDeveloper, Display, TEXT("The content has been copied to the clipboard."));
+						FPlatformApplicationMisc::ClipboardCopy(*Text);
+						return;
 					}
-					const FString ClassName = ClassIterator->GetName();
-					const FString ContextName = ClassName.Replace(TEXT("Schema"), TEXT(""));
-
-					static const TArray<FString> StringsToFilter = {
-						TEXT("EngineTest"),
-						TEXT("TestPlugin"),
-						TEXT("RigVMExample"),
-					};
-
-					bool bValidContextName = true;
-					for (const FString& StringToFilter : StringsToFilter)
+	#endif
+					if (!Settings.FilePath.IsEmpty())
 					{
-						if (ContextName.Contains(StringToFilter, ESearchCase::IgnoreCase))
+						if (FFileHelper::SaveStringToFile(Text, *Settings.FilePath))
 						{
-							bValidContextName = false;
-							break;
+							UE_LOG(LogRigVMDeveloper, Display, TEXT("The content has been save to '%s'."), *Settings.FilePath);
+						}
+						else
+						{
+							UE_LOG(LogRigVMDeveloper, Error, TEXT("The content could not be saved to '%s'."), *Settings.FilePath);
 						}
 					}
-					if (!bValidContextName)
-					{
-						continue;
-					}
-					ContextNames.Add(ContextName);
 				}
 
-				if (!ContextNames.IsEmpty())
-				{
-					ContextNamesJoined = RigVMStringUtils::JoinStrings(ContextNames, TEXT(", "));
-				}
-
-				UE_LOG(LogRigVMDeveloper, Warning, TEXT("Valid Arguments:\n"
-					"context=MyContextName - filter the nodes by context in question (%s)\n"
-					"exclude-deprecated - only includes the current / non-deprecated nodes\n"
-					"exclude-functions - excludes RigVM struct based functions\n"
-					"exclude-dispatches - excludes RigVM dispatch factories"
-				), *ContextNamesJoined);
-				return;
+				ConsecutiveCommands = Settings.CommandsToRun;
 			}
 		}
-			
-		TSharedPtr<FJsonObject> JsonObject = RigVMJsonUtils::AvailableNodesToJson(Settings, [](const UScriptStruct* InStruct) -> bool
-		{
-			check(InStruct);
-			const UPackage* Package = InStruct->GetPackage();
 
-			static const TArray<FString> StringsToFilter = {
-				TEXT("EngineTest"),
-				TEXT("TestPlugin"),
-				TEXT("RigVMExample"),
-			};
-			for (const FString& StringToFilter : StringsToFilter)
+		if (!ConsecutiveCommands.IsEmpty())
+		{
+			FScopedSlowTask Progress((float)ConsecutiveCommands.Num(), LOCTEXT("ExecutingConsecutiveCommands", "Executing consecutive commands"));
+			Progress.MakeDialog(true);
+			for (const FString& Command : ConsecutiveCommands)
 			{
-				if (Package->GetName().Contains(StringToFilter))
+				if (Progress.ShouldCancel())
 				{
-					return false;
+					break;
 				}
-			}
-			return true;
-		});
-			
-		if (JsonObject)
-		{
-			FString Text;
-			RigVMJsonUtils::FJsonWriterRef Writer = RigVMJsonUtils::MakeJsonWriter(Text);
-			if (FJsonSerializer::Serialize(JsonObject, Writer))
-			{
-#if WITH_EDITOR
-				UE_LOG(LogRigVMDeveloper, Display, TEXT("The content has been copied to the clipboard."));
-				FPlatformApplicationMisc::ClipboardCopy(*Text);
-#endif
+				Progress.EnterProgressFrame(1.f, FText::FromString(Command));
+				GEngine->Exec(nullptr, *Command);
 			}
 		}
-	})
+})
 );
 
 FAutoConsoleCommand FCmdRigVMJsonPrintGraph
@@ -510,6 +563,31 @@ TSharedPtr<FJsonObject> RigVMJsonUtils::ToJson(FIntrospectionSettings& InSetting
 		}
 	}
 
+	if (InSettings.bCreateImages)
+	{
+		FString FileName;
+		FString Subject;
+		if (Template && InSettings.bMergeTemplates)
+		{
+			Subject = TemplateName;
+			Subject.ReplaceInline(TEXT(" "), TEXT("_"));
+			FileName = FString::Printf(TEXT("Template_%s"), *Subject);
+			FileName.ReplaceInline(TEXT(" "), TEXT("_"));
+		}
+		else
+		{
+			FileName = Struct->GetStructCPPName();
+			Subject = Struct->GetPathName();
+		}
+		
+		const FString OutputBaseFolder = FPaths::GetPath(InSettings.FilePath);
+		const FString RelativeImageFilePath = FString::Printf(TEXT("Images/%s.png"), *FileName);
+		const FString AbsoluteImageFilePath = FString::Printf(TEXT("%s/%s"), *OutputBaseFolder, *RelativeImageFilePath);
+		const FString Command = FString::Printf(TEXT("RigVM.JSON.RenderNodeToPNG subject=%s filepath=%s"), *Subject, *AbsoluteImageFilePath);
+		InSettings.CommandsToRun.AddUnique(Command);
+		JsonObject->SetStringField(TEXT("Image"), RelativeImageFilePath);
+	}
+
 	TMap<FName,int32> ArgumentLookup;
 	for (int32 Index = 0; Index < InFunction->Arguments.Num(); ++Index)
 	{
@@ -697,6 +775,18 @@ TSharedPtr<FJsonObject> RigVMJsonUtils::ToJson(FIntrospectionSettings& InSetting
 		}
 	}
 
+	if (InSettings.bCreateImages)
+	{
+		const FString FileName = Template->GetName().ToString();
+		const FString Subject = FileName;
+		const FString OutputBaseFolder = FPaths::GetPath(InSettings.FilePath);
+		const FString RelativeImageFilePath = FString::Printf(TEXT("Images/%s.png"), *FileName);
+		const FString AbsoluteImageFilePath = FString::Printf(TEXT("%s/%s"), *OutputBaseFolder, *RelativeImageFilePath);
+		const FString Command = FString::Printf(TEXT("RigVM.JSON.RenderNodeToPNG subject=%s filepath=%s"), *Subject, *AbsoluteImageFilePath);
+		InSettings.CommandsToRun.AddUnique(Command);
+		JsonObject->SetStringField(TEXT("Image"), RelativeImageFilePath);
+	}
+
 	TArray<FRigVMExecuteArgument> ExecuteArguments = InFactory->GetExecuteArguments_NoLock(DispatchContext, InSettings.Registry);
 	for (const FRigVMExecuteArgument& ExecuteArgument : ExecuteArguments)
 	{
@@ -960,6 +1050,7 @@ bool RigVMJsonUtils::GenericStructAttributesToJson(FIntrospectionSettings& InSet
 	OutJsonObject->SetStringField(TEXT("Keywords"), KeywordsMetadata);
 	OutJsonObject->SetStringField(TEXT("ToolTip"), PostProcessToolTip(InStruct->GetToolTipText().ToString()));
 	OutJsonObject->SetStringField(TEXT("PackageName"), InStruct->GetOutermost()->GetPathName());
+
 	return true;
 }
 
diff --git a/Engine/Plugins/Runtime/RigVM/Source/RigVMDeveloper/Public/RigVMJsonUtils.h b/Engine/Plugins/Runtime/RigVM/Source/RigVMDeveloper/Public/RigVMJsonUtils.h
--- a/Engine/Plugins/Runtime/RigVM/Source/RigVMDeveloper/Public/RigVMJsonUtils.h
+++ b/Engine/Plugins/Runtime/RigVM/Source/RigVMDeveloper/Public/RigVMJsonUtils.h
@@ -46,7 +46,10 @@ namespace RigVMJsonUtils
 		bool bExcludeDispatches = false;
 		bool bExcludeDeprecated = false;
 		bool bMergeTemplates = false;
+		bool bCreateImages = false;
+		FString FilePath;
 		TSet<TRigVMTypeIndex> DiscoveredTypes;
+		TArray<FString> CommandsToRun;
 	};
 
 	// introspection for the registry
diff --git a/Engine/Plugins/Runtime/RigVM/Source/RigVMEditor/Private/RigVMEditorBlueprintLibrary.cpp b/Engine/Plugins/Runtime/RigVM/Source/RigVMEditor/Private/RigVMEditorBlueprintLibrary.cpp
--- a/Engine/Plugins/Runtime/RigVM/Source/RigVMEditor/Private/RigVMEditorBlueprintLibrary.cpp
+++ b/Engine/Plugins/Runtime/RigVM/Source/RigVMEditor/Private/RigVMEditorBlueprintLibrary.cpp
@@ -378,7 +378,7 @@ FAutoConsoleCommand FCmdRigVMJsonRenderNodeToPNG
 		auto PrintHelp = []()
 		{
 			UE_LOG(LogRigVMDeveloper, Warning, TEXT("Valid Arguments:\n"
-				"subject=/Script/RigVM.RigVMFunction_MathBoolOr - the path name of the subject (struct or graph function)\n"
+				"subject=/Script/RigVM.RigVMFunction_MathBoolOr - the path name of the subject (struct or graph function). alternatively you can pass a template notation.\n"
 				"filepath=c:/temp/NodePreview.png\n"
 			));
 			return;					
@@ -390,6 +390,7 @@ FAutoConsoleCommand FCmdRigVMJsonRenderNodeToPNG
 		}
 		UObject* Subject = nullptr;
 		FRigVMGraphFunctionIdentifier FunctionIdentifier;
+		FName TemplateNotation = NAME_None;
 		FString FilePath;
 
 		for (const FString& Arg : Args)
@@ -400,8 +401,22 @@ FAutoConsoleCommand FCmdRigVMJsonRenderNodeToPNG
 				Subject = FindObject<UObject>(nullptr, *SubjectPathName);
 				if (!Subject)
 				{
-					UE_LOG(LogRigVMDeveloper, Error, TEXT("Invalid Subject '%s'"), *SubjectPathName);
-					return;
+					const TChunkedArray<FRigVMTemplate>& Templates = FRigVMRegistry::Get().GetTemplates();
+					for (const FRigVMTemplate& Template : Templates)
+					{
+						FString ExistingTemplateName = Template.GetName().ToString();
+						ExistingTemplateName.ReplaceInline(TEXT(" "), TEXT("_"));
+						if (ExistingTemplateName.Equals(SubjectPathName, ESearchCase::IgnoreCase))
+						{
+							TemplateNotation = Template.GetNotation();
+							break;
+						}
+					}
+					if (TemplateNotation.IsNone())
+					{
+						UE_LOG(LogRigVMDeveloper, Error, TEXT("Invalid Subject '%s'"), *SubjectPathName);
+						return;
+					}
 				}
 			}
 			else if (Arg.StartsWith(TEXT("filepath="), ESearchCase::IgnoreCase))
@@ -426,6 +441,17 @@ FAutoConsoleCommand FCmdRigVMJsonRenderNodeToPNG
 				UE_LOG(LogRigVMDeveloper, Warning, TEXT("Error when rendering node of '%s' to file '%s'."), *Subject->GetPathName(), *FilePath);
 			}
 		}
+		else if (!TemplateNotation.IsNone())
+		{
+			if (URigVMEditorBlueprintLibrary::RenderRigVMTemplateToPNG(TemplateNotation, FilePath))
+			{
+				UE_LOG(LogRigVMDeveloper, Log, TEXT("Successfully rendered image of '%s' to file '%s'."), *TemplateNotation.ToString(), *FilePath);
+			}
+			else
+			{
+				UE_LOG(LogRigVMDeveloper, Warning, TEXT("Error when rendering node of '%s' to file '%s'."), *TemplateNotation.ToString(), *FilePath);
+			}
+		}
 		else
 		{
 			PrintHelp();
@@ -435,34 +461,48 @@ FAutoConsoleCommand FCmdRigVMJsonRenderNodeToPNG
 
 bool URigVMEditorBlueprintLibrary::RenderRigVMSubjectToPNG(UObject* InSubject, const FString& OutFileName)
 {
-	URigVMEditorAsset* EditorAsset = Cast<URigVMEditorAsset>(MakeRigVMAssetForSubject(InSubject));
+	const URigVMSchema* Schema = nullptr;
+	URigVMEditorAsset* EditorAsset = Cast<URigVMEditorAsset>(MakeRigVMAssetForSubject(InSubject, Schema));
 	if (!EditorAsset)
 	{
 		return false;
 	}
 
+	EditorAsset->SetAutoVMRecompile(false);
+
 	URigVMGraph* Model = EditorAsset->GetRigVMClient()->GetDefaultModel();
 	if (!Model)
 	{
 		return false;
 	}
 
+	Model->SetSchemaClass(Schema->GetClass());
+
 	URigVMController* Controller = EditorAsset->GetRigVMClient()->GetOrCreateController(Model);
 	if (!Controller)
 	{
 		return false;
 	}
 
+	Controller->SetSchemaClass(Schema->GetClass());
+
 	URigVMNode* ModelNode = nullptr;
 	if (UScriptStruct* ScriptStruct = Cast<UScriptStruct>(InSubject))
 	{
 		if (ScriptStruct->IsChildOf(FRigVMDispatchFactory::StaticStruct()))
 		{
 			const FName FactoryName = FRigVMDispatchFactory::GetFactoryName(ScriptStruct);
-			FRigVMRegistryReadLock Registry;
-			if (const FRigVMDispatchFactory* Factory = Registry->FindDispatchFactory_NoLock(FactoryName))
+			FName TemplateNotation = NAME_None;
+			{
+				FRigVMRegistryReadLock Registry;
+				if (const FRigVMDispatchFactory* Factory = Registry->FindDispatchFactory_NoLock(FactoryName))
+				{
+					TemplateNotation = Factory->GetTemplateNotation_NoLock(Registry);
+				}
+			}
+			if (!TemplateNotation.IsNone())
 			{
-				ModelNode = Controller->AddTemplateNode(Factory->GetTemplateNotation_NoLock(Registry));
+				ModelNode = Controller->AddTemplateNode(TemplateNotation);
 			}
 		}
 		else if (ScriptStruct->IsChildOf(FRigVMStruct::StaticStruct()))
@@ -478,6 +518,57 @@ bool URigVMEditorBlueprintLibrary::RenderRigVMSubjectToPNG(UObject* InSubject, c
 	return RenderModelNodeToPNG(EditorAsset, ModelNode, OutFileName);
 }
 
+bool URigVMEditorBlueprintLibrary::RenderRigVMTemplateToPNG(FName InTemplateNotation, const FString& OutFileName)
+{
+	const FRigVMTemplate* Template = FRigVMRegistry::Get().FindTemplate(InTemplateNotation);
+	if (!Template)
+	{
+		return false;
+	}
+
+	UObject* Subject = nullptr;
+	if (const FRigVMDispatchFactory* Factory = Template->GetDispatchFactory())
+	{
+		Subject = Factory->GetScriptStruct();
+	}
+	else if (const FRigVMFunction* Function = Template->GetPrimaryPermutation())
+	{
+		Subject = Function->Struct;
+	}
+	if (!Subject)
+	{
+		return false;
+	}
+
+	const URigVMSchema* Schema = nullptr;
+	URigVMEditorAsset* EditorAsset = Cast<URigVMEditorAsset>(MakeRigVMAssetForSubject(Subject, Schema));
+	if (!EditorAsset)
+	{
+		return false;
+	}
+
+	EditorAsset->SetAutoVMRecompile(false);
+
+	URigVMGraph* Model = EditorAsset->GetRigVMClient()->GetDefaultModel();
+	if (!Model)
+	{
+		return false;
+	}
+
+	Model->SetSchemaClass(Schema->GetClass());
+
+	URigVMController* Controller = EditorAsset->GetRigVMClient()->GetOrCreateController(Model);
+	if (!Controller)
+	{
+		return false;
+	}
+
+	Controller->SetSchemaClass(Schema->GetClass());
+
+	URigVMNode* ModelNode = Controller->AddTemplateNode(Template->GetNotation());
+	return RenderModelNodeToPNG(EditorAsset, ModelNode, OutFileName);
+}
+
 bool URigVMEditorBlueprintLibrary::RenderModelNodeToPNG(URigVMEditorAsset* EditorAsset, URigVMNode* InModelNode, const FString& OutFileName)
 {
 	if (!InModelNode)
@@ -563,7 +654,7 @@ bool URigVMEditorBlueprintLibrary::RenderModelNodeToPNG(URigVMEditorAsset* Edito
 	return FFileHelper::SaveArrayToFile(CompressedBitmap, *OutFileName);
 }
 
-UObject* URigVMEditorBlueprintLibrary::MakeRigVMAssetForSubject(UObject* InUnitOrDispatchStructOrLibraryNode)
+UObject* URigVMEditorBlueprintLibrary::MakeRigVMAssetForSubject(UObject* InUnitOrDispatchStructOrLibraryNode, const URigVMSchema *& OutSchema)
 {
 	if (!InUnitOrDispatchStructOrLibraryNode)
 	{
@@ -627,6 +718,7 @@ UObject* URigVMEditorBlueprintLibrary::MakeRigVMAssetForSubject(UObject* InUnitO
 	{
 		return nullptr;
 	}
+	OutSchema = SchemaToUse;
 
 	UClass* RuntimeAssetClass = SchemaToUse->GetRuntimeAssetClass();
 	if (!RuntimeAssetClass)
diff --git a/Engine/Plugins/Runtime/RigVM/Source/RigVMEditor/Public/RigVMEditorBlueprintLibrary.h b/Engine/Plugins/Runtime/RigVM/Source/RigVMEditor/Public/RigVMEditorBlueprintLibrary.h
--- a/Engine/Plugins/Runtime/RigVM/Source/RigVMEditor/Public/RigVMEditorBlueprintLibrary.h
+++ b/Engine/Plugins/Runtime/RigVM/Source/RigVMEditor/Public/RigVMEditorBlueprintLibrary.h
@@ -136,8 +136,11 @@ public:
 	UFUNCTION(BlueprintCallable, Category = "RigVM Blueprint")
 	static UE_API bool RenderRigVMSubjectToPNG(UObject* InUnitOrDispatchStructOrLibraryNode, const FString& OutFileName);
 
+	UFUNCTION(BlueprintCallable, Category = "RigVM Blueprint")
+	static UE_API bool RenderRigVMTemplateToPNG(FName InTemplateNotation, const FString& OutFileName);
+
 	static UE_API bool RenderModelNodeToPNG(URigVMEditorAsset* EditorAsset, URigVMNode* InModelNode, const FString& OutFileName);
-	static UE_API UObject* MakeRigVMAssetForSubject(UObject* InUnitOrDispatchStructOrLibraryNode);
+	static UE_API UObject* MakeRigVMAssetForSubject(UObject* InUnitOrDispatchStructOrLibraryNode, const URigVMSchema *& OutSchema);
 };
 
 #undef UE_API
