RigVM: minimal environment to support collapse nodes

#rnx

#Synced-CL 49065145

diff --git a/Engine/Plugins/Runtime/RigVM/Source/RigVMDeveloper/Private/RigVMModel/RigVMController.cpp b/Engine/Plugins/Runtime/RigVM/Source/RigVMDeveloper/Private/RigVMModel/RigVMController.cpp
--- a/Engine/Plugins/Runtime/RigVM/Source/RigVMDeveloper/Private/RigVMModel/RigVMController.cpp
+++ b/Engine/Plugins/Runtime/RigVM/Source/RigVMDeveloper/Private/RigVMModel/RigVMController.cpp
@@ -10234,7 +10234,10 @@ bool URigVMController::SetPinDefaultValue(URigVMPin* InPin, const FString& InDef
 		&& InPin->GetCPPType() != RigVMTypeUtils::FNameType
 		&& bValidatePinDefaults)
 	{
-		ensure(!InDefaultValue.IsEmpty());
+		if(InDefaultValue.IsEmpty())
+		{
+			return false;
+		}
 	}
 
 	URigVMGraph* Graph = GetGraph();
diff --git a/Engine/Plugins/Runtime/RigVM/Source/RigVMDeveloper/Public/RigVMModel/Nodes/RigVMCollapseNode.h b/Engine/Plugins/Runtime/RigVM/Source/RigVMDeveloper/Public/RigVMModel/Nodes/RigVMCollapseNode.h
--- a/Engine/Plugins/Runtime/RigVM/Source/RigVMDeveloper/Public/RigVMModel/Nodes/RigVMCollapseNode.h
+++ b/Engine/Plugins/Runtime/RigVM/Source/RigVMDeveloper/Public/RigVMModel/Nodes/RigVMCollapseNode.h
@@ -50,6 +50,7 @@ private:
 	FString NodeDescription;
 
 	friend class URigVMController;
+	friend class FRigVMMinimalEnvironment;
 };
 
 #undef UE_API
diff --git a/Engine/Plugins/Runtime/RigVM/Source/RigVMDeveloper/Public/RigVMModel/RigVMController.h b/Engine/Plugins/Runtime/RigVM/Source/RigVMDeveloper/Public/RigVMModel/RigVMController.h
--- a/Engine/Plugins/Runtime/RigVM/Source/RigVMDeveloper/Public/RigVMModel/RigVMController.h
+++ b/Engine/Plugins/Runtime/RigVM/Source/RigVMDeveloper/Public/RigVMModel/RigVMController.h
@@ -1738,6 +1738,7 @@ private:
 	friend class URigVMSchema;
 	friend class URigVMEdGraphFunctionRefNodeSpawner;
 	friend class FRigVMControllerGraphSectionsScope;
+	friend class FRigVMMinimalEnvironment;
 };
 
 class FRigVMControllerNotifGuard
diff --git a/Engine/Plugins/Runtime/RigVM/Source/RigVMEditor/Private/Editor/RigVMGraphDetailCustomization.cpp b/Engine/Plugins/Runtime/RigVM/Source/RigVMEditor/Private/Editor/RigVMGraphDetailCustomization.cpp
--- a/Engine/Plugins/Runtime/RigVM/Source/RigVMEditor/Private/Editor/RigVMGraphDetailCustomization.cpp
+++ b/Engine/Plugins/Runtime/RigVM/Source/RigVMEditor/Private/Editor/RigVMGraphDetailCustomization.cpp
@@ -875,8 +875,14 @@ void FRigVMFunctionArgumentDefaultNode::GenerateChildContent(IDetailChildrenBuil
 	}
 
 	MinimalEnvironment = MakeShared<FRigVMMinimalEnvironment>(LibraryNode->GetImplementingOuter<IRigVMAssetInterface>());
-	URigVMNode* FuncRefNode = MinimalEnvironment->GetController()->AddFunctionReferenceNode(LibraryNode, FVector2D::ZeroVector, LibraryNode->GetName());
-	MinimalEnvironment->SetNode(FuncRefNode);
+	if (LibraryNode->GetTypedOuter<URigVMFunctionLibrary>())
+	{
+		MinimalEnvironment->SetFunctionReference(LibraryNode, LibraryNode->GetName());
+	}
+	else
+	{
+		MinimalEnvironment->SetEmptyCollapseNode(Cast<URigVMCollapseNode>(LibraryNode), LibraryNode->GetName());
+	}
 
 	if (URigVMGraph* Model = MinimalEnvironment->GetModel())
 	{
@@ -901,7 +907,7 @@ void FRigVMFunctionArgumentDefaultNode::OnGraphChanged(const FEdGraphEditAction&
 	}
 }
 
-void FRigVMFunctionArgumentDefaultNode::HandleModifiedEvent(ERigVMGraphNotifType InNotifType, URigVMGraph* InGraph, UObject* InSubject, bool bIsPreview)
+void FRigVMFunctionArgumentDefaultNode::HandleModifiedEvent(ERigVMGraphNotifType InNotifType, URigVMGraph* InGraph, UObject* InSubject, bool bFromTargetToSource)
 {
 	if (bIsHandlingModifiedEvent)
 	{
@@ -914,30 +920,24 @@ void FRigVMFunctionArgumentDefaultNode::HandleModifiedEvent(ERigVMGraphNotifType
 		return;
 	}
 
-	if (MinimalEnvironment->GetModel() != InGraph)
-	{
-		return;
-	}
-
-	const URigVMFunctionReferenceNode* FunctionReferenceNode = Cast<URigVMFunctionReferenceNode>(MinimalEnvironment->GetNode());
-	if (!FunctionReferenceNode)
+	if (!GraphPtr.IsValid())
 	{
 		return;
 	}
 
-	if (!WeakRigVMClientHost.IsValid())
+	URigVMLibraryNode* SourceNode = Cast<URigVMLibraryNode>(GraphPtr->GetOuter());
+	if (!SourceNode)
 	{
 		return;
 	}
 
-	URigVMFunctionLibrary* FunctionLibrary = WeakRigVMClientHost->GetLocalFunctionLibrary();
-	if (!FunctionLibrary)
+	URigVMLibraryNode* TargetNode = Cast<URigVMLibraryNode>(MinimalEnvironment->GetNode());
+	if (!TargetNode)
 	{
 		return;
 	}
 
-	URigVMLibraryNode* FunctionDefinition = FunctionLibrary->FindFunction(FunctionReferenceNode->GetFunctionIdentifier().GetFunctionFName());
-	if (!FunctionDefinition)
+	if (!WeakRigVMClientHost.IsValid())
 	{
 		return;
 	}
@@ -947,18 +947,17 @@ void FRigVMFunctionArgumentDefaultNode::HandleModifiedEvent(ERigVMGraphNotifType
 	// we sync the pin states between these two environments.
 	// reentry is guarded against at the top.
 	URigVMController* Controller = nullptr;
-	if (bIsPreview)
+	if (bFromTargetToSource)
 	{
-		if (FunctionReferenceNode->GetGraph() != InGraph)
+		if (TargetNode->GetGraph() != InGraph)
 		{
 			return;
 		}
-		Controller = WeakRigVMClientHost->GetController(FunctionLibrary);
+		Controller = WeakRigVMClientHost->GetController(SourceNode->GetGraph());
 	}
 	else
 	{
-		URigVMLibraryNode* LibraryNode = Cast<URigVMLibraryNode>(InGraph->GetOuter());
-		if (LibraryNode != FunctionDefinition)
+		if (SourceNode->GetGraph() != InGraph)
 		{
 			return;
 		}
@@ -975,18 +974,40 @@ void FRigVMFunctionArgumentDefaultNode::HandleModifiedEvent(ERigVMGraphNotifType
 		case ERigVMGraphNotifType::PinExpansionChanged:
 		case ERigVMGraphNotifType::PinArraySizeChanged:
 		{
-			if (const URigVMPin* PinOnReference = Cast<URigVMPin>(InSubject))
+			if (const URigVMPin* Pin = Cast<URigVMPin>(InSubject))
 			{
+				if (Pin->GetNode() != (bFromTargetToSource ? TargetNode : SourceNode))
+				{
+					return;
+				}
+				
 				bool bSetupUndoRedo = !GIsTransacting;
 				TSharedPtr<FScopedTransaction> Transaction;
 				if (bSetupUndoRedo)
 				{
 					Transaction = MakeShared<FScopedTransaction>(LOCTEXT("UpdateFunctionArgument", "Update Function Argument"));
 				}
-				const FString SegmentPath = PinOnReference->GetSegmentPath(true);
-				const FString PinPath = URigVMPin::JoinPinPath(FunctionDefinition->GetName(), SegmentPath);
-				(void)Controller->SetPinExpansion(PinPath, PinOnReference->IsExpanded(), bSetupUndoRedo);
-				(void)Controller->SetPinDefaultValue(PinPath, PinOnReference->GetDefaultValue(), true, bSetupUndoRedo);
+				const FString SegmentPath = Pin->GetSegmentPath(true);
+				const URigVMNode* OtherNode = Pin->GetNode() == TargetNode ? SourceNode : TargetNode;
+				const FString PinPath = URigVMPin::JoinPinPath(OtherNode->GetName(), SegmentPath);
+				(void)Controller->SetPinExpansion(PinPath, Pin->IsExpanded(), bSetupUndoRedo);
+				(void)Controller->SetPinDefaultValue(PinPath, Pin->GetDefaultValue(), true, bSetupUndoRedo);
+			}
+			break;
+		}
+		case ERigVMGraphNotifType::PinAdded:
+		case ERigVMGraphNotifType::PinIndexChanged:
+		case ERigVMGraphNotifType::PinRemoved:
+		case ERigVMGraphNotifType::PinRenamed:
+		case ERigVMGraphNotifType::PinTypeChanged:
+		{
+			if (!bFromTargetToSource)
+			{
+				if (TargetNode->IsA<URigVMFunctionReferenceNode>())
+				{
+					return;
+				}
+				MinimalEnvironment->SynchronizeCollapseNode(Cast<URigVMCollapseNode>(SourceNode));
 			}
 			break;
 		}
diff --git a/Engine/Plugins/Runtime/RigVM/Source/RigVMEditor/Private/Editor/RigVMMinimalEnvironment.cpp b/Engine/Plugins/Runtime/RigVM/Source/RigVMEditor/Private/Editor/RigVMMinimalEnvironment.cpp
--- a/Engine/Plugins/Runtime/RigVM/Source/RigVMEditor/Private/Editor/RigVMMinimalEnvironment.cpp
+++ b/Engine/Plugins/Runtime/RigVM/Source/RigVMEditor/Private/Editor/RigVMMinimalEnvironment.cpp
@@ -15,6 +15,7 @@ FRigVMMinimalEnvironment::FRigVMMinimalEnvironment(const IRigVMAssetInterface* I
 		{
 			Outer = const_cast<UObject*>(AssetObject);
 		}
+		WeakRigVMClientHost = InAssetInterface->GetRigVMClientHost();
 	}
 	ModelController = TStrongObjectPtr<URigVMController>(NewObject<URigVMController>(Outer));
 	ModelGraph = TStrongObjectPtr<URigVMGraph>(NewObject<URigVMGraph>(ModelController.Get(), TEXT("PreviewGraph")));
@@ -74,6 +75,90 @@ void FRigVMMinimalEnvironment::SetNode(URigVMNode* InModelNode)
 	Tick_GameThead(0.f);
 }
 
+URigVMNode* FRigVMMinimalEnvironment::SetFunctionReference(URigVMLibraryNode* InFunctionDefinition, const FString& InNodeName)
+{
+	URigVMNode* FunctionRefNode = GetController()->AddFunctionReferenceNode(InFunctionDefinition, FVector2D::ZeroVector, InNodeName);
+	if (FunctionRefNode)
+	{
+		SetNode(FunctionRefNode);
+	}
+	return FunctionRefNode;
+}
+
+URigVMNode* FRigVMMinimalEnvironment::SetEmptyCollapseNode(URigVMCollapseNode* InSourceCollapseNode, const FString& InNodeName)
+{
+	URigVMCollapseNode* CollapseNode = NewObject<URigVMCollapseNode>(GetModel(), *InNodeName);
+	if (CollapseNode)
+	{
+		FString ContainedGraphName = CollapseNode->GetName() + TEXT("_ContainedGraph");
+		if (URigVMGraph* ContainedGraph = NewObject<URigVMGraph>(CollapseNode, *ContainedGraphName))
+		{
+			ContainedGraph->SetSchemaClass(GetModel()->GetSchemaClass());
+			CollapseNode->ContainedGraph = ContainedGraph;
+		}
+
+		GetController()->AddGraphNode(CollapseNode, false);
+		SetNode(CollapseNode);
+		SynchronizeCollapseNode(InSourceCollapseNode);
+	}
+	return CollapseNode;
+}
+
+void FRigVMMinimalEnvironment::SynchronizeCollapseNode(URigVMCollapseNode* InSource)
+{
+	if (!InSource)
+	{
+		return;
+	}
+
+	URigVMCollapseNode* Target = Cast<URigVMCollapseNode>(GetNode());
+	if (!Target)
+	{
+		return;
+	}
+
+	IRigVMClientHost* RigVMClientHost = WeakRigVMClientHost.Get();
+	if (!RigVMClientHost)
+	{
+		return;
+	}
+
+	FRigVMClient* Client = RigVMClientHost->GetRigVMClient();
+	if (!Client)
+	{
+		return;
+	}
+
+	URigVMController* SourceController = Client->GetOrCreateController(InSource->GetGraph());
+	if (!SourceController)
+	{
+		return;
+	}
+
+	URigVMController* TargetController = GetController();
+	if (!TargetController)
+	{
+		return;
+	}
+
+	TArray<URigVMController::FRepopulatePinsNodeData> SourceNodesPinData, TargetNodesPinData;
+	SourceController->GenerateRepopulatePinsNodeData(SourceNodesPinData, InSource, false, false, false);
+	TargetController->GenerateRepopulatePinsNodeData(TargetNodesPinData, Target, false, false, false);
+	if (SourceNodesPinData.IsEmpty() || TargetNodesPinData.IsEmpty())
+	{
+		return;
+	}
+	SourceNodesPinData.SetNum(1, EAllowShrinking::No);
+	TargetNodesPinData.SetNum(1, EAllowShrinking::No);
+	
+	TargetNodesPinData[0].NewPinInfos = SourceNodesPinData[0].NewPinInfos;
+	{
+		const FRigVMRegistry& Registry = FRigVMRegistry::Get();
+		TargetController->GenerateRepopulatePinLists(Registry, TargetNodesPinData[0]);
+	}
+	TargetController->RepopulatePins(TargetNodesPinData);
+}
+
 void FRigVMMinimalEnvironment::SetFunctionNode(const FRigVMGraphFunctionIdentifier& InIdentifier)
 {
 	check(InIdentifier.IsValid());
diff --git a/Engine/Plugins/Runtime/RigVM/Source/RigVMEditor/Public/Editor/RigVMGraphDetailCustomization.h b/Engine/Plugins/Runtime/RigVM/Source/RigVMEditor/Public/Editor/RigVMGraphDetailCustomization.h
--- a/Engine/Plugins/Runtime/RigVM/Source/RigVMEditor/Public/Editor/RigVMGraphDetailCustomization.h
+++ b/Engine/Plugins/Runtime/RigVM/Source/RigVMEditor/Public/Editor/RigVMGraphDetailCustomization.h
@@ -210,7 +210,7 @@ private:
 private:
 
 	UE_API void OnGraphChanged(const FEdGraphEditAction& InAction);
-	UE_API void HandleModifiedEvent(ERigVMGraphNotifType InNotifType, URigVMGraph* InGraph, UObject* InSubject, bool bIsPreview);
+	UE_API void HandleModifiedEvent(ERigVMGraphNotifType InNotifType, URigVMGraph* InGraph, UObject* InSubject, bool bFromTargetToSource);
 
 	TWeakObjectPtr<URigVMGraph> GraphPtr;
 	TWeakObjectPtr<URigVMEdGraph> EdGraphOuterPtr;
diff --git a/Engine/Plugins/Runtime/RigVM/Source/RigVMEditor/Public/Editor/RigVMMinimalEnvironment.h b/Engine/Plugins/Runtime/RigVM/Source/RigVMEditor/Public/Editor/RigVMMinimalEnvironment.h
--- a/Engine/Plugins/Runtime/RigVM/Source/RigVMEditor/Public/Editor/RigVMMinimalEnvironment.h
+++ b/Engine/Plugins/Runtime/RigVM/Source/RigVMEditor/Public/Editor/RigVMMinimalEnvironment.h
@@ -25,15 +25,20 @@ public:
 
 	UE_API void SetSchemata(const IRigVMAssetInterface* InAssetInterface);
 	UE_API void SetNode(URigVMNode* InModelNode);
+	UE_API URigVMNode* SetFunctionReference(URigVMLibraryNode* InFunctionDefinition, const FString& InNodeName);
+	UE_API URigVMNode* SetEmptyCollapseNode(URigVMCollapseNode* InSourceCollapseNode, const FString& InNodeName);
 	UE_API void SetFunctionNode(const FRigVMGraphFunctionIdentifier& InIdentifier);
 	
 	UE_API FSimpleDelegate& OnChanged();
 	UE_API void Tick_GameThead(float InDeltaTime);
 
+	UE_API void SynchronizeCollapseNode(URigVMCollapseNode* InSource);
+	
 private:
 
 	UE_API void HandleModified(ERigVMGraphNotifType InNotification, URigVMGraph* InGraph, UObject* InSubject);
-	
+
+	TWeakInterfacePtr<IRigVMClientHost> WeakRigVMClientHost;
 	TStrongObjectPtr<URigVMGraph> ModelGraph;
 	TStrongObjectPtr<URigVMController> ModelController;
 	UClass* EdGraphClass;
