Control Rig: Fix crash when redoing duplication of array element pin
#rnx

#Synced-CL 49061339

diff --git a/Engine/Plugins/Runtime/RigVM/Source/RigVMDeveloper/Private/RigVMModel/RigVMController.cpp b/Engine/Plugins/Runtime/RigVM/Source/RigVMDeveloper/Private/RigVMModel/RigVMController.cpp
--- a/Engine/Plugins/Runtime/RigVM/Source/RigVMDeveloper/Private/RigVMModel/RigVMController.cpp
+++ b/Engine/Plugins/Runtime/RigVM/Source/RigVMDeveloper/Private/RigVMModel/RigVMController.cpp
@@ -11868,21 +11868,23 @@ URigVMPin* URigVMController::InsertArrayPin(URigVMPin* ArrayPin, int32 InIndex,
 		Action.SetTitle(FString::Printf(TEXT("Insert Array Pin")));
 		GetActionStack()->BeginAction(Action);
 	}
+	
+	TFunction<void(URigVMPin*)> ClearPinPathCacheRecursively;
+	ClearPinPathCacheRecursively = [&](URigVMPin* InPin)
+		{
+			InPin->CachedPinPath.ResetCachedValue();
+			InPin->CachedPinPathWithNodePath.ResetCachedValue();
+			for (URigVMPin* SubPin : InPin->GetSubPins())
+			{
+				ClearPinPathCacheRecursively(SubPin);
+			}
+		};
 
 	for (int32 ExistingIndex = ArrayPin->GetSubPins().Num() - 1; ExistingIndex >= InIndex; ExistingIndex--)
 	{
 		URigVMPin* ExistingPin = ArrayPin->GetSubPins()[ExistingIndex];
 		RenameObject(ExistingPin, *FString::FormatAsNumber(ExistingIndex + 1));
-
-		TArray<URigVMPin*> PinsToClearCache;
-		PinsToClearCache.Add(ExistingPin);
-		for (int32 i=0; i<PinsToClearCache.Num(); ++i)
-		{
-			URigVMPin* PinToClearCache = PinsToClearCache[i];
-			PinToClearCache->CachedPinPath.ResetCachedValue();
-			PinToClearCache->CachedPinPathWithNodePath.ResetCachedValue();
-			PinsToClearCache.Append(PinToClearCache->GetSubPins());
-		}
+		ClearPinPathCacheRecursively(ExistingPin);
 
 		// temporarily force a the index based display name
 		// so that the ed graph node can pick that up during handling the notify
