RigVM: Further improve fuzzy search for action menu

#rnx

#Synced-CL 49201700

diff --git a/Engine/Plugins/Runtime/RigVM/Source/RigVMDeveloper/Private/EdGraph/RigVMEdGraphSchema.cpp b/Engine/Plugins/Runtime/RigVM/Source/RigVMDeveloper/Private/EdGraph/RigVMEdGraphSchema.cpp
--- a/Engine/Plugins/Runtime/RigVM/Source/RigVMDeveloper/Private/EdGraph/RigVMEdGraphSchema.cpp
+++ b/Engine/Plugins/Runtime/RigVM/Source/RigVMDeveloper/Private/EdGraph/RigVMEdGraphSchema.cpp
@@ -1082,96 +1082,205 @@ float URigVMEdGraphSchema::GetActionFilteredWeight(const FEdGraphSchemaAction& I
 	}
 	TArray<FString> MenuDescriptionsReversed = MenuDescriptions;
 	Algo::Reverse(MenuDescriptionsReversed);
-	
-	const TArray<FString>& Keywords = InCurrentAction.GetSearchKeywordsArray();
 
-	auto LevenshteinSingle = [](const FStringView& A, const FStringView& B) -> float
+	TArray<FString> KeywordsArray = InCurrentAction.GetSearchKeywordsArray();
+	FString KeywordsJoined = RigVMStringUtils::JoinStrings(KeywordsArray,TEXT(",")).ToLower();
+	KeywordsArray.Reset();
+	if (!RigVMStringUtils::SplitString(KeywordsJoined, TEXT(","), KeywordsArray))
 	{
-		float WorstCase = static_cast<float>(A.Len() + B.Len());
-		WorstCase = WorstCase < 1.0f ? 1.0f : WorstCase;
-		return 1.0f - (static_cast<float>(Algo::LevenshteinDistance(A, B)) / WorstCase);
-	};
+		KeywordsArray = InCurrentAction.GetSearchKeywordsArray();
+	}
+	const TSet<FString> UniqueKeywords = TSet<FString>(KeywordsArray);
+	KeywordsArray = UniqueKeywords.Array();
+
+	auto LevenshteinSingle = [](const FStringView& A, const FStringView& B, const int32& MinLength) -> float
+		{
+			if (A.Len() < MinLength || B.Len() < MinLength)
+			{
+				return A.Equals(B, ESearchCase::CaseSensitive) ? 1.f : 0.f;
+			}
+			
+			const float WorstCase = static_cast<float>(FMath::Max<int32>(A.Len(), B.Len()) - MinLength);
+			if (WorstCase < 1.f - SMALL_NUMBER)
+			{
+				return 0.f;
+			}
+
+			const int32 Distance = Algo::LevenshteinDistance(A, B);
+			return FMath::Max(0.f, 1.0f - (static_cast<float>(Distance) / WorstCase));
+		};
 
 	// returns the maximum weight for a given clause within an array of candidates
-	auto LevenshteinSingleMaxInArray = [LevenshteinSingle](const FString& A, const TArrayView<const FString>& B) -> float
-	{
-		float Weight = 0.f;
-		for (const FString& ElementInB : B)
+	auto LevenshteinSingleMaxInArray = [LevenshteinSingle](const FString& A, const TArrayView<const FString>& B, const int32& MinLength) -> float
 		{
-			Weight = FMath::Max(Weight, LevenshteinSingle(A, ElementInB));
-		}
-		return Weight;
-	};
+			float Weight = 0.f;
+			for (const FString& ElementInB : B)
+			{
+				Weight = FMath::Max(Weight, LevenshteinSingle(A, ElementInB, MinLength));
+			}
+			return Weight;
+		};
 
 	// returns the maximum weight for an array of clauses within an array of candidates
-	auto LevenshteinMaxInArray = [LevenshteinSingleMaxInArray](const TArrayView<const FString>& A, const TArrayView<const FString>& B) -> float
-	{
-		float Weight = 0.f;
-		for (const FString& ElementInA : A)
+	auto LevenshteinMaxInArray = [LevenshteinSingleMaxInArray](const TArrayView<const FString>& A, const TArrayView<const FString>& B, const int32& MinLength) -> float
 		{
-			Weight = FMath::Max(Weight, LevenshteinSingleMaxInArray(ElementInA, B));
-		}
-		return Weight;
-	};
+			float Weight = 0.f;
+			for (const FString& ElementInA : A)
+			{
+				Weight = FMath::Max(Weight, LevenshteinSingleMaxInArray(ElementInA, B, MinLength));
+			}
+			return Weight;
+		};
 
 	// returns the maximum weight for all clauses for a matching array of candidates
-	auto LevenshteinSameSizeArrays = [LevenshteinSingle](const TArrayView<const FString>& A, const TArrayView<const FString>& B) -> float
+	auto LevenshteinSameSizeArrays = [LevenshteinSingle](const TArrayView<const FString>& A, const TArrayView<const FString>& B, const int32& MinLength) -> float
+		{
+			if (A.Num() != B.Num() || A.IsEmpty() || B.IsEmpty())
+			{
+				return 0.f;
+			}
+			float Weight = 0.f;
+			for (int32 Index = 0; Index < A.Num(); Index++)
+			{
+				const float SingleWeight = LevenshteinSingle(A[Index], B[Index], MinLength);
+				if (SingleWeight < SMALL_NUMBER)
+				{
+					return 0;
+				}
+				Weight += SingleWeight;
+			}
+			return Weight / static_cast<float>(A.Num());
+		};
+
+	// returns the maximum weight of an array of clauses for any matching subarrays within a larger array of candidates
+	auto LevenshteinArbitraryArrays = [LevenshteinSameSizeArrays](const TArrayView<const FString>& A, const TArrayView<const FString>& B, const int32& MinLength) -> float
+		{
+			if (A.IsEmpty() || B.IsEmpty() || A.Num() > B.Num())
+			{
+				return 0.f;
+			}
+			float Weight = 0.f;
+			for (int32 Offset = 0; Offset <= B.Num() - A.Num(); Offset++)
+			{
+				const TArrayView<const FString> PartOfB = B.Slice(Offset, A.Num());
+				Weight = FMath::Max(Weight, LevenshteinSameSizeArrays(A, PartOfB, MinLength));
+			}
+			return Weight;
+		};
+
+	const FString JoinedSanitizedFilterTerms = RigVMStringUtils::JoinStrings(SanitizedFilterTerms, TEXT(""));
+	const FString JoinedMenuDescriptions = RigVMStringUtils::JoinStrings(MenuDescriptions, TEXT(""));
+	const TArray<FString>& Categories = InCurrentAction.GetCategoryChain();
+	const FString JoinedCategories = RigVMStringUtils::JoinStrings(Categories, TEXT(""));
+
+	enum EScoreType : int32
 	{
-		if (A.Num() != B.Num() || A.IsEmpty() || B.IsEmpty())
+		Description,
+		Keywords,
+		ConstantCategoryWeight,
+		CategoriesFuzzyMatch,
+		NumberOfWordsDifference,
+		TotalNumberOfScoreTypes
+	};
+
+	static constexpr uint64 Divisor = 32;
+	static constexpr uint32 MinimumFuzzyLength = 3;
+
+	// build up the score with a set of priorities
+	TArray<float, TInlineAllocator<TotalNumberOfScoreTypes>> Scores;
+	Scores.AddZeroed(TotalNumberOfScoreTypes);
+
+	// exact match start has the highest weight
+	if (JoinedMenuDescriptions.Equals(JoinedSanitizedFilterTerms, ESearchCase::IgnoreCase))
+	{
+		Scores[Description] = 1;
+	}
+	else if (JoinedMenuDescriptions.StartsWith(JoinedSanitizedFilterTerms, ESearchCase::IgnoreCase))
+	{
+		Scores[Description] = 0.9f;
+	}
+	else if (JoinedMenuDescriptions.Contains(JoinedSanitizedFilterTerms, ESearchCase::IgnoreCase))
+	{
+		Scores[Description] = 0.8f;
+	}
+	else
+	{
+		Scores[Description] = 0.7f * LevenshteinSingle(JoinedSanitizedFilterTerms, JoinedMenuDescriptions, MinimumFuzzyLength);
+		if (Scores[Description] < SMALL_NUMBER)
 		{
-			return 0.f;
+			const FString CategoryAndDescription = JoinedCategories + JoinedMenuDescriptions;
+			Scores[Description] = 0.6f * LevenshteinSingle(JoinedSanitizedFilterTerms, CategoryAndDescription, MinimumFuzzyLength);
 		}
-		float Weight = 0.f;
-		for (int32 Index = 0; Index < A.Num(); Index++)
+		else if (Scores[Description] < SMALL_NUMBER)
 		{
-			Weight += LevenshteinSingle(A[Index], B[Index]);
+			Scores[Description] = 0.5f * FMath::Max(LevenshteinArbitraryArrays(SanitizedFilterTerms, MenuDescriptions, MinimumFuzzyLength), LevenshteinArbitraryArrays(SanitizedFilterTerms, MenuDescriptionsReversed, MinimumFuzzyLength));
 		}
-		return Weight / static_cast<float>(A.Num());
-	};
+	}
 
-	// returns the maximum weight of an array of clauses for any matching subarrays within a larger array of candidates
-	auto LevenshteinArbitraryArrays = [LevenshteinSameSizeArrays](const TArrayView<const FString>& A, const TArrayView<const FString>& B) -> float
+	if (Scores[Description] < 0.25f)
 	{
-		if (A.IsEmpty() || B.IsEmpty() || A.Num() > B.Num())
+		// exact keyword match start has the second-highest weight
+		if (KeywordsArray.ContainsByPredicate([JoinedSanitizedFilterTerms](const FString& Keyword) -> bool
+		{
+			return Keyword.Equals(JoinedSanitizedFilterTerms, ESearchCase::IgnoreCase);
+		}))
 		{
-			return 0.f;
+			Scores[Keywords] = 1.f;
 		}
-		float Weight = 0.f;
-		for (int32 Offset = 0; Offset <= B.Num() - A.Num(); Offset++)
+		else if (KeywordsArray.ContainsByPredicate([JoinedSanitizedFilterTerms](const FString& Keyword) -> bool
 		{
-			const TArrayView<const FString> PartOfB = B.Slice(Offset, A.Num());
-			Weight = FMath::Max(Weight, LevenshteinSameSizeArrays(A, PartOfB));
+			return Keyword.Contains(JoinedSanitizedFilterTerms, ESearchCase::IgnoreCase);
+		}))
+		{
+			Scores[Keywords] = 0.9f;
+		}
+		else
+		{
+			// fuzzy match within the keywords
+			Scores[Keywords] = 0.5f * LevenshteinMaxInArray(SanitizedFilterTerms, KeywordsArray, MinimumFuzzyLength);
 		}
-		return Weight;
+	}
+
+	// if we find multiple matches with the same weight - let's add the category in as well
+	static const TMap<FString,float> CategoryWeights = {
+		{TEXT("math"), 0.5f},
+		{TEXT("array"), 0.4f},
+		{TEXT("core"), 0.3f},
+		{TEXT("execution"), 0.2f},
+		{TEXT("simulation"), 0.1f},
 	};
 
-	const FString JoinedSanitizedFilterTerms = RigVMStringUtils::JoinStrings(SanitizedFilterTerms, TEXT(""));
-	const FString JoinedMenuDescriptions = RigVMStringUtils::JoinStrings(MenuDescriptions, TEXT(""));
+	if ((Scores[Description] + Scores[Keywords]) > SMALL_NUMBER)
+	{
+		if (!Categories.IsEmpty())
+		{
+			if (const float* CategoryWeight = CategoryWeights.Find(Categories[0].ToLower()))
+			{
+				Scores[ConstantCategoryWeight] = *CategoryWeight;
+			}
+		}
+	}
 
-	// build up the score with a set of priorities
-	float Score = 0.f;
-	
-	// exact match start has the highest weight
-	if (JoinedMenuDescriptions.StartsWith(JoinedSanitizedFilterTerms, ESearchCase::IgnoreCase))
+	// fuzzy match within the categories
+	Scores[CategoriesFuzzyMatch] = LevenshteinSingle(JoinedSanitizedFilterTerms, JoinedCategories, MinimumFuzzyLength);
+
+	if ((Scores[Description] + Scores[Keywords] + Scores[CategoriesFuzzyMatch]) > SMALL_NUMBER)
 	{
-		Score += 2048.f;
+		// additional weight based on the distance in the number of terms and the number of menu descriptions (get transform vs get transform array)
+		Scores[NumberOfWordsDifference] = 1.f - (FMath::Max(0.f, static_cast<float>(MenuDescriptions.Num() - SanitizedFilterTerms.Num())) / static_cast<float>(FMath::Max3(MenuDescriptions.Num(), SanitizedFilterTerms.Num(), 1)));
 	}
-	// string contained somewhere second highest weight
-	if (JoinedMenuDescriptions.Contains(JoinedSanitizedFilterTerms, ESearchCase::IgnoreCase))
+
+	uint64 OverallScore = 0;
+	uint64 Multiplier = 1;
+	for (int32 ScoreIndex = Scores.Num() - 1; ScoreIndex >= 0; --ScoreIndex)
 	{
-		Score += 1024.f;
+		const int32 ScoreInt32 = FMath::Min(static_cast<int32>(Divisor) - 1, FMath::FloorToInt32(Scores[ScoreIndex] * (static_cast<float>(Divisor) - 0.1f)));
+		OverallScore += static_cast<uint64>(ScoreInt32) * Multiplier;
+		Multiplier *= Divisor;
 	}
-	
-	// fuzzy match with the joined string
-	Score += 512.f * LevenshteinSingle(JoinedSanitizedFilterTerms, JoinedMenuDescriptions);
-	// fuzzy match within the menu descriptions
-	Score += 256.f * FMath::Max(LevenshteinArbitraryArrays(SanitizedFilterTerms, MenuDescriptions), LevenshteinArbitraryArrays(SanitizedFilterTerms, MenuDescriptionsReversed));
-	// fuzzy match within the keywords
-	Score += 64.f * LevenshteinMaxInArray(SanitizedFilterTerms, Keywords);
-	// additional weight based on the distance in the number of terms and the number of menu descriptions (get transform vs get transform array)
-	Score += 32.f * (1.f - (FMath::Max(0.f, static_cast<float>(MenuDescriptions.Num() - SanitizedFilterTerms.Num())) / static_cast<float>(FMath::Max3(MenuDescriptions.Num(), SanitizedFilterTerms.Num(), 1))));
-
-	return Score;
+
+	const double OverallScoreDouble = static_cast<double>(OverallScore);
+	return static_cast<float>(OverallScoreDouble);
 }
 
 bool URigVMEdGraphSchema::SupportsPinType(const UScriptStruct* ScriptStruct) const
diff --git a/Engine/Source/Editor/GraphEditor/Private/SGraphActionMenu.cpp b/Engine/Source/Editor/GraphEditor/Private/SGraphActionMenu.cpp
--- a/Engine/Source/Editor/GraphEditor/Private/SGraphActionMenu.cpp
+++ b/Engine/Source/Editor/GraphEditor/Private/SGraphActionMenu.cpp
@@ -1319,9 +1319,10 @@ void SGraphActionMenu::ScoreAndAddActions(int32 StartingIndex)
 		// If we're filtering, search check to see if we need to show this action
 		bool bShowAction = true;
 		float EachWeight = TNumericLimits<float>::Lowest();
+		const bool bFuzzyFiltering = ActionSchema->AllowsFuzzyActionFiltering(*CurrentAction);
 
 		// only filter out per strict rule if the schema doesn't rely on fuzzy filter weights
-		if (!ActionSchema->AllowsFuzzyActionFiltering(*CurrentAction))
+		if (!bFuzzyFiltering)
 		{
 			const FString& SearchText = CurrentAction->GetFullSearchText();
 			for (int32 FilterIndex = 0; (FilterIndex < FilterTerms.Num()) && bShowAction; ++FilterIndex)
@@ -1341,6 +1342,13 @@ void SGraphActionMenu::ScoreAndAddActions(int32 StartingIndex)
 		{
 			// Get the 'weight' of this in relation to the filter
 			EachWeight = ActionSchema->GetActionFilteredWeight(*CurrentAction, FilterTerms, SanitizedFilterTerms, DraggedFromPins);
+
+			// skip this action if the weight is too low
+			if (bFuzzyFiltering && EachWeight <= SMALL_NUMBER)
+	        {
+			    continue;
+            }
+			
 			// If this action has a greater relevance than others, cache its index.
 			if (EachWeight > BestMatchCount)
 			{
