RigVM: Follow up to 49202478 - fix static PVS analysis

#rnx

#Synced-CL 49201700

diff --git a/Engine/Plugins/Runtime/RigVM/Source/RigVMDeveloper/Private/EdGraph/RigVMEdGraphSchema.cpp b/Engine/Plugins/Runtime/RigVM/Source/RigVMDeveloper/Private/EdGraph/RigVMEdGraphSchema.cpp
--- a/Engine/Plugins/Runtime/RigVM/Source/RigVMDeveloper/Private/EdGraph/RigVMEdGraphSchema.cpp
+++ b/Engine/Plugins/Runtime/RigVM/Source/RigVMDeveloper/Private/EdGraph/RigVMEdGraphSchema.cpp
@@ -1211,7 +1211,7 @@ float URigVMEdGraphSchema::GetActionFilteredWeight(const FEdGraphSchemaAction& I
 			const FString CategoryAndDescription = JoinedCategories + JoinedMenuDescriptions;
 			Scores[Description] = 0.6f * LevenshteinSingle(JoinedSanitizedFilterTerms, CategoryAndDescription, MinimumFuzzyLength);
 		}
-		else if (Scores[Description] < SMALL_NUMBER)
+		if (Scores[Description] < SMALL_NUMBER)
 		{
 			Scores[Description] = 0.5f * FMath::Max(LevenshteinArbitraryArrays(SanitizedFilterTerms, MenuDescriptions, MinimumFuzzyLength), LevenshteinArbitraryArrays(SanitizedFilterTerms, MenuDescriptionsReversed, MinimumFuzzyLength));
 		}
