RigVM: Use generated class information rather than class name to test if a class is a skeleton class.

#rnx
#Synced-CL 48772759

diff --git a/Engine/Plugins/Runtime/RigVM/Source/RigVM/Private/RigVMBlueprintGeneratedClass.cpp b/Engine/Plugins/Runtime/RigVM/Source/RigVM/Private/RigVMBlueprintGeneratedClass.cpp
--- a/Engine/Plugins/Runtime/RigVM/Source/RigVM/Private/RigVMBlueprintGeneratedClass.cpp
+++ b/Engine/Plugins/Runtime/RigVM/Source/RigVM/Private/RigVMBlueprintGeneratedClass.cpp
@@ -24,11 +24,20 @@ uint8* URigVMBlueprintGeneratedClass::GetPersistentUberGraphFrame(UObject* Obj,
 
 void URigVMBlueprintGeneratedClass::PostInitInstance(UObject* InObj, FObjectInstancingGraph* InstanceGraph)
 {
-	// Skip SKEL classes.
-	if (InObj->GetName().StartsWith(TEXT("SKEL_")) || InObj->GetName().StartsWith(TEXT("Default__SKEL_")))
+#if WITH_EDITOR
+	// If this is a non-authoritative class (e.g. a skeleton class, skip it).
+	if (GetAuthoritativeClass() != this)
 	{
 		return;
 	}
+	
+	// There's one case where the class source blueprint doesn't specify a generated class that we must also skip over.
+	const UBlueprint* ClassSource = Cast<UBlueprint>(ClassGeneratedBy);
+	if (ClassSource && ClassSource->GeneratedClass == nullptr)
+	{
+		return;
+	}
+#endif
 
 	if (URigVMHost* Owner = Cast<URigVMHost>(InObj))
 	{
