RigVM: Moving debug memory storage to the context. Only add debug memory if needed per instance. Allow VM to create the debug memory

#rnx

full cook is here https://horde.devtools.epicgames.com/job/6964ce1dcc4c489b81fb3026

#Synced-CL 49723408

diff --git a/Engine/Plugins/Runtime/RigVM/Source/RigVM/Private/RigVMCore/RigVM.cpp b/Engine/Plugins/Runtime/RigVM/Source/RigVM/Private/RigVMCore/RigVM.cpp
--- a/Engine/Plugins/Runtime/RigVM/Source/RigVM/Private/RigVMCore/RigVM.cpp
+++ b/Engine/Plugins/Runtime/RigVM/Source/RigVM/Private/RigVMCore/RigVM.cpp
@@ -174,9 +174,6 @@ void URigVM::Save(FArchive& Ar)
 	Ar << Parameters;
 	UE_RIGVM_ARCHIVETRACE_ENTRY(Ar, TEXT("Parameters"));
 
-	Ar << OperandToDebugRegisters;
-	UE_RIGVM_ARCHIVETRACE_ENTRY(Ar, TEXT("OperandToDebugRegisters"));
-	
 	Ar << LiteralMemoryStorage;
 	UE_RIGVM_ARCHIVETRACE_ENTRY(Ar, TEXT("LiteralMemoryStorage"));
 	Ar << DefaultWorkMemoryStorage;
@@ -231,7 +228,7 @@ void URigVM::Save(FArchive& Ar)
 	}
 	else
 	{
-		Ar << DefaultDebugMemoryStorage;
+		Ar << DefaultDebugMemoryStorage_DEPRECATED;
 		UE_RIGVM_ARCHIVETRACE_ENTRY(Ar, TEXT("DefaultDebugMemoryStorage"));
 
 		bool bStoredLocalizedRegistry = false;
@@ -291,12 +288,6 @@ void URigVM::Load(FArchive& Ar)
 				return;
 			}
 
-			if (Ar.CustomVer(FUE5ReleaseStreamObjectVersion::GUID) >= FUE5ReleaseStreamObjectVersion::RigVMSaveDebugMapInGraphFunctionData ||
-			    Ar.CustomVer(FFortniteMainBranchObjectVersion::GUID) >= FFortniteMainBranchObjectVersion::RigVMSaveDebugMapInGraphFunctionData)
-			{
-				Ar << OperandToDebugRegisters;
-			}
-			
 			if (Ar.CustomVer(FRigVMObjectVersion::GUID) >= FRigVMObjectVersion::VMStoringUserDefinedStructMap 
 				&& Ar.CustomVer(FRigVMObjectVersion::GUID) < FRigVMObjectVersion::HostStoringUserDefinedData)
 			{	
@@ -335,7 +326,11 @@ void URigVM::Load(FArchive& Ar)
 	if (Ar.CustomVer(FUE5ReleaseStreamObjectVersion::GUID) >= FUE5ReleaseStreamObjectVersion::RigVMSaveDebugMapInGraphFunctionData ||
 		Ar.CustomVer(FFortniteMainBranchObjectVersion::GUID) >= FFortniteMainBranchObjectVersion::RigVMSaveDebugMapInGraphFunctionData)
 	{
-		Ar << OperandToDebugRegisters;
+		if (Ar.CustomVer(FRigVMObjectVersion::GUID) < FRigVMObjectVersion::DebugOperandMappingSimplified)
+		{
+			TMap<FRigVMOperand, TArray<FRigVMOperand>> PreviousOperandToDebugRegisters;
+			Ar << PreviousOperandToDebugRegisters;
+		}
 	}
 		
 	if (Ar.CustomVer(FRigVMObjectVersion::GUID) >= FRigVMObjectVersion::VMStoringUserDefinedStructMap
@@ -361,7 +356,7 @@ void URigVM::Load(FArchive& Ar)
 	if (Ar.CustomVer(FRigVMObjectVersion::GUID) >= FRigVMObjectVersion::VMMemoryStorageDefaultsGeneratedAtVM)
 	{
 		Ar << DefaultWorkMemoryStorage;
-		Ar << DefaultDebugMemoryStorage;
+		Ar << DefaultDebugMemoryStorage_DEPRECATED;
 	}
 
 	if (Ar.CustomVer(FRigVMObjectVersion::GUID) >= FRigVMObjectVersion::LocalizedRegistry)
@@ -411,11 +406,9 @@ void URigVM::CopyDataForSerialization(URigVM* InVM)
 	ByteCodeStorage = InVM->ByteCodeStorage;
 	Parameters = InVM->Parameters;
 
-	OperandToDebugRegisters = InVM->OperandToDebugRegisters;
-
 	LiteralMemoryStorage = InVM->LiteralMemoryStorage;
 	DefaultWorkMemoryStorage = InVM->DefaultWorkMemoryStorage;
-	DefaultDebugMemoryStorage = InVM->DefaultDebugMemoryStorage;
+	DefaultDebugMemoryStorage_DEPRECATED.Reset();
 	CopyExternalVariableDefs(InVM);
 
 	LocalizedRegistry.Reset();
@@ -606,7 +599,7 @@ UClass* URigVM::GetNativizedClass(const TArray<FRigVMExternalVariableDef>& InExt
 bool URigVM::ValidateBytecode()
 {
 	// check all operands on all ops for validity
-	const TArray<const FRigVMMemoryStorageStruct*> LocalMemory = { &GetDefaultWorkMemory(), &GetDefaultLiteralMemory(), &GetDefaultDebugMemory() };
+	const TArray<const FRigVMMemoryStorageStruct*> LocalMemory = { &GetDefaultWorkMemory(), &GetDefaultLiteralMemory(), &DefaultDebugMemoryStorage_DEPRECATED };
 	
 	auto CheckOperandValidity = [LocalMemory, this](const FRigVMOperand& InOperand) -> bool
 	{
@@ -780,7 +773,6 @@ void URigVM::Reset_Internal()
 	Instructions.Reset();
 	Parameters.Reset();
 	ParametersNameMap.Reset();
-	OperandToDebugRegisters.Reset();
 
 	FunctionNamesPtr = &FunctionNamesStorage;
 	FunctionsPtr = &FunctionsStorage;
@@ -815,8 +807,6 @@ void URigVM::Empty(FRigVMExtendedExecuteContext& Context)
 
 	InvalidateCachedMemory(Context);
 
-	OperandToDebugRegisters.Empty();
-
 	Context.ExternalVariableRuntimeData.Reset();
 }
 
@@ -898,6 +888,83 @@ const FRigVMMemoryStorageStruct* URigVM::GetMemoryByType(const FRigVMExtendedExe
 	return const_cast<URigVM*>(this)->GetMemoryByType(const_cast<FRigVMExtendedExecuteContext&>(Context), InMemoryType/*, false*/);
 }
 
+FRigVMMemoryStorageStruct* URigVM::CreateDebugMemory(FRigVMExtendedExecuteContext& Context)
+{
+#if WITH_EDITOR
+	// first find all properties on the work memory which are used by any input or output operand
+	// on the bytecode while skipping execute contexts.
+	TMap<FRigVMOperand, FRigVMOperand> OperandMap;
+	TSet<int32> WorkPropertyIndices;
+	const FRigVMByteCode& ByteCode = GetByteCode();
+	for (int32 Index = 0; Index < ByteCode.GetNumInstructions(); Index++)
+	{
+		const FRigVMOperandArray InputOperands = ByteCode.GetInputOperands(Index);
+		const FRigVMOperandArray OutputOperands = ByteCode.GetOutputOperands(Index);
+		for (int32 Phase = 0; Phase < 2; Phase++)
+		{
+			const FRigVMOperandArray& Operands = Phase == 0 ? InputOperands : OutputOperands;
+			for (const FRigVMOperand& Operand : Operands)
+			{
+				if (Operand.GetMemoryType() != ERigVMMemoryType::Work)
+				{
+					continue;
+				}
+				
+				// ignore the register offset / property paths since we need the whole operand's memory
+				bool bAlreadyExistedInList = false;
+				const int32 DebugOperandIndex = WorkPropertyIndices.Add(Operand.GetRegisterIndex(), &bAlreadyExistedInList).AsInteger();
+				if (bAlreadyExistedInList)
+				{
+					continue;
+				}
+				
+				OperandMap.Add(
+					FRigVMOperand(ERigVMMemoryType::Work, Operand.GetRegisterIndex()),
+					FRigVMOperand(ERigVMMemoryType::Debug, DebugOperandIndex));
+			}
+		}
+	}
+
+	const FRigVMMemoryStorageStruct& WorkMemory = GetDefaultWorkMemory();
+
+	// create a debug property for each operand
+	TArray<FRigVMPropertyDescription> PropertyDescriptions;
+	for (const int32& WorkPropertyIndex : WorkPropertyIndices)
+	{
+		 if (!WorkMemory.IsValidIndex(WorkPropertyIndex))
+		 {
+			 continue;
+		 }
+		
+		const FProperty* Property = WorkMemory.GetProperty(WorkPropertyIndex);
+		check(Property);
+
+		const FString CPPType = RigVMTypeUtils::ArrayTypeFromBaseType(RigVMTypeUtils::GetCPPTypeFromProperty(Property));
+		UObject* CPPTypeObject = RigVMTypeUtils::GetCPPTypeObjectFromProperty(Property);
+
+		// ignore any execute types
+		if (const UScriptStruct* ScriptStruct = Cast<UScriptStruct>(CPPTypeObject))
+		{
+			if (ScriptStruct->IsChildOf(FRigVMExecutePin::StaticStruct()))
+			{
+				continue;
+			}
+		}
+
+		FRigVMPropertyDescription PropertyDescription(Property->GetFName(), CPPType, CPPTypeObject, TEXT("()"), true);
+		PropertyDescriptions.Add(PropertyDescription);
+	}
+
+	FRigVMMemoryStorageStruct DebugMemory;
+	DebugMemory.AddProperties(PropertyDescriptions);
+
+	Context.DebugMemoryStorage = DebugMemory;
+	Context.DebuggedOperands.Reset();
+	Context.OperandToDebugRegister = OperandMap;
+#endif
+	return &Context.DebugMemoryStorage; 
+}
+
 void URigVM::GenerateMemoryType(FRigVMExtendedExecuteContext& Context, ERigVMMemoryType InMemoryType, const TArray<FRigVMPropertyDescription>* InProperties)
 {
 	GenerateDefaultMemoryType(InMemoryType, InProperties);
@@ -912,9 +979,7 @@ void URigVM::GenerateMemoryType(FRigVMExtendedExecuteContext& Context, ERigVMMem
 
 	case ERigVMMemoryType::Debug:
 	{
-#if WITH_EDITOR
-		Context.DebugMemoryStorage = DefaultDebugMemoryStorage;
-#endif
+		// debug memory is managed by the context of the host.
 		break;
 	}
 
@@ -946,7 +1011,7 @@ void URigVM::GenerateDefaultMemoryType(ERigVMMemoryType InMemoryType, const TArr
 #if WITH_EDITOR
 		case ERigVMMemoryType::Debug:
 		{
-			Memory = &DefaultDebugMemoryStorage;
+			Memory = &DefaultDebugMemoryStorage_DEPRECATED;
 			break;
 		}
 #endif
@@ -984,7 +1049,7 @@ FRigVMMemoryStorageStruct* URigVM::GetDefaultMemoryByType(ERigVMMemoryType InMem
 		case ERigVMMemoryType::Debug:
 		{
 #if WITH_EDITOR
-			MemoryStorage = &DefaultDebugMemoryStorage;
+			MemoryStorage = &DefaultDebugMemoryStorage_DEPRECATED;
 #endif
 			break;
 		}
@@ -1039,7 +1104,7 @@ void URigVM::ClearMemory_Internal()
 	LiteralMemoryStorage = FRigVMMemoryStorageStruct(ERigVMMemoryType::Invalid);
 	DefaultWorkMemoryStorage = FRigVMMemoryStorageStruct(ERigVMMemoryType::Invalid);
 #if WITH_EDITOR
-	DefaultDebugMemoryStorage = FRigVMMemoryStorageStruct(ERigVMMemoryType::Invalid);
+	DefaultDebugMemoryStorage_DEPRECATED = FRigVMMemoryStorageStruct(ERigVMMemoryType::Invalid);
 #endif
 
 	InvalidateCachedMemory_Internal();
@@ -1572,9 +1637,6 @@ bool URigVM::InitializeInstance(FRigVMExtendedExecuteContext& Context, bool bCop
 	{
 		Context.VMHash = GetVMHash();
 		Context.WorkMemoryStorage = DefaultWorkMemoryStorage;
-#if WITH_EDITOR
-		Context.DebugMemoryStorage = DefaultDebugMemoryStorage;
-#endif // WITH_EDITOR
 	}
 
 	return true;
diff --git a/Engine/Plugins/Runtime/RigVM/Source/RigVM/Private/RigVMCore/RigVMExecuteContext.cpp b/Engine/Plugins/Runtime/RigVM/Source/RigVM/Private/RigVMCore/RigVMExecuteContext.cpp
--- a/Engine/Plugins/Runtime/RigVM/Source/RigVM/Private/RigVMCore/RigVMExecuteContext.cpp
+++ b/Engine/Plugins/Runtime/RigVM/Source/RigVM/Private/RigVMCore/RigVMExecuteContext.cpp
@@ -194,3 +194,54 @@ FString FRigVMExtendedExecuteContext::GetVMPathName() const
 {
 	return VM ? VM->GetPathName() : TEXT("Unknown VM");
 }
+
+const FRigVMOperand* FRigVMExtendedExecuteContext::GetDebugOperandForOperand(const FRigVMOperand& InOperand) const
+{
+	const FRigVMOperand KeyOperand(InOperand.GetMemoryType(), InOperand.GetRegisterIndex()); // no register offset
+	if (DebuggedOperands.Num() == OperandToDebugRegister.Num() || DebuggedOperands.Contains(KeyOperand))
+	{
+		if(const FRigVMOperand* DebugOperandPtr = OperandToDebugRegister.Find(KeyOperand))
+		{
+			return DebugOperandPtr;
+		}
+	}
+
+	return nullptr;
+}
+
+void FRigVMExtendedExecuteContext::MarkOperandForDebugging(const FRigVMOperand& InOperand, bool bEnableDebugging)
+{
+	const FRigVMOperand KeyOperand(InOperand.GetMemoryType(), InOperand.GetRegisterIndex(), INDEX_NONE /* remove register offset */);
+	if (bEnableDebugging && OperandToDebugRegister.Contains(KeyOperand))
+	{
+		DebuggedOperands.Add(KeyOperand);
+	}
+	else
+	{
+		DebuggedOperands.Remove(KeyOperand);
+	}
+}
+
+void FRigVMExtendedExecuteContext::UnmarkOperandForDebugging(const FRigVMOperand& InOperand)
+{
+	MarkOperandForDebugging(InOperand, false);
+}
+
+void FRigVMExtendedExecuteContext::MarkAllOperandsForDebugging(bool bEnableDebugging)
+{
+	if (bEnableDebugging)
+	{
+		DebuggedOperands.Reset();
+		DebuggedOperands.Reserve(OperandToDebugRegister.Num());
+		OperandToDebugRegister.GetKeys(DebuggedOperands);
+	}
+	else
+	{
+		DebuggedOperands.Reset();
+	}
+}
+
+void FRigVMExtendedExecuteContext::UnmarkAllOperandsForDebugging()
+{
+	DebuggedOperands.Reset();
+}
diff --git a/Engine/Plugins/Runtime/RigVM/Source/RigVM/Private/RigVMCore/RigVMGraphFunctionDefinition.cpp b/Engine/Plugins/Runtime/RigVM/Source/RigVM/Private/RigVMCore/RigVMGraphFunctionDefinition.cpp
--- a/Engine/Plugins/Runtime/RigVM/Source/RigVM/Private/RigVMCore/RigVMGraphFunctionDefinition.cpp
+++ b/Engine/Plugins/Runtime/RigVM/Source/RigVM/Private/RigVMCore/RigVMGraphFunctionDefinition.cpp
@@ -132,11 +132,12 @@ FArchive& operator<<(FArchive& Ar, FRigVMFunctionCompilationData& Data)
 
 	// Serialize OperandToDebugRegisters
 	{
-		uint8 NumKeys = static_cast<uint8>(Data.OperandToDebugRegisters.Num());
-		Ar << NumKeys;
+		if (Ar.IsLoading() && Ar.CustomVer(FRigVMObjectVersion::GUID) < FRigVMObjectVersion::DebugOperandMappingSimplified)
+   		{
+			uint8 NumKeys = 0;
+			Ar << NumKeys;
 
-		if (Ar.IsLoading())
-		{
+			// this data is just loaded but then discarded.
 			for (int32 KeyIndex=0; KeyIndex<NumKeys; ++KeyIndex)
 			{
 				FRigVMOperand Key;
@@ -149,20 +150,6 @@ FArchive& operator<<(FArchive& Ar, FRigVMFunctionCompilationData& Data)
 				{
 					Ar << Values[ValueIndex];
 				}
-				Data.OperandToDebugRegisters.Add(Key, Values);
-			}
-		}
-		else
-		{
-			for (TPair<FRigVMOperand, TArray<FRigVMOperand>>& Pair : Data.OperandToDebugRegisters)
-			{
-				Ar << Pair.Key;
-				uint8 NumValues = static_cast<uint8>(Pair.Value.Num());
-				Ar << NumValues;
-				for (FRigVMOperand& Operand : Pair.Value)
-				{
-					Ar << Operand;
-				}
 			}
 		}
 	}
diff --git a/Engine/Plugins/Runtime/RigVM/Source/RigVM/Private/RigVMHost.cpp b/Engine/Plugins/Runtime/RigVM/Source/RigVM/Private/RigVMHost.cpp
--- a/Engine/Plugins/Runtime/RigVM/Source/RigVM/Private/RigVMHost.cpp
+++ b/Engine/Plugins/Runtime/RigVM/Source/RigVM/Private/RigVMHost.cpp
@@ -847,6 +847,25 @@ const FRigVMMemoryStorageStruct* URigVMHost::GetMemoryByType(ERigVMMemoryType In
 	return VM->GetMemoryByType(GetRigVMExtendedExecuteContext(), InMemoryType);
 }
 
+FRigVMMemoryStorageStruct* URigVMHost::GetDebugMemory(bool bCreateIfRequired)
+{
+	if (bCreateIfRequired)
+	{
+		if (GetMemoryByType(ERigVMMemoryType::Debug)->Num() == 0)
+		{
+			if (GetMemoryByType(ERigVMMemoryType::Work)->Num() > 0)
+			{
+				if (VM)
+				{
+					// create the debug memory for this instance
+					return VM->CreateDebugMemory(GetRigVMExtendedExecuteContext());
+				}
+			}
+		}
+	}
+	return GetMemoryByType(ERigVMMemoryType::Debug);
+}
+
 void URigVMHost::DrawIntoPDI(FPrimitiveDrawInterface* PDI, const FTransform& InTransform)
 {
 	for (const FRigVMDrawInstruction& Instruction : DrawInterface)
@@ -1325,7 +1344,7 @@ void URigVMHost::InstantiateVMFromCDO()
 			if(!VM->IsNativized())
 			{
 				ExtendedExecuteContext.WorkMemoryStorage = SourceVM->GetDefaultWorkMemory();
-				ExtendedExecuteContext.DebugMemoryStorage = SourceVM->GetDefaultDebugMemory();
+				ExtendedExecuteContext.DebugMemoryStorage = FRigVMMemoryStorageStruct(ERigVMMemoryType::Invalid);
 				ExtendedExecuteContext.VMHash = SourceVM->GetVMHash();
 			}
 		}
diff --git a/Engine/Plugins/Runtime/RigVM/Source/RigVM/Public/RigVMCore/RigVM.h b/Engine/Plugins/Runtime/RigVM/Source/RigVM/Public/RigVMCore/RigVM.h
--- a/Engine/Plugins/Runtime/RigVM/Source/RigVM/Public/RigVMCore/RigVM.h
+++ b/Engine/Plugins/Runtime/RigVM/Source/RigVM/Public/RigVMCore/RigVM.h
@@ -310,6 +310,9 @@ public:
 		return GetMemoryByType(Context, ERigVMMemoryType::Debug);
 	}
 
+	// creates the debug memory for a single instance of this VM
+	RIGVM_API FRigVMMemoryStorageStruct* CreateDebugMemory(FRigVMExtendedExecuteContext& Context);
+
 	UE_DEPRECATED(5.4, "Please, use GetInstanceMemoryArray with a Context param")
 	TArray<URigVMMemoryStorage*> GetLocalMemoryArray() { return TArray<URigVMMemoryStorage*>(); }
 
@@ -333,16 +336,6 @@ public:
 		return DefaultWorkMemoryStorage;
 	}
 
-	// The default non mutable reference debug memory (generated by the compiler and copied to the instances as initial state)
-	const FRigVMMemoryStorageStruct& GetDefaultDebugMemory() const
-	{
-		return DefaultDebugMemoryStorage;
-	}
-	FRigVMMemoryStorageStruct& GetDefaultDebugMemory()
-	{
-		return DefaultDebugMemoryStorage;
-	}
-
 	// Generates the Default memory and copy it to the Context if required by the type of memory
 	RIGVM_API virtual void GenerateMemoryType(FRigVMExtendedExecuteContext& Context, ERigVMMemoryType InMemoryType, const TArray<FRigVMPropertyDescription>* InProperties);
 
@@ -368,7 +361,7 @@ public:
 	FRigVMMemoryStorageStruct DefaultWorkMemoryStorage;
 
 	UPROPERTY()
-	FRigVMMemoryStorageStruct DefaultDebugMemoryStorage;
+	FRigVMMemoryStorageStruct DefaultDebugMemoryStorage_DEPRECATED;
 
 #if WITH_EDITORONLY_DATA
 	// Deprecated 5.4
@@ -1065,19 +1058,12 @@ private:
 		ERigVMMemoryType HandleMemoryType,
 		const FRigVMMemoryHandle& PreCachedHandle);
 
-
-
 	void CopyOperandForDebuggingIfNeeded(FRigVMExtendedExecuteContext& Context, const FRigVMOperand& InArg, const FRigVMMemoryHandle& InHandle)
 	{
 #if WITH_EDITOR
-		const FRigVMOperand KeyOperand(InArg.GetMemoryType(), InArg.GetRegisterIndex()); // no register offset
-		if(const TArray<FRigVMOperand>* DebugOperandsPtr = OperandToDebugRegisters.Find(KeyOperand))
+		if (const FRigVMOperand* DebugOperandPtr = Context.GetDebugOperandForOperand(InArg))
 		{
-			const TArray<FRigVMOperand>& DebugOperands = *DebugOperandsPtr;
-			for(const FRigVMOperand& DebugOperand : DebugOperands)
-			{
-				CopyOperandForDebuggingImpl(Context, InArg, InHandle, DebugOperand);
-			}
+			CopyOperandForDebuggingImpl(Context, InArg, InHandle, *DebugOperandPtr);
 		}
 #endif
 	}
@@ -1087,8 +1073,6 @@ private:
 	RIGVM_API FRigVMCopyOp GetCopyOpForOperands(const FRigVMOperand& InSource, const FRigVMOperand& InTarget);
 	RIGVM_API void RefreshExternalPropertyPaths();
 	
-	TMap<FRigVMOperand, TArray<FRigVMOperand>> OperandToDebugRegisters;
-
 protected:
 
 	struct FEntryExecuteGuard
@@ -1167,4 +1151,3 @@ protected:
 	friend struct FRigVMCompilerWorkData;
 	friend struct FRigVMCodeGenerator;
 };
-
diff --git a/Engine/Plugins/Runtime/RigVM/Source/RigVM/Public/RigVMCore/RigVMExecuteContext.h b/Engine/Plugins/Runtime/RigVM/Source/RigVM/Public/RigVMCore/RigVMExecuteContext.h
--- a/Engine/Plugins/Runtime/RigVM/Source/RigVM/Public/RigVMCore/RigVMExecuteContext.h
+++ b/Engine/Plugins/Runtime/RigVM/Source/RigVM/Public/RigVMCore/RigVMExecuteContext.h
@@ -43,7 +43,7 @@ struct FRigVMExecuteContext;
 struct FRigVMExtendedExecuteContext;
 struct FRigVMLogSettings;
 
-extern RIGVM_API TAutoConsoleVariable<bool> CVarRigVMReportAllMessages;
+extern UE_API TAutoConsoleVariable<bool> CVarRigVMReportAllMessages;
 
 USTRUCT()
 struct FRigVMSlice
@@ -379,13 +379,13 @@ struct FRigVMExecuteContext : public FRigVMExecutePin
 	/** The world this VM is running in */
 	const UWorld* GetWorld() const { return World; }
 
-	RIGVM_API void SetOwningObject(const UObject* InOwningObject);
+	UE_API void SetOwningObject(const UObject* InOwningObject);
 
-	RIGVM_API void SetOwningComponent(const USceneComponent* InOwningComponent);
+	UE_API void SetOwningComponent(const USceneComponent* InOwningComponent);
 
-	RIGVM_API void SetOwningActor(const AActor* InActor);
+	UE_API void SetOwningActor(const AActor* InActor);
 
-	RIGVM_API void SetWorld(const UWorld* InWorld);
+	UE_API void SetWorld(const UWorld* InWorld);
 
 	void SetToWorldSpaceTransform(const FTransform& InToWorldSpaceTransform) { ToWorldSpaceTransform = InToWorldSpaceTransform; }
 
@@ -498,7 +498,7 @@ struct FRigVMExecuteContext : public FRigVMExecutePin
 	/**
 	 * Serialize this type from another
 	 */
-	RIGVM_API bool SerializeFromMismatchedTag(const FPropertyTag& Tag, FStructuredArchive::FSlot Slot);
+	UE_API bool SerializeFromMismatchedTag(const FPropertyTag& Tag, FStructuredArchive::FSlot Slot);
 
 protected:
 
@@ -820,8 +820,8 @@ struct FRigVMExtendedExecuteContext
 		return 0;
 	}
 
-	RIGVM_API int32 GetInstructionIndex() const;
-	RIGVM_API FString GetVMPathName() const;
+	UE_API int32 GetInstructionIndex() const;
+	UE_API FString GetVMPathName() const;
 
 	uint32 VMHash = 0;
 
@@ -891,6 +891,18 @@ struct FRigVMExtendedExecuteContext
 	TArray<FRigVMInstructionSetExecuteState> LazyBranchExecuteState;
 	TArray<FRigVMExternalVariableRuntimeData> ExternalVariableRuntimeData;
 
+	UE_API const FRigVMOperand* GetDebugOperandForOperand(const FRigVMOperand& InOperand) const;
+	UE_API void MarkOperandForDebugging(const FRigVMOperand& InOperand, bool bEnableDebugging = true);
+	UE_API void UnmarkOperandForDebugging(const FRigVMOperand& InOperand);
+	UE_API void MarkAllOperandsForDebugging(bool bEnableDebugging = true);
+	UE_API void UnmarkAllOperandsForDebugging();
+
+	/** An optional operand to debug register map only for this context's memory */
+	TMap<FRigVMOperand, FRigVMOperand> OperandToDebugRegister;
+
+	/** A set indicating which operands should be debugged */
+	TSet<FRigVMOperand> DebuggedOperands;
+
 #if WITH_EDITOR
 
 	UE_DEPRECATED(5.8, "DebugInfo pointer has been deprecated, please use DebugInfoWeak.")
diff --git a/Engine/Plugins/Runtime/RigVM/Source/RigVM/Public/RigVMCore/RigVMGraphFunctionDefinition.h b/Engine/Plugins/Runtime/RigVM/Source/RigVM/Public/RigVMCore/RigVMGraphFunctionDefinition.h
--- a/Engine/Plugins/Runtime/RigVM/Source/RigVM/Public/RigVMCore/RigVMGraphFunctionDefinition.h
+++ b/Engine/Plugins/Runtime/RigVM/Source/RigVM/Public/RigVMCore/RigVMGraphFunctionDefinition.h
@@ -146,8 +146,6 @@ struct FRigVMFunctionCompilationData
 	UPROPERTY(Transient)
 	bool bEncounteredSurpressedErrors;
 
-	TMap<FRigVMOperand, TArray<FRigVMOperand>> OperandToDebugRegisters;
-
 	bool IsValid() const
 	{
 		return Hash != 0;
diff --git a/Engine/Plugins/Runtime/RigVM/Source/RigVM/Public/RigVMHost.h b/Engine/Plugins/Runtime/RigVM/Source/RigVM/Public/RigVMHost.h
--- a/Engine/Plugins/Runtime/RigVM/Source/RigVM/Public/RigVMHost.h
+++ b/Engine/Plugins/Runtime/RigVM/Source/RigVM/Public/RigVMHost.h
@@ -288,7 +288,7 @@ public:
 	const FRigVMMemoryStorageStruct* GetLiteralMemory() const { return GetMemoryByType(ERigVMMemoryType::Literal); }
 
 	// The instanced debug watch memory
-	FRigVMMemoryStorageStruct* GetDebugMemory() { return GetMemoryByType(ERigVMMemoryType::Debug); }
+	UE_API FRigVMMemoryStorageStruct* GetDebugMemory(bool bCreateIfRequired = false);
 	const FRigVMMemoryStorageStruct* GetDebugMemory() const { return GetMemoryByType(ERigVMMemoryType::Debug); }
 
 	DECLARE_EVENT_TwoParams(URigVM, FRigVMExecutedEvent, class URigVMHost*, const FName&);
@@ -582,6 +582,11 @@ public:
 	{
 		return IsInDebugMode() || VMRuntimeSettings.bEnableProfiling;
 	}
+
+	const FRigVMOperand* GetDebugOperandForOperand(const FRigVMOperand& InOperand) const
+	{
+		return GetRigVMExtendedExecuteContext().GetDebugOperandForOperand(InOperand);
+	}
 	
 	FRigVMDebugInfo& GetDebugInfo() { return DebugInfo.Get(); }
 	const FRigVMDebugInfo& GetDebugInfo() const { return DebugInfo.Get(); }
diff --git a/Engine/Plugins/Runtime/RigVM/Source/RigVM/Public/RigVMObjectVersion.h b/Engine/Plugins/Runtime/RigVM/Source/RigVM/Public/RigVMObjectVersion.h
--- a/Engine/Plugins/Runtime/RigVM/Source/RigVM/Public/RigVMObjectVersion.h
+++ b/Engine/Plugins/Runtime/RigVM/Source/RigVM/Public/RigVMObjectVersion.h
@@ -81,6 +81,9 @@ struct FRigVMObjectVersion
 		// Object archive is now storing the version container
 		ObjectArchiveVersionContainerSerialization,
 
+		// Debug operand mapping simplified and moved to context only
+		DebugOperandMappingSimplified,
+
 		// -----<new versions can be added above this line>-------------------------------------------------
 		VersionPlusOne,
 		LatestVersion = VersionPlusOne - 1
diff --git a/Engine/Plugins/Runtime/RigVM/Source/RigVMDeveloper/Private/RigVMCompiler/RigVMCompiler.cpp b/Engine/Plugins/Runtime/RigVM/Source/RigVMDeveloper/Private/RigVMCompiler/RigVMCompiler.cpp
--- a/Engine/Plugins/Runtime/RigVM/Source/RigVMDeveloper/Private/RigVMCompiler/RigVMCompiler.cpp
+++ b/Engine/Plugins/Runtime/RigVM/Source/RigVMDeveloper/Private/RigVMCompiler/RigVMCompiler.cpp
@@ -243,7 +243,7 @@ const FProperty* FRigVMCompilerWorkData::GetPropertyForOperand(const FRigVMOpera
 		}
 	case ERigVMMemoryType::Debug:
 		{
-			Property = GetPropertyFromMemory(VM->GetDefaultDebugMemory(), InOperand);
+			// debug memory is not supported.
 			break;
 		}
 	case ERigVMMemoryType::External:
@@ -1316,7 +1316,7 @@ bool URigVMCompiler::Compile(const FRigVMCompileSettings& InSettings, TArray<URi
 			
 			FRigVMASTProxy PinProxy = FRigVMASTProxy::MakeFromUObject(InterfacePin);
 			FRigVMVarExprAST* TempVarExpr = WorkData.AST->MakeExpr<FRigVMVarExprAST>(FRigVMExprAST::EType::Var, PinProxy);
-			FindOrAddRegister(TempVarExpr, WorkData, false);
+			FindOrAddRegister(TempVarExpr, WorkData);
 		}
 	}
 
@@ -1370,32 +1370,6 @@ bool URigVMCompiler::Compile(const FRigVMCompileSettings& InSettings, TArray<URi
 		}
 	}
 
-	if(WorkData.WatchedPins.Num() > 0)
-	{
-		for(int32 GraphIndex=0; GraphIndex<VisitedGraphs.Num(); GraphIndex++)
-		{
-			URigVMGraph* VisitedGraph = VisitedGraphs[GraphIndex];
-			for(URigVMNode* ModelNode : VisitedGraph->GetNodes())
-			{
-				for(URigVMPin* ModelPin : ModelNode->GetPins())
-				{
-					if(ModelPin->GetDirection() == ERigVMPinDirection::Input)
-					{
-						if(ModelPin->GetSourceLinks(true).Num() == 0)
-						{
-							continue;
-						}
-					}
-					FRigVMASTProxy PinProxy = FRigVMASTProxy::MakeFromUObject(ModelPin);
-					FRigVMVarExprAST TempVarExpr(FRigVMExprAST::EType::Var, PinProxy);
-					TempVarExpr.ParserPtr = WorkData.AST.Get();
-
-					FindOrAddRegister(&TempVarExpr, WorkData, true);
-				}
-			}
-		}
-	}
-
 	// now that we have determined the needed memory, let's
 	// setup properties as needed as well as property paths
 
@@ -1497,15 +1471,6 @@ bool URigVMCompiler::Compile(const FRigVMCompileSettings& InSettings, TArray<URi
 		WorkData.VM->GetByteCode().AlignByteCode();
 	}
 
-	// setup debug registers after all other registers have been created
-	if(Settings.EnablePinWatches)
-	{
-		for(URigVMPin* WatchedPin : WorkData.WatchedPins)
-		{
-			MarkDebugWatch(WorkData.Settings, true, WatchedPin, WorkData.VM, WorkData.PinPathToOperand, WorkData.AST);
-		}
-	}
-
 	// now that we have determined the needed memory, let's
 	// update the property paths once more
 	for(ERigVMMemoryType MemoryType : MemoryTypes)
@@ -1640,7 +1605,6 @@ bool URigVMCompiler::Compile(const FRigVMCompileSettings& InSettings, TArray<URi
 		}
 
 		OutFunctionCompilationData->Hash = GetTypeHash(*OutFunctionCompilationData);
-		OutFunctionCompilationData->OperandToDebugRegisters = WorkData.VM->OperandToDebugRegisters;
 	}
 
 	// Localize all of the registry information
@@ -3236,19 +3200,6 @@ bool URigVMCompiler::TraverseInlineFunction(const FRigVMInlineFunctionExprAST* I
 				// Store correspondence between Function operand and New operands
 				checkf(!FunctionOperandToNewOperand.Contains(OriginalOperand) || FunctionOperandToNewOperand.FindChecked(OriginalOperand) == *Operand, TEXT("hit a correspondence issue with a function operand, function %s in package %s"), *FunctionReferenceNode->GetFunctionIdentifier().GetLibraryNodePath(), *GetPackage()->GetPathName());
 				FunctionOperandToNewOperand.Add(OriginalOperand, *Operand);
-
-				// Also find the function debug operands and store correspondence to new debug operands
-				if (const TArray<FRigVMOperand>* FunctionDebugRegisters = FunctionCompilationData->OperandToDebugRegisters.Find(OriginalOperand))
-				{
-					for (const FRigVMOperand& FunctionDebugRegister : *FunctionDebugRegisters)
-					{
-						FRigVMCompilerWorkData::FFunctionRegisterData Data = {FunctionReferenceNode, FunctionDebugRegister.GetMemoryType(), FunctionDebugRegister.GetRegisterIndex()};
-						if (FRigVMOperand* NewDebugOperand = WorkData.FunctionRegisterToOperand.Find(Data))
-						{
-							FunctionOperandToNewOperand.Add(FunctionDebugRegister, *NewDebugOperand);
-						}
-					}
-				}
 			}
 
 			if (Instruction.OpCode == ERigVMOpCode::Execute)
@@ -3326,34 +3277,6 @@ bool URigVMCompiler::TraverseInlineFunction(const FRigVMInlineFunctionExprAST* I
 				Paths.Add(PinPathToOperand.Key);
 			}
 		}
-
-		// Fill out the operand to debug registers with the correct operands from the function
-		for (const TPair<FRigVMOperand, TArray<FRigVMOperand>>& Pair : FunctionCompilationData->OperandToDebugRegisters)
-		{
-			const FRigVMOperand& FunctionOperand = Pair.Key;
-			const TArray<FRigVMOperand>& FunctionDebugOperands = Pair.Value;
-			if (FRigVMOperand* NewOperand = FunctionOperandToNewOperand.Find(FunctionOperand))
-			{
-				TArray<FRigVMOperand> NewValues;
-				NewValues.Reserve(Pair.Value.Num());
-				for (const FRigVMOperand& FunctionDebugOperand : FunctionDebugOperands)
-				{
-					if (FRigVMOperand* NewDebugOperand = FunctionOperandToNewOperand.Find(FunctionDebugOperand))
-					{
-						NewValues.Add(*NewDebugOperand);
-						if (const TArray<FString>* Paths = OperandToPinPath.Find(*NewDebugOperand))
-						{
-							for (const FString& Path : *Paths)
-							{
-								FString PinPath = FString::Printf(TEXT("DebugWatch:%s"), *Path);
-								WorkData.PinPathToOperand->Add(PinPath, *NewDebugOperand);
-							}
-						}
-					}
-				}
-				WorkData.VM->OperandToDebugRegisters.Add(*NewOperand, NewValues);
-			}
-		}
 	}
 
 	return true;
@@ -3896,14 +3819,14 @@ void URigVMCompiler::AddCopyOperator(
 	AddCopyOperator(CopyOpInfo.Op, CopyOpInfo.AssignExpr, CopyOpInfo.SourceExpr, CopyOpInfo.TargetExpr, WorkData, bDelayCopyOperations);
 }
 
-FString URigVMCompiler::GetPinHashImpl(const URigVMPin* InPin, const FRigVMVarExprAST* InVarExpr, bool bIsDebugValue, const URigVMLibraryNode* FunctionCompiling, const FRigVMASTProxy& InPinProxy)
+FString URigVMCompiler::GetPinHashImpl(const URigVMPin* InPin, const FRigVMVarExprAST* InVarExpr, const URigVMLibraryNode* FunctionCompiling, const FRigVMASTProxy& InPinProxy)
 {
 	if (!InPin)
 	{
 		return FString();
 	}
 
-	FString Prefix = bIsDebugValue ? TEXT("DebugWatch:") : TEXT("");
+	FString Prefix;
 	FString Suffix;
 
 	if (InPin->IsExecuteContext())
@@ -3918,14 +3841,14 @@ FString URigVMCompiler::GetPinHashImpl(const URigVMPin* InPin, const FRigVMVarEx
 	bool bIsVariable = false;
 	bool bIsFunctionInterfacePin = false;
 
-	if (InVarExpr != nullptr && !bIsDebugValue)
+	if (InVarExpr != nullptr)
 	{
 		// for IO array pins we'll walk left and use that pin hash instead
 		if(const FRigVMVarExprAST* SourceVarExpr = GetSourceVarExpr(InVarExpr))
 		{
 			if(SourceVarExpr != InVarExpr)
 			{
-				return GetPinHash(SourceVarExpr->GetPin(), SourceVarExpr, bIsDebugValue, FunctionCompiling);
+				return GetPinHash(SourceVarExpr->GetPin(), SourceVarExpr, FunctionCompiling);
 			}
 		}
 
@@ -3956,7 +3879,7 @@ FString URigVMCompiler::GetPinHashImpl(const URigVMPin* InPin, const FRigVMVarEx
 	const bool bIsVariableClone = InVarExpr && InVarExpr->Name.ToString().EndsWith(TEXT("::Clone"));
 	if (VariableNode && !bIsVariableClone)
 	{
-		if (InPin->GetName() == TEXT("Value") && !bIsDebugValue)
+		if (InPin->GetName() == TEXT("Value"))
 		{
 			FName VariableName = VariableNode->GetVariableName();
 
@@ -4046,7 +3969,7 @@ FString URigVMCompiler::GetPinHashImpl(const URigVMPin* InPin, const FRigVMVarEx
 			const FString VariableName = InPin->IsBoundToVariable() ? InPin->GetBoundVariableName() : InPin->GetName();
 			return FString::Printf(TEXT("%sVariable::%s%s"), *Prefix, *VariableName, *Suffix);
 		}
-		if (InVarExpr && !bIsDebugValue)
+		if (InVarExpr)
 		{
 			const FRigVMASTProxy NodeProxy = InVarExpr->GetProxy().GetSibling(Node);
 			if (const FRigVMExprAST* NodeExpr = InVarExpr->GetParser()->GetExprForSubject(NodeProxy))
@@ -4119,10 +4042,10 @@ FString URigVMCompiler::GetPinHashImpl(const URigVMPin* InPin, const FRigVMVarEx
 	return FString::Printf(TEXT("%s%s%s"), *Prefix, *PinPath, *Suffix);
 }
 
-FString URigVMCompiler::GetPinHash(const URigVMPin* InPin, const FRigVMVarExprAST* InVarExpr, bool bIsDebugValue, const URigVMLibraryNode* FunctionCompiling, const FRigVMASTProxy& InPinProxy)
+FString URigVMCompiler::GetPinHash(const URigVMPin* InPin, const FRigVMVarExprAST* InVarExpr, const URigVMLibraryNode* FunctionCompiling, const FRigVMASTProxy& InPinProxy)
 {
-	const FString Hash = GetPinHashImpl(InPin, InVarExpr, bIsDebugValue, FunctionCompiling, InPinProxy);
-	if(!bIsDebugValue && FunctionCompiling == nullptr)
+	const FString Hash = GetPinHashImpl(InPin, InVarExpr, FunctionCompiling, InPinProxy);
+	if(FunctionCompiling == nullptr)
 	{
 		ensureMsgf(!Hash.Contains(TEXT("FunctionLibrary::")), TEXT("A library path should never be part of a pin hash %s."), *Hash);
 	}
@@ -4179,7 +4102,7 @@ const FRigVMVarExprAST* URigVMCompiler::GetSourceVarExpr(const FRigVMExprAST* In
 	return nullptr;
 }
 
-void URigVMCompiler::MarkDebugWatch(const FRigVMCompileSettings& InSettings, bool bRequired,
+void URigVMCompiler::MarkDebugWatch(FRigVMExtendedExecuteContext& Context, bool bRequired,
 	URigVMPin* InPin, URigVM* OutVM, TMap<FString, FRigVMOperand>* OutOperands,
 	TSharedPtr<FRigVMParserAST> InRuntimeAST)
 {
@@ -4189,16 +4112,22 @@ void URigVMCompiler::MarkDebugWatch(const FRigVMCompileSettings& InSettings, boo
 	check(InRuntimeAST.IsValid());
 	
 	URigVMPin* Pin = InPin->GetRootPin();
-	URigVMPin* SourcePin = Pin;
-	if(InSettings.ASTSettings.bFoldAssignments)
+	TArray<const FRigVMExprAST*> Expressions = InRuntimeAST->GetExpressionsForSubject(Pin);
+
+	if (Expressions.IsEmpty())
 	{
-		while(SourcePin->GetSourceLinks().Num() > 0)
+		FString PinHash = GetPinHashImpl(Pin, nullptr, nullptr);
+		if(const FRigVMOperand* Operand = OutOperands->Find(PinHash))
 		{
-			SourcePin = SourcePin->GetSourceLinks()[0]->GetSourcePin();
+			if (OutVM->GetDebugMemory(Context)->Num() == 0)
+			{
+				OutVM->CreateDebugMemory(Context);
+			}
+			Context.MarkOperandForDebugging(*Operand, bRequired);
 		}
+		return;
 	}
-		
-	TArray<const FRigVMExprAST*> Expressions = InRuntimeAST->GetExpressionsForSubject(SourcePin);
+	
 	TArray<FRigVMOperand> VisitedKeys;
 	for(int32 ExpressionIndex=0;ExpressionIndex<Expressions.Num();ExpressionIndex++)
 	{
@@ -4234,48 +4163,14 @@ void URigVMCompiler::MarkDebugWatch(const FRigVMCompileSettings& InSettings, boo
 			}
 		}
 
-		FString PinHash = GetPinHash(Pin, VarExpression, false, CurrentCompilationFunction);
-		if(!OutOperands->Contains(PinHash))
-		{
-			PinHash = GetPinHash(SourcePin, VarExpression, false, CurrentCompilationFunction);
-		}
-
+		FString PinHash = GetPinHashImpl(Pin, VarExpression, CurrentCompilationFunction);
 		if(const FRigVMOperand* Operand = OutOperands->Find(PinHash))
 		{
-			const FRigVMASTProxy PinProxy = FRigVMASTProxy::MakeFromUObject(Pin);
-			FRigVMVarExprAST TempVarExpr(FRigVMExprAST::EType::Var, PinProxy);
-			TempVarExpr.ParserPtr = InRuntimeAST.Get();
-
-			const FString DebugPinHash = GetPinHash(Pin, &TempVarExpr, true, CurrentCompilationFunction);
-			const FRigVMOperand* DebugOperand = OutOperands->Find(DebugPinHash);
-			if(DebugOperand)
+			if (OutVM->GetDebugMemory(Context)->Num() == 0)
 			{
-				if(DebugOperand->IsValid())
-				{
-					FRigVMOperand KeyOperand(Operand->GetMemoryType(), Operand->GetRegisterIndex()); // no register offset
-					if(bRequired)
-					{
-						if(!VisitedKeys.Contains(KeyOperand))
-						{
-							OutVM->OperandToDebugRegisters.FindOrAdd(KeyOperand).AddUnique(*DebugOperand);
-							VisitedKeys.Add(KeyOperand);
-						}
-					}
-					else
-					{
-						TArray<FRigVMOperand>* MappedOperands = OutVM->OperandToDebugRegisters.Find(KeyOperand);
-						if(MappedOperands)
-						{
-							MappedOperands->Remove(*DebugOperand);
-
-							if(MappedOperands->IsEmpty())
-							{
-								OutVM->OperandToDebugRegisters.Remove(KeyOperand);
-							}
-						}
-					}
-				}
+				OutVM->CreateDebugMemory(Context);
 			}
+			Context.MarkOperandForDebugging(*Operand, bRequired);
 		}
 	}
 }
@@ -4386,20 +4281,14 @@ int32 URigVMCompiler::GetElementSizeFromCPPType(const FString& InCPPType, UScrip
 	return 0;
 }
 
-FRigVMOperand URigVMCompiler::FindOrAddRegister(const FRigVMVarExprAST* InVarExpr, FRigVMCompilerWorkData& WorkData, bool bIsDebugValue)
+FRigVMOperand URigVMCompiler::FindOrAddRegister(const FRigVMVarExprAST* InVarExpr, FRigVMCompilerWorkData& WorkData)
 {
-	if(!bIsDebugValue)
-	{
-		InVarExpr = GetSourceVarExpr(InVarExpr);
-	}
-	
-	if (!bIsDebugValue)
+	InVarExpr = GetSourceVarExpr(InVarExpr);
+
+	FRigVMOperand const* ExistingOperandPtr = WorkData.ExprToOperand.Find(InVarExpr);
+	if (ExistingOperandPtr)
 	{
-		FRigVMOperand const* ExistingOperand = WorkData.ExprToOperand.Find(InVarExpr);
-		if (ExistingOperand)
-		{
-			return *ExistingOperand;
-		}
+		return *ExistingOperandPtr;
 	}
 
 	const URigVMPin::FPinOverrideMap& PinOverrides = InVarExpr->GetParser()->GetPinOverrides();
@@ -4414,17 +4303,17 @@ FRigVMOperand URigVMCompiler::FindOrAddRegister(const FRigVMVarExprAST* InVarExp
 	
 	FString CPPType = Pin->GetCPPType();
 	FString BaseCPPType = Pin->IsArray() ? Pin->GetArrayElementCppType() : CPPType;
-	FString Hash = GetPinHash(Pin, InVarExpr, bIsDebugValue, CurrentCompilationFunction);
+	FString Hash = GetPinHash(Pin, InVarExpr, CurrentCompilationFunction);
 	FRigVMOperand Operand;
 	FString RegisterKey = Hash;
 
 	bool bIsExecutePin = Pin->IsExecuteContext();
-	bool bIsLiteral = InVarExpr->GetType() == FRigVMExprAST::EType::Literal && !bIsDebugValue;
+	bool bIsLiteral = InVarExpr->GetType() == FRigVMExprAST::EType::Literal;
 	bool bIsVariable = Pin->IsRootPin() && (Pin->GetName() == URigVMVariableNode::ValueName) &&
 		InVarExpr->GetPin()->GetNode()->IsA<URigVMVariableNode>();
 
 	// external variables don't require to add any register.
-	if((bIsVariable || Pin->IsDefinedAsInputVariable()) && !bIsDebugValue)
+	if((bIsVariable || Pin->IsDefinedAsInputVariable()))
 	{
 		const TArray<FRigVMExternalVariableDef>& VMExternalVariableDefs = WorkData.VM->GetExternalVariableDefs();
 		for(int32 ExternalVariableIndex = 0; ExternalVariableIndex < VMExternalVariableDefs.Num(); ExternalVariableIndex++)
@@ -4442,12 +4331,11 @@ FRigVMOperand URigVMCompiler::FindOrAddRegister(const FRigVMVarExprAST* InVarExp
 	}
 
 	const ERigVMMemoryType MemoryType =
-		bIsLiteral ? ERigVMMemoryType::Literal:
-		(bIsDebugValue ? ERigVMMemoryType::Debug : ERigVMMemoryType::Work);
+		bIsLiteral ? ERigVMMemoryType::Literal : ERigVMMemoryType::Work;
 
 	TArray<FString> HashesWithSharedOperand;
 	
-	FRigVMOperand const* ExistingOperandPtr = WorkData.PinPathToOperand->Find(Hash);
+	ExistingOperandPtr = WorkData.PinPathToOperand->Find(Hash);
 	if (!ExistingOperandPtr)
 	{
 		if(WorkData.Settings.ASTSettings.bFoldAssignments) 		
@@ -4461,7 +4349,7 @@ FRigVMOperand URigVMCompiler::FindOrAddRegister(const FRigVMVarExprAST* InVarExp
 			{
 				if (const URigVMPin* VirtualPin = Cast<URigVMPin>(Proxy.GetSubject()))
 				{
-					const FString VirtualPinHash = GetPinHash(VirtualPin, InVarExpr, bIsDebugValue, CurrentCompilationFunction, Proxy);
+					const FString VirtualPinHash = GetPinHash(VirtualPin, InVarExpr, CurrentCompilationFunction, Proxy);
 					HashesWithSharedOperand.Add(VirtualPinHash);
 					if (Pin != VirtualPin)
 					{
@@ -4487,11 +4375,9 @@ FRigVMOperand URigVMCompiler::FindOrAddRegister(const FRigVMVarExprAST* InVarExp
 			WorkData.PinPathToOperand->Add(VirtualPinHash, ExistingOperand);
 		}
 		
-		if (!bIsDebugValue)
-		{
-			check(!WorkData.ExprToOperand.Contains(InVarExpr));
-			WorkData.ExprToOperand.Add(InVarExpr, ExistingOperand);
-		}
+		check(!WorkData.ExprToOperand.Contains(InVarExpr));
+		WorkData.ExprToOperand.Add(InVarExpr, ExistingOperand);
+		
 		return ExistingOperand;
 	}
 
@@ -4556,37 +4442,7 @@ FRigVMOperand URigVMCompiler::FindOrAddRegister(const FRigVMVarExprAST* InVarExp
 			ScriptStruct = GetScriptStructForCPPType(BaseCPPType);
 		}
 
-		if (!Operand.IsValid())
-		{
-			const int32 NumSlices = 1;
-			int32 Register = INDEX_NONE;
-
-			// debug watch register might already exists - look for them by name
-			if(bIsDebugValue)
-			{
-				Operand = WorkData.FindProperty(MemoryType, RegisterName);
-				if(Operand.IsValid())
-				{
-					FRigVMPropertyDescription Property = WorkData.GetProperty(Operand);
-					if(Property.IsValid())
-					{
-						if(ExistingOperandPtr == nullptr)
-						{
-							WorkData.PinPathToOperand->Add(Hash, Operand);
-						}
-						return Operand;
-					}
-				}
-			}
-		}
-
-		if(bIsDebugValue)
-		{
-			// debug values are always stored as arrays
-			CPPType = RigVMTypeUtils::ArrayTypeFromBaseType(CPPType);
-			JoinedDefaultValue = URigVMPin::GetDefaultValueForArray({ JoinedDefaultValue });
-		}
-		else if(Pin->GetDirection() == ERigVMPinDirection::Hidden)
+		if(Pin->GetDirection() == ERigVMPinDirection::Hidden)
 		{
 			bool bValidHiddenPin = false;
 			if(Pin->GetNode()->IsA<URigVMUnitNode>())
@@ -4655,11 +4511,8 @@ FRigVMOperand URigVMCompiler::FindOrAddRegister(const FRigVMVarExprAST* InVarExp
 		}
 	}
 	
-	if (!bIsDebugValue)
-	{
-		check(!WorkData.ExprToOperand.Contains(InVarExpr));
-		WorkData.ExprToOperand.Add(InVarExpr, Operand);
-	}
+	check(!WorkData.ExprToOperand.Contains(InVarExpr));
+	WorkData.ExprToOperand.Add(InVarExpr, Operand);
 
 	return Operand;
 }
diff --git a/Engine/Plugins/Runtime/RigVM/Source/RigVMDeveloper/Private/RigVMEditorAsset.cpp b/Engine/Plugins/Runtime/RigVM/Source/RigVMDeveloper/Private/RigVMEditorAsset.cpp
--- a/Engine/Plugins/Runtime/RigVM/Source/RigVMDeveloper/Private/RigVMEditorAsset.cpp
+++ b/Engine/Plugins/Runtime/RigVM/Source/RigVMDeveloper/Private/RigVMEditorAsset.cpp
@@ -1957,6 +1957,8 @@ void IRigVMEditorAssetInterface::RecompileVM()
 	RigGraphDisplaySettings.MaxMicroSecondsFrames.Reset();
 	RigGraphDisplaySettings.TotalMicroSecondsFrames.Reset();
 
+	WeakWatchedPins.Reset();
+
 	URigVM* VM = GetVM(true);
 	FRigVMExtendedExecuteContext* Context = GetRigVMExtendedExecuteContext();
 	if (VM && Context)
@@ -2276,6 +2278,63 @@ bool IRigVMEditorAssetInterface::IsMarkedDirtyDuringLoad() const
 	return bDirtyDuringLoad;
 }
 
+void IRigVMEditorAssetInterface::UpdateDebugMemoryOnHost(URigVMHost* InHost)
+{
+	// potentially set up the debug information based on the watched pins
+	TArray<URigVMGraph*> Graphs = GetAllModels();
+	if (!WeakWatchedPins.IsSet())
+	{
+		WeakWatchedPins = TArray<TWeakObjectPtr<URigVMPin>>();
+			
+		for (int32 GraphIndex = 0; GraphIndex < Graphs.Num(); GraphIndex++)
+		{
+			for (URigVMNode* Node : Graphs[GraphIndex]->GetNodes())
+			{
+				for (URigVMPin* Pin : Node->GetPins())
+				{
+					if (Pin->RequiresWatch())
+					{
+						WeakWatchedPins.GetValue().Add(Pin);
+					}
+
+					if (URigVMFunctionReferenceNode* FunctionReferenceNode = Cast<URigVMFunctionReferenceNode>(Node))
+					{
+						if (URigVMLibraryNode* LibraryNode = FunctionReferenceNode->LoadReferencedNode())
+						{
+							if (URigVMGraph* ContainedGraph = LibraryNode->GetContainedGraph())
+							{
+								Graphs.AddUnique(ContainedGraph);
+							}
+						}
+					}
+					else if (URigVMLibraryNode* LibraryNode = Cast<URigVMLibraryNode>(Node))
+					{
+						if (URigVMGraph* ContainedGraph = LibraryNode->GetContainedGraph())
+						{
+							Graphs.AddUnique(ContainedGraph);
+						}
+					}
+				}
+			}
+		}
+	}
+
+	if (Graphs[0]->GetRuntimeAST().IsValid())
+	{
+		URigVMCompiler* Compiler = URigVMCompiler::StaticClass()->GetDefaultObject<URigVMCompiler>();
+			
+		for (const TWeakObjectPtr<URigVMPin>& WeakPin : WeakWatchedPins.GetValue())
+		{
+			if (!WeakPin.IsValid())
+			{
+				continue;
+			}
+
+			Compiler->MarkDebugWatch(InHost->GetRigVMExtendedExecuteContext(), true, WeakPin.Get(), InHost->GetVM(), &PinToOperandMap, Graphs[0]->GetRuntimeAST());
+		}
+	}
+}
+
 void IRigVMEditorAssetInterface::RequestRigVMInit() const
 {
 	TArray<UObject*> Instances = GetArchetypeInstances(true, true);
@@ -2924,6 +2983,8 @@ void IRigVMEditorAssetInterface::SetObjectBeingDebugged(UObject* NewObject)
 	if(URigVMHost* NewRigBeingDebugged = Cast<URigVMHost>(NewObject))
 	{
 		NewRigBeingDebugged->bIsBeingDebugged = true;
+
+		UpdateDebugMemoryOnHost(NewRigBeingDebugged);
 	}
 	
 	ResetEarlyExitInstruction(true);
@@ -3533,20 +3594,13 @@ void IRigVMEditorAssetInterface::HandleModifiedEvent(ERigVMGraphNotifType InNoti
 						}
 						else
 						{
-							if(DebuggedHost->GetDebugMemory()->Num() == 0)
-							{
-								RequestAutoVMRecompilation();
-								(void)GetObject()->MarkPackageDirty();
-							}
-							else
-							{
-								Compiler->MarkDebugWatch(GetVMCompileSettings(), true, Pin, DebuggedHost->GetVM(), &PinToOperandMap, RuntimeAST);
-							}
+							(void)DebuggedHost->GetDebugMemory(true /* create if required */);
+							Compiler->MarkDebugWatch(DebuggedHost->GetRigVMExtendedExecuteContext(), true, Pin, DebuggedHost->GetVM(), &PinToOperandMap, RuntimeAST);
 						}
 					}
 					else
 					{
-						Compiler->MarkDebugWatch(GetVMCompileSettings(), false, Pin, DebuggedHost->GetVM(), &PinToOperandMap, RuntimeAST);
+						Compiler->MarkDebugWatch(DebuggedHost->GetRigVMExtendedExecuteContext(), false, Pin, DebuggedHost->GetVM(), &PinToOperandMap, RuntimeAST);
 					}
 				}
 				// break; fall through
@@ -3937,6 +3991,11 @@ void IRigVMEditorAssetInterface::InitializeArchetypeInstances()
 			}
 		}
 	}
+
+	if (URigVMHost* DebuggedHost = GetDebuggedRigVMHost())
+	{
+		UpdateDebugMemoryOnHost(DebuggedHost);
+	}
 }
 
 #if WITH_EDITOR
diff --git a/Engine/Plugins/Runtime/RigVM/Source/RigVMDeveloper/Public/RigVMCompiler/RigVMCompiler.h b/Engine/Plugins/Runtime/RigVM/Source/RigVMDeveloper/Public/RigVMCompiler/RigVMCompiler.h
--- a/Engine/Plugins/Runtime/RigVM/Source/RigVMDeveloper/Public/RigVMCompiler/RigVMCompiler.h
+++ b/Engine/Plugins/Runtime/RigVM/Source/RigVMDeveloper/Public/RigVMCompiler/RigVMCompiler.h
@@ -334,14 +334,14 @@ public:
 	TMap<FString, const FRigVMFunctionCompilationData*> CompiledFunctions;
 
 	static UE_API UScriptStruct* GetScriptStructForCPPType(const FString& InCPPType);
-	static UE_API FString GetPinHash(const URigVMPin* InPin, const FRigVMVarExprAST* InVarExpr, bool bIsDebugValue = false, const URigVMLibraryNode* FunctionCompiling = nullptr, const FRigVMASTProxy& InPinProxy = FRigVMASTProxy());
+	static UE_API FString GetPinHash(const URigVMPin* InPin, const FRigVMVarExprAST* InVarExpr, const URigVMLibraryNode* FunctionCompiling = nullptr, const FRigVMASTProxy& InPinProxy = FRigVMASTProxy());
 
 	// follows assignment expressions to find the source ref counted containers
 	// since ref counted containers are not copied for assignments.
 	// this is currently only used for arrays.
 	static UE_API const FRigVMVarExprAST* GetSourceVarExpr(const FRigVMExprAST* InExpr);
 
-	UE_API void MarkDebugWatch(const FRigVMCompileSettings& InSettings, bool bRequired, URigVMPin* InPin, URigVM* OutVM, TMap<FString, FRigVMOperand>* OutOperands, TSharedPtr<FRigVMParserAST> InRuntimeAST);
+	UE_API void MarkDebugWatch(FRigVMExtendedExecuteContext& Context, bool bRequired, URigVMPin* InPin, URigVM* OutVM, TMap<FString, FRigVMOperand>* OutOperands, TSharedPtr<FRigVMParserAST> InRuntimeAST);
 
 private:
 
@@ -351,7 +351,7 @@ private:
 	UE_API TArray<URigVMPin*> GetLinkedPins(URigVMPin* InPin, bool bInputs = true, bool bOutputs = true, bool bRecursive = true);
 	UE_API int32 GetElementSizeFromCPPType(const FString& InCPPType, UScriptStruct* InScriptStruct);
 
-	static UE_API FString GetPinHashImpl(const URigVMPin* InPin, const FRigVMVarExprAST* InVarExpr, bool bIsDebugValue = false, const URigVMLibraryNode* FunctionCompiling = nullptr, const FRigVMASTProxy& InPinProxy = FRigVMASTProxy());
+	static UE_API FString GetPinHashImpl(const URigVMPin* InPin, const FRigVMVarExprAST* InVarExpr, const URigVMLibraryNode* FunctionCompiling = nullptr, const FRigVMASTProxy& InPinProxy = FRigVMASTProxy());
 
 	UE_API bool TraverseExpression(const FRigVMExprAST* InExpr, FRigVMCompilerWorkData& WorkData);
 	UE_API bool TraverseChildren(const FRigVMExprAST* InExpr, FRigVMCompilerWorkData& WorkData);
@@ -381,7 +381,7 @@ private:
 		FRigVMCompilerWorkData& WorkData,
 		bool bDelayCopyOperations = true);
 
-	UE_API FRigVMOperand FindOrAddRegister(const FRigVMVarExprAST* InVarExpr, FRigVMCompilerWorkData& WorkData, bool bIsDebugValue = false);
+	UE_API FRigVMOperand FindOrAddRegister(const FRigVMVarExprAST* InVarExpr, FRigVMCompilerWorkData& WorkData);
 	UE_API const FRigVMCompilerWorkData::FRigVMASTProxyArray& FindProxiesWithSharedOperand(const FRigVMVarExprAST* InVarExpr, FRigVMCompilerWorkData& WorkData);
 
 	static UE_API FString GetPinNameWithDirectionPrefix(const URigVMPin* Pin);
@@ -415,6 +415,7 @@ private:
 	
 	friend class FRigVMCompilerImportErrorContext;
 	friend class FFunctionCompilationScope;
+	friend class SRigVMGraphNode;
 };
 
 class FFunctionCompilationScope
diff --git a/Engine/Plugins/Runtime/RigVM/Source/RigVMDeveloper/Public/RigVMEditorAsset.h b/Engine/Plugins/Runtime/RigVM/Source/RigVMDeveloper/Public/RigVMEditorAsset.h
--- a/Engine/Plugins/Runtime/RigVM/Source/RigVMDeveloper/Public/RigVMEditorAsset.h
+++ b/Engine/Plugins/Runtime/RigVM/Source/RigVMDeveloper/Public/RigVMEditorAsset.h
@@ -61,8 +61,8 @@ DECLARE_EVENT_OneParam(IRigVMEditorAssetInterface, FOnRigVMRequestInspectMemoryS
 DECLARE_EVENT_OneParam( IRigVMEditorAssetInterface, FOnRigVMAssetChangedEvent, class UObject * );
 DECLARE_EVENT_OneParam(IRigVMEditorAssetInterface, FOnRigVMSetObjectBeingDebugged, UObject* /*InDebugObj*/);
 
-USTRUCT()
-struct UE_API FRigVMPythonSettings
+USTRUCT(MinimalAPI)
+struct FRigVMPythonSettings
 {
 	GENERATED_BODY();
 
@@ -722,7 +722,9 @@ protected:
 	bool bSuspendPythonMessagesForRigVMClient;
 	bool bMarkBlueprintAsStructurallyModifiedPending;
 
+	UE_API void UpdateDebugMemoryOnHost(URigVMHost* InHost);
 	
+	TOptional<TArray<TWeakObjectPtr<URigVMPin>>> WeakWatchedPins;
 
 	static FCriticalSection QueuedCompilerMessageDelegatesMutex;
 	static TArray<FOnRigVMReportCompilerMessage::FDelegate> QueuedCompilerMessageDelegates;
diff --git a/Engine/Plugins/Runtime/RigVM/Source/RigVMEditor/Private/Widgets/SRigVMGraphNode.cpp b/Engine/Plugins/Runtime/RigVM/Source/RigVMEditor/Private/Widgets/SRigVMGraphNode.cpp
--- a/Engine/Plugins/Runtime/RigVM/Source/RigVMEditor/Private/Widgets/SRigVMGraphNode.cpp
+++ b/Engine/Plugins/Runtime/RigVM/Source/RigVMEditor/Private/Widgets/SRigVMGraphNode.cpp
@@ -1226,17 +1226,22 @@ void SRigVMGraphNode::GetNodeInfoPopups(FNodeInfoContext* Context, TArray<FGraph
 							}
 						}
 
-						if(DefaultValues.IsEmpty())
+						if(DefaultValues.IsEmpty() && ModelNode.IsValid())
 						{
-							FString PinHash = URigVMCompiler::GetPinHash(ModelPin, nullptr, true);
-							if (const FRigVMOperand* WatchOperand = RigAsset->PinToOperandMap.Find(PinHash))
+							// calling on the impl here directly to avoid the ensure / exception around functions
+							const FString PinHash = URigVMCompiler::GetPinHashImpl(ModelPin, nullptr);
+							
+							if (const FRigVMOperand* PinOperand = RigAsset->PinToOperandMap.Find(PinHash))
 							{
-								FRigVMMemoryStorageStruct* Memory = ActiveObject->GetDebugMemory();
-								// We mark PPF_ExternalEditor so that default values are also printed
-								const FString DebugValue = Memory->GetDataAsStringSafe(WatchOperand->GetRegisterIndex(), PPF_ExternalEditor | STRUCT_ExportTextItemNative);
-								if(!DebugValue.IsEmpty())
+								if (const FRigVMOperand* DebugOperand = ActiveObject->GetDebugOperandForOperand(*PinOperand))
 								{
-									DefaultValues = URigVMPin::SplitDefaultValue(DebugValue);
+									FRigVMMemoryStorageStruct* Memory = ActiveObject->GetDebugMemory();
+									// We mark PPF_ExternalEditor so that default values are also printed
+									const FString DebugValue = Memory->GetDataAsStringSafe(DebugOperand->GetRegisterIndex(), PPF_ExternalEditor | STRUCT_ExportTextItemNative);
+									if(!DebugValue.IsEmpty())
+									{
+										DefaultValues = URigVMPin::SplitDefaultValue(DebugValue);
+									}
 								}
 							}
 						}
