Control Rig: Prevent crash when an array iterator is trying to access an invalid index
#rnx

#Synced-CL 49727611

diff --git a/Engine/Plugins/Runtime/RigVM/Source/RigVM/Private/RigVMFunctions/RigVMDispatch_Array.cpp b/Engine/Plugins/Runtime/RigVM/Source/RigVM/Private/RigVMFunctions/RigVMDispatch_Array.cpp
--- a/Engine/Plugins/Runtime/RigVM/Source/RigVM/Private/RigVMFunctions/RigVMDispatch_Array.cpp
+++ b/Engine/Plugins/Runtime/RigVM/Source/RigVM/Private/RigVMFunctions/RigVMDispatch_Array.cpp
@@ -1640,22 +1640,37 @@ void FRigVMDispatch_ArrayIterator::Execute(FRigVMExtendedExecuteContext& InConte
 		Index = -1;;
 		Ratio = 0.f;
 	}
+	
+	// Do not rely on Count, the array might have changed inside the loop
+	const int32 Num = ArrayHelper.Num();
+	if (Num != Count)
+	{
+		static constexpr TCHAR InvalidSizeMessage[] = TEXT("ArrayIterator: The size of the array has changed inside the loop (initially %d, current %d).");
+		InContext.GetPublicData<>().Logf(EMessageSeverity::Warning, InvalidSizeMessage, Count, Num);
+	}
 
 	Index++;
-	if(Count > 0)
+	if(Num > 1)
 	{
-		Ratio = float(Index) / float(Count - 1);
+		Ratio = float(Index) / float(Num - 1);
+	}
+	else
+	{
+		Ratio = 0.f;
 	}
 
-	if(Index >= Count)
+	if(Index >= Num)
 	{
 		Block = FRigVMStruct::ControlFlowCompletedName;
 		return;
 	}
 	
-	uint8* TargetMemory = Handles[1].GetOutputData();
-	const uint8* SourceMemory = ArrayHelper.GetRawPtr(Index);
-	URigVMMemoryStorage::CopyProperty(Handles[1].GetProperty(), TargetMemory, ArrayProperty->Inner, SourceMemory);
+	if (ArrayHelper.IsValidIndex(Index))
+	{
+		uint8* TargetMemory = Handles[1].GetOutputData();
+		const uint8* SourceMemory = ArrayHelper.GetRawPtr(Index);
+		URigVMMemoryStorage::CopyProperty(Handles[1].GetProperty(), TargetMemory, ArrayProperty->Inner, SourceMemory);
+	}
 }
 
 #if WITH_EDITOR
