UAF: Crash fix for nullptr schema on collapsed nodes using the minimal rigvm environment

#Synced-CL 49758930

diff --git a/Engine/Plugins/Runtime/RigVM/Source/RigVMEditor/Private/Editor/RigVMGraphDetailCustomization.cpp b/Engine/Plugins/Runtime/RigVM/Source/RigVMEditor/Private/Editor/RigVMGraphDetailCustomization.cpp
--- a/Engine/Plugins/Runtime/RigVM/Source/RigVMEditor/Private/Editor/RigVMGraphDetailCustomization.cpp
+++ b/Engine/Plugins/Runtime/RigVM/Source/RigVMEditor/Private/Editor/RigVMGraphDetailCustomization.cpp
@@ -873,8 +873,20 @@ void FRigVMFunctionArgumentDefaultNode::GenerateChildContent(IDetailChildrenBuil
 	{
 		return;
 	}
-
-	MinimalEnvironment = MakeShared<FRigVMMinimalEnvironment>(LibraryNode->GetImplementingOuter<IRigVMAssetInterface>());
+	
+	if (IRigVMAssetInterface* AssetInterfaceOuter = LibraryNode->GetImplementingOuter<IRigVMAssetInterface>())
+	{
+		MinimalEnvironment = MakeShared<FRigVMMinimalEnvironment>(AssetInterfaceOuter);
+	}
+	else if (IRigVMClientHost* ClientHostOuter = LibraryNode->GetImplementingOuter<IRigVMClientHost>())
+	{
+		MinimalEnvironment = MakeShared<FRigVMMinimalEnvironment>(ClientHostOuter, nullptr);
+	}
+	else
+	{
+		MinimalEnvironment = MakeShared<FRigVMMinimalEnvironment>(nullptr);
+	}
+	
 	if (LibraryNode->GetTypedOuter<URigVMFunctionLibrary>())
 	{
 		MinimalEnvironment->SetFunctionReference(LibraryNode, LibraryNode->GetName());
diff --git a/Engine/Plugins/Runtime/RigVM/Source/RigVMEditor/Private/Editor/RigVMMinimalEnvironment.cpp b/Engine/Plugins/Runtime/RigVM/Source/RigVMEditor/Private/Editor/RigVMMinimalEnvironment.cpp
--- a/Engine/Plugins/Runtime/RigVM/Source/RigVMEditor/Private/Editor/RigVMMinimalEnvironment.cpp
+++ b/Engine/Plugins/Runtime/RigVM/Source/RigVMEditor/Private/Editor/RigVMMinimalEnvironment.cpp
@@ -17,14 +17,32 @@ FRigVMMinimalEnvironment::FRigVMMinimalEnvironment(const IRigVMAssetInterface* I
 		}
 		WeakRigVMClientHost = InAssetInterface->GetRigVMClientHost();
 	}
-	ModelController = TStrongObjectPtr<URigVMController>(NewObject<URigVMController>(Outer));
+	
+	InitGraphData(Outer);
+	
+	SetSchemata(InAssetInterface);
+}
+
+FRigVMMinimalEnvironment::FRigVMMinimalEnvironment(IRigVMClientHost* InRigVMClientHost, UObject* InAssetObject)
+{
+	UObject* Outer = InAssetObject ? InAssetObject : GetTransientPackage();
+	WeakRigVMClientHost = InRigVMClientHost;
+	
+	InitGraphData(Outer);
+	
+	SetSchemata(InRigVMClientHost, InAssetObject);
+}
+
+void FRigVMMinimalEnvironment::InitGraphData(UObject* InOuter)
+{
+	check(InOuter);
+	
+	ModelController = TStrongObjectPtr<URigVMController>(NewObject<URigVMController>(InOuter));
 	ModelGraph = TStrongObjectPtr<URigVMGraph>(NewObject<URigVMGraph>(ModelController.Get(), TEXT("PreviewGraph")));
 	ModelController->SetGraph(ModelGraph.Get());
-
+	
 	EdGraphClass = URigVMEdGraph::StaticClass();
 	EdGraphNodeClass = URigVMEdGraphNode::StaticClass();
-
-	SetSchemata(InAssetInterface);
 }
 
 URigVMGraph* FRigVMMinimalEnvironment::GetModel() const
@@ -217,19 +235,35 @@ void FRigVMMinimalEnvironment::SetSchemata(const IRigVMAssetInterface* InAssetIn
 	const TScriptInterface<IRigVMClientHost> ClientHost = InAssetInterface->GetRigVMClientHost();
 	check(ClientHost);
 
-	EdGraphClass = ClientHost->GetRigVMEdGraphClass();
-	EdGraphNodeClass = ClientHost->GetRigVMEdGraphNodeClass();
+	SetSchemata(ClientHost.GetInterface(), InAssetInterface->GetObject());
+}
 
+void FRigVMMinimalEnvironment::SetSchemata(const IRigVMClientHost* InRigVMClientHost, const UObject* InAssetObject)
+{
+	if (!InRigVMClientHost)
+	{
+		return;
+	}
+	
+	EdGraphClass = InRigVMClientHost->GetRigVMEdGraphClass();
+	EdGraphNodeClass = InRigVMClientHost->GetRigVMEdGraphNodeClass();
+	
 	if(!EdGraph.IsValid() || EdGraph->GetClass() != EdGraphClass)
 	{
 		EdGraph = TStrongObjectPtr<URigVMEdGraph>(NewObject<URigVMEdGraph>(ModelGraph.Get(), EdGraphClass));
 		EdGraph->ModelNodePath = ModelGraph->GetNodePath();
 	}
-
-	ModelController->SetSchemaClass(ClientHost->GetRigVMSchemaClass());
-	const UObject* AssetObject = InAssetInterface->GetObject();
-	check(AssetObject);
-	EdGraph->SetBlueprintClass(AssetObject->GetClass());
+	
+	ModelController->SetSchemaClass(InRigVMClientHost->GetRigVMSchemaClass());
+	ModelGraph->SetSchemaClass(InRigVMClientHost->GetRigVMSchemaClass());
+	if (InAssetObject)
+	{
+		EdGraph->SetBlueprintClass(InAssetObject->GetClass());
+	}
+	else
+	{
+		EdGraph->Schema = InRigVMClientHost->GetRigVMEdGraphSchemaClass();
+	}
 }
 
 FSimpleDelegate& FRigVMMinimalEnvironment::OnChanged()
diff --git a/Engine/Plugins/Runtime/RigVM/Source/RigVMEditor/Public/Editor/RigVMMinimalEnvironment.h b/Engine/Plugins/Runtime/RigVM/Source/RigVMEditor/Public/Editor/RigVMMinimalEnvironment.h
--- a/Engine/Plugins/Runtime/RigVM/Source/RigVMEditor/Public/Editor/RigVMMinimalEnvironment.h
+++ b/Engine/Plugins/Runtime/RigVM/Source/RigVMEditor/Public/Editor/RigVMMinimalEnvironment.h
@@ -16,6 +16,7 @@ class FRigVMMinimalEnvironment : public TSharedFromThis<FRigVMMinimalEnvironment
 {
 public:
 	UE_API FRigVMMinimalEnvironment(const IRigVMAssetInterface* InAssetInterface = nullptr);
+	UE_API FRigVMMinimalEnvironment(IRigVMClientHost* InRigVMClientHost, UObject* InAssetObject);
 
 	UE_API URigVMGraph* GetModel() const;
 	UE_API URigVMController* GetController() const;
@@ -24,6 +25,7 @@ public:
 	UE_API URigVMEdGraphNode* GetEdGraphNode() const;
 
 	UE_API void SetSchemata(const IRigVMAssetInterface* InAssetInterface);
+	UE_API void SetSchemata(const IRigVMClientHost* InRigVMClientHost, const UObject* InAssetObject);
 	UE_API void SetNode(URigVMNode* InModelNode);
 	UE_API URigVMNode* SetFunctionReference(URigVMLibraryNode* InFunctionDefinition, const FString& InNodeName);
 	UE_API URigVMNode* SetEmptyCollapseNode(URigVMCollapseNode* InSourceCollapseNode, const FString& InNodeName);
@@ -35,8 +37,8 @@ public:
 	UE_API void SynchronizeCollapseNode(URigVMCollapseNode* InSource);
 	
 private:
-
 	UE_API void HandleModified(ERigVMGraphNotifType InNotification, URigVMGraph* InGraph, UObject* InSubject);
+	void InitGraphData(UObject* InOuter);
 
 	TWeakInterfacePtr<IRigVMClientHost> WeakRigVMClientHost;
 	TStrongObjectPtr<URigVMGraph> ModelGraph;
