[RigVM]

Fixes deadlock when calling FRigVMTemplate::GetExecuteArguments which could recursively take a readlock causing deadlocks

#rnx

#Synced-CL 49730605

diff --git a/Engine/Plugins/Runtime/RigVM/Source/RigVM/Private/RigVMCore/RigVMTemplate.cpp b/Engine/Plugins/Runtime/RigVM/Source/RigVM/Private/RigVMCore/RigVMTemplate.cpp
--- a/Engine/Plugins/Runtime/RigVM/Source/RigVM/Private/RigVMCore/RigVMTemplate.cpp
+++ b/Engine/Plugins/Runtime/RigVM/Source/RigVM/Private/RigVMCore/RigVMTemplate.cpp
@@ -1628,7 +1628,7 @@ const TArray<FRigVMExecuteArgument>& FRigVMTemplate::GetExecuteArguments_NoLock(
 					for (TFieldIterator<FProperty> It(Struct, EFieldIterationFlags::None); It; ++It)
 					{
 						FRigVMTemplateArgument Argument = FRigVMTemplateArgument::Make_NoLock(*It, MutableRegistryHandle);
-						if(Argument.IsExecute())
+						if(Argument.IsExecute_NoLock(InRegistry))
 						{
 							ExecuteArguments.Emplace(Argument.Name, Argument.Direction, Argument.GetTypeIndex_NoLock(0, InRegistry));
 						}
