Control Rig: Optimize the retrieval of plugin prefixes
#rnx

#Synced-CL 49807743

diff --git a/Engine/Plugins/Animation/ControlRig/Source/ControlRig/Private/ControlRig.cpp b/Engine/Plugins/Animation/ControlRig/Source/ControlRig/Private/ControlRig.cpp
--- a/Engine/Plugins/Animation/ControlRig/Source/ControlRig/Private/ControlRig.cpp
+++ b/Engine/Plugins/Animation/ControlRig/Source/ControlRig/Private/ControlRig.cpp
@@ -3926,13 +3926,12 @@ void UControlRig::GenerateRequiredPluginsData(FRigVMExtendedExecuteContext& Cont
 
 	if (const URigHierarchy* Hierarchy = GetHierarchy())
 	{
-		TMap<FString,FString> AllPluginPrefixes;
 		for (int32 Index = 0; Index < Hierarchy->ElementComponents.Num(); Index++)
 		{
 			if(Hierarchy->ElementComponents[Index].IsValid())
 			{
 				const FRigBaseComponent* Component = Hierarchy->ElementComponents[Index].GetPtr<FRigBaseComponent>();
-				const FString PluginName = RigVMTypeUtils::GetPluginName(Component->GetScriptStruct(), AllPluginPrefixes);
+				const FString PluginName = RigVMTypeUtils::GetPluginName(Component->GetScriptStruct());
 				if (!PluginName.IsEmpty())
 				{
 					RequiredPlugins.AddUnique(PluginName);
diff --git a/Engine/Plugins/Runtime/RigVM/Source/RigVM/Private/RigVMCore/RigVM.cpp b/Engine/Plugins/Runtime/RigVM/Source/RigVM/Private/RigVMCore/RigVM.cpp
--- a/Engine/Plugins/Runtime/RigVM/Source/RigVM/Private/RigVMCore/RigVM.cpp
+++ b/Engine/Plugins/Runtime/RigVM/Source/RigVM/Private/RigVMCore/RigVM.cpp
@@ -3472,10 +3472,10 @@ TArray<const UObject*> URigVM::GetUserDefinedDependencies(const TArray<const FRi
 	return Dependencies;
 }
 
-void URigVM::GetRequiredPlugins(TArray<FString>& OutPlugins, TMap<FString, FString>& OutPluginPrefixes)
+void URigVM::GetRequiredPlugins(TArray<FString>& OutPlugins)
 {
-	GetDefaultMemoryByType(ERigVMMemoryType::Literal)->GetRequiredPlugins(OutPlugins, OutPluginPrefixes);
-	GetDefaultMemoryByType(ERigVMMemoryType::Work)->GetRequiredPlugins(OutPlugins, OutPluginPrefixes);
+	GetDefaultMemoryByType(ERigVMMemoryType::Literal)->GetRequiredPlugins(OutPlugins);
+	GetDefaultMemoryByType(ERigVMMemoryType::Work)->GetRequiredPlugins(OutPlugins);
 
 	ResolveFunctionsIfRequired();
 
@@ -3497,11 +3497,11 @@ void URigVM::GetRequiredPlugins(TArray<FString>& OutPlugins, TMap<FString, FStri
 		FString PluginName;
 		if (Function->Struct)
 		{
-			PluginName = RigVMTypeUtils::GetPluginName(Function->Struct, OutPluginPrefixes);
+			PluginName = RigVMTypeUtils::GetPluginName(Function->Struct);
 		}
 		else if (Function->Factory && Function->Factory->GetScriptStruct())
 		{
-			PluginName = RigVMTypeUtils::GetPluginName(Function->Factory->GetScriptStruct(), OutPluginPrefixes);
+			PluginName = RigVMTypeUtils::GetPluginName(Function->Factory->GetScriptStruct());
 		}
 		if (!PluginName.IsEmpty())
 		{
diff --git a/Engine/Plugins/Runtime/RigVM/Source/RigVM/Private/RigVMCore/RigVMMemoryStorageStruct.cpp b/Engine/Plugins/Runtime/RigVM/Source/RigVM/Private/RigVMCore/RigVMMemoryStorageStruct.cpp
--- a/Engine/Plugins/Runtime/RigVM/Source/RigVM/Private/RigVMCore/RigVMMemoryStorageStruct.cpp
+++ b/Engine/Plugins/Runtime/RigVM/Source/RigVM/Private/RigVMCore/RigVMMemoryStorageStruct.cpp
@@ -68,12 +68,11 @@ void FRigVMMemoryStorageStruct::GetUserDefinedDependencies(TArray<const UObject*
 	});
 }
 
-void FRigVMMemoryStorageStruct::GetRequiredPlugins(TArray<FString>& OutPlugins, TMap<FString,FString>& OutPluginPrefixes) const
+void FRigVMMemoryStorageStruct::GetRequiredPlugins(TArray<FString>& OutPlugins) const
 {
-	TMap<FString,FString> PluginPrefixes;
-	ForEachObjectDependency([&OutPlugins, &PluginPrefixes](const UObject* Dependency)
+	ForEachObjectDependency([&OutPlugins](const UObject* Dependency)
 	{
-		const FString PluginName = RigVMTypeUtils::GetPluginName(Dependency, PluginPrefixes);
+		const FString PluginName = RigVMTypeUtils::GetPluginName(Dependency);
 		if (!PluginName.IsEmpty())
 		{
 			OutPlugins.AddUnique(PluginName);
diff --git a/Engine/Plugins/Runtime/RigVM/Source/RigVM/Private/RigVMCore/RigVMPropertyBag.cpp b/Engine/Plugins/Runtime/RigVM/Source/RigVM/Private/RigVMCore/RigVMPropertyBag.cpp
--- a/Engine/Plugins/Runtime/RigVM/Source/RigVM/Private/RigVMCore/RigVMPropertyBag.cpp
+++ b/Engine/Plugins/Runtime/RigVM/Source/RigVM/Private/RigVMCore/RigVMPropertyBag.cpp
@@ -312,12 +312,11 @@ void FRigVMPropertyBag::GetUserDefinedDependencies(TArray<const UObject*>& OutDe
 	});
 }
 
-void FRigVMPropertyBag::GetRequiredPlugins(TArray<FString>& OutPlugins, TMap<FString,FString>& OutPluginPrefixes) const
+void FRigVMPropertyBag::GetRequiredPlugins(TArray<FString>& OutPlugins) const
 {
-	TMap<FString,FString> PluginPrefixes;
-	ForEachObjectDependency([&OutPlugins, &PluginPrefixes](const UObject* Dependency)
+	ForEachObjectDependency([&OutPlugins](const UObject* Dependency)
 	{
-		const FString PluginName = RigVMTypeUtils::GetPluginName(Dependency, PluginPrefixes);
+		const FString PluginName = RigVMTypeUtils::GetPluginName(Dependency);
 		if (!PluginName.IsEmpty())
 		{
 			OutPlugins.AddUnique(PluginName);
diff --git a/Engine/Plugins/Runtime/RigVM/Source/RigVM/Private/RigVMHost.cpp b/Engine/Plugins/Runtime/RigVM/Source/RigVM/Private/RigVMHost.cpp
--- a/Engine/Plugins/Runtime/RigVM/Source/RigVM/Private/RigVMHost.cpp
+++ b/Engine/Plugins/Runtime/RigVM/Source/RigVM/Private/RigVMHost.cpp
@@ -1198,11 +1198,10 @@ const TArray<FString>& URigVMHost::GetRequiredPlugins() const
 
 void URigVMHost::GenerateRequiredPluginsData(FRigVMExtendedExecuteContext& InContext)
 {
-	TMap<FString,FString> AllPluginPrefixes;
 	RequiredPlugins.Reset();
 	if (VM)
 	{
-		VM->GetRequiredPlugins(RequiredPlugins, AllPluginPrefixes);
+		VM->GetRequiredPlugins(RequiredPlugins);
 	}
 
 	const TArray<FRigVMExternalVariable> ExternalVariables = GetExternalVariables();
@@ -1212,7 +1211,7 @@ void URigVMHost::GenerateRequiredPluginsData(FRigVMExtendedExecuteContext& InCon
 		{
 			continue;
 		}
-		const FString PluginName = RigVMTypeUtils::GetPluginName(ExternalVariable.GetCPPTypeObject(), AllPluginPrefixes);
+		const FString PluginName = RigVMTypeUtils::GetPluginName(ExternalVariable.GetCPPTypeObject());
 		if (!PluginName.IsEmpty())
 		{
 			RequiredPlugins.AddUnique(PluginName);
diff --git a/Engine/Plugins/Runtime/RigVM/Source/RigVM/Private/RigVMRuntimeAsset.cpp b/Engine/Plugins/Runtime/RigVM/Source/RigVM/Private/RigVMRuntimeAsset.cpp
--- a/Engine/Plugins/Runtime/RigVM/Source/RigVM/Private/RigVMRuntimeAsset.cpp
+++ b/Engine/Plugins/Runtime/RigVM/Source/RigVM/Private/RigVMRuntimeAsset.cpp
@@ -63,11 +63,10 @@ void URigVMRuntimeAsset::Initialize(UClass* InEditorOnlyDataClass)
 
 void URigVMRuntimeAsset::GenerateRequiredPluginsData(FRigVMExtendedExecuteContext& InContext)
 {
-	TMap<FString,FString> AllPluginPrefixes;
 	RequiredPlugins.Reset();
 	if (VM)
 	{
-		VM->GetRequiredPlugins(RequiredPlugins, AllPluginPrefixes);
+		VM->GetRequiredPlugins(RequiredPlugins);
 	}
 
 	const TArray<FRigVMExternalVariable> ExternalVariables = GetExternalVariables();
@@ -77,7 +76,7 @@ void URigVMRuntimeAsset::GenerateRequiredPluginsData(FRigVMExtendedExecuteContex
 		{
 			continue;
 		}
-		const FString PluginName = RigVMTypeUtils::GetPluginName(ExternalVariable.GetCPPTypeObject(), AllPluginPrefixes);
+		const FString PluginName = RigVMTypeUtils::GetPluginName(ExternalVariable.GetCPPTypeObject());
 		if (!PluginName.IsEmpty())
 		{
 			RequiredPlugins.AddUnique(PluginName);
diff --git a/Engine/Plugins/Runtime/RigVM/Source/RigVM/Private/RigVMTypeUtils.cpp b/Engine/Plugins/Runtime/RigVM/Source/RigVM/Private/RigVMTypeUtils.cpp
--- a/Engine/Plugins/Runtime/RigVM/Source/RigVM/Private/RigVMTypeUtils.cpp
+++ b/Engine/Plugins/Runtime/RigVM/Source/RigVM/Private/RigVMTypeUtils.cpp
@@ -14,6 +14,10 @@
 #include "UObject/LinkerLoad.h"
 #include "Interfaces/IPluginManager.h"
 
+
+static FDelayedAutoRegisterHelper GRigVMPluginPrefixCacheInit(EDelayedRegisterRunPhase::EndOfEngineInit,
+	[]() { FRigVMPluginsPrefixCache::Get()->Refresh(); });
+
 namespace RigVMTypeUtils
 {
 	const TCHAR BoolType[] = TEXT("bool");
@@ -604,56 +608,14 @@ bool RigVMTypeUtils::AreCompatible(const FProperty* InSourceProperty, const FPro
 	return bCompatible;
 }
 
-FString RigVMTypeUtils::GetPluginName(const UObject* InObject, TMap<FString,FString>& InOutPluginPrefixes)
+FString RigVMTypeUtils::GetPluginName(const UObject* InObject)
 {
 	if(!InObject)
 	{
 		return FString();
 	}
 	
-	if (InOutPluginPrefixes.IsEmpty())
-	{
-		TArray<TSharedRef<IPlugin>> Plugins = IPluginManager::Get().GetDiscoveredPlugins();
-		for (const TSharedRef<IPlugin>& Plugin : Plugins)
-		{
-			const TArray<FString> Prefixes = GetPluginModulePrefixes(Plugin.Get());
-			for (const FString& Prefix : Prefixes)
-			{
-				InOutPluginPrefixes.Add(Prefix, Plugin->GetName());
-			}
-		}
-	}
-
-	const FString Path = FSoftObjectPath(InObject).ToString();
-	if (Path.IsEmpty())
-	{
-		return FString();
-	}
-	
-	static const FString ScriptPrefix = TEXT("/Script/");
-	if (Path.StartsWith(ScriptPrefix))
-	{
-		int32 PeriodPos = INDEX_NONE;
-		if (Path.FindChar(TEXT('.'), PeriodPos))
-		{
-			if (const FString* PluginName = InOutPluginPrefixes.Find(Path.Left(PeriodPos + 1)))
-			{
-				return *PluginName;
-			}
-		}
-	}
-	else
-	{
-		const int32 SecondSlashPos = Path.Find(TEXT("/"), ESearchCase::IgnoreCase, ESearchDir::FromStart, 1);
-		if (SecondSlashPos != INDEX_NONE)
-		{
-			if (const FString* PluginName = InOutPluginPrefixes.Find(Path.Left(SecondSlashPos + 1)))
-			{
-				return *PluginName;
-			}
-		}
-	}
-	return FString();
+	return FRigVMPluginsPrefixCache::Get()->GetPluginName(InObject);
 }
 
 TArray<FString> RigVMTypeUtils::GetPluginModulePrefixes(const IPlugin& InPlugin)
@@ -734,3 +696,103 @@ bool RigVMTypeUtils::CPPTypeFromExternalVariable(const FRigVMExternalVariable& I
 		
 	return true;
 }
+
+TSharedRef<FRigVMPluginsPrefixCache> FRigVMPluginsPrefixCache::Get()
+{
+	static TSharedRef<FRigVMPluginsPrefixCache> sPluginsCache = MakeShared<FRigVMPluginsPrefixCache>();
+	static bool bDelegatesRegistered = false;
+	if (!bDelegatesRegistered)
+	{
+		IPluginManager& PluginManager = IPluginManager::Get();
+		
+		PluginManager.OnNewPluginMounted().AddSP(sPluginsCache, &FRigVMPluginsPrefixCache::OnPluginMounted);
+		PluginManager.OnPluginUnmounted().AddSP(sPluginsCache, &FRigVMPluginsPrefixCache::OnPluginUnmounted);
+		
+		bDelegatesRegistered = true;
+	}
+	return sPluginsCache;
+}
+
+FString FRigVMPluginsPrefixCache::GetPluginName(const UObject* InObject) const
+{
+	if (!bInitialized)
+	{
+		const_cast<FRigVMPluginsPrefixCache*>(this)->Refresh();
+	}
+	
+	const FString Path = FSoftObjectPath(InObject).ToString();
+	if (Path.IsEmpty())
+	{
+		return FString();
+	}
+	
+	FReadScopeLock ReadLock(CacheLock);
+	
+	static const FString ScriptPrefix = TEXT("/Script/");
+	if (Path.StartsWith(ScriptPrefix))
+	{
+		int32 PeriodPos = INDEX_NONE;
+		if (Path.FindChar(TEXT('.'), PeriodPos))
+		{
+			if (const FString* PluginName = CachedPluginPrefixes.Find(Path.Left(PeriodPos + 1)))
+			{
+				return *PluginName;
+			}
+		}
+	}
+	else
+	{
+		const int32 SecondSlashPos = Path.Find(TEXT("/"), ESearchCase::IgnoreCase, ESearchDir::FromStart, 1);
+		if (SecondSlashPos != INDEX_NONE)
+		{
+			if (const FString* PluginName = CachedPluginPrefixes.Find(Path.Left(SecondSlashPos + 1)))
+			{
+				return *PluginName;
+			}
+		}
+	}
+	
+	return FString();
+}
+
+void FRigVMPluginsPrefixCache::Refresh()
+{
+	FWriteScopeLock WriteLock(CacheLock);
+	
+	CachedPluginPrefixes.Reset();
+	
+	TArray<TSharedRef<IPlugin>> Plugins = IPluginManager::Get().GetDiscoveredPlugins();
+	for (const TSharedRef<IPlugin>& Plugin : Plugins)
+	{
+		const TArray<FString> Prefixes = RigVMTypeUtils::GetPluginModulePrefixes(Plugin.Get());
+		for (const FString& Prefix : Prefixes)
+		{
+			CachedPluginPrefixes.Add(Prefix, Plugin->GetName());
+		}
+	}
+	
+	bInitialized = true;
+}
+
+FRigVMPluginsPrefixCache::FRigVMPluginsPrefixCache()
+	: bInitialized(false)
+{
+	
+}
+
+FRigVMPluginsPrefixCache::~FRigVMPluginsPrefixCache()
+{
+	IPluginManager& PluginManager = IPluginManager::Get();
+	PluginManager.OnNewPluginMounted().RemoveAll(this);
+	PluginManager.OnPluginUnmounted().RemoveAll(this);
+}
+
+void FRigVMPluginsPrefixCache::OnPluginMounted(IPlugin& InPlugin)
+{
+	Refresh();
+}
+
+void FRigVMPluginsPrefixCache::OnPluginUnmounted(IPlugin& InPlugin)
+{
+	Refresh();
+}
diff --git a/Engine/Plugins/Runtime/RigVM/Source/RigVM/Public/RigVMCore/RigVM.h b/Engine/Plugins/Runtime/RigVM/Source/RigVM/Public/RigVMCore/RigVM.h
--- a/Engine/Plugins/Runtime/RigVM/Source/RigVM/Public/RigVMCore/RigVM.h
+++ b/Engine/Plugins/Runtime/RigVM/Source/RigVM/Public/RigVMCore/RigVM.h
@@ -986,7 +986,7 @@ public:
 	RIGVM_API bool ValidateBytecode();
 	RIGVM_API void RefreshArgumentNameCaches();
 	RIGVM_API TArray<const UObject*> GetUserDefinedDependencies(const TArray<const FRigVMMemoryStorageStruct*> InMemory);
-	RIGVM_API void GetRequiredPlugins(TArray<FString>& OutPlugins, TMap<FString,FString>& OutPluginPrefixes);
+	RIGVM_API void GetRequiredPlugins(TArray<FString>& OutPlugins);
 
 private:
 
diff --git a/Engine/Plugins/Runtime/RigVM/Source/RigVM/Public/RigVMCore/RigVMMemoryStorageStruct.h b/Engine/Plugins/Runtime/RigVM/Source/RigVM/Public/RigVMCore/RigVMMemoryStorageStruct.h
--- a/Engine/Plugins/Runtime/RigVM/Source/RigVM/Public/RigVMCore/RigVMMemoryStorageStruct.h
+++ b/Engine/Plugins/Runtime/RigVM/Source/RigVM/Public/RigVMCore/RigVMMemoryStorageStruct.h
@@ -23,7 +23,7 @@ struct FRigVMMemoryStorageStruct : public FInstancedPropertyBag
 	UE_API bool Serialize(FArchive& Ar);
 	UE_API void AddStructReferencedObjects(FReferenceCollector& Collector);
 	UE_API void GetUserDefinedDependencies(TArray<const UObject*>& OutDependencies) const;
-	UE_API void GetRequiredPlugins(TArray<FString>& OutPlugins, TMap<FString,FString>& OutPluginPrefixes) const;
+	UE_API void GetRequiredPlugins(TArray<FString>& OutPlugins) const;
 
 	friend FArchive& operator<<(FArchive& Ar, FRigVMMemoryStorageStruct& P)
 	{
diff --git a/Engine/Plugins/Runtime/RigVM/Source/RigVM/Public/RigVMCore/RigVMPropertyBag.h b/Engine/Plugins/Runtime/RigVM/Source/RigVM/Public/RigVMCore/RigVMPropertyBag.h
--- a/Engine/Plugins/Runtime/RigVM/Source/RigVM/Public/RigVMCore/RigVMPropertyBag.h
+++ b/Engine/Plugins/Runtime/RigVM/Source/RigVM/Public/RigVMCore/RigVMPropertyBag.h
@@ -114,7 +114,7 @@ struct FRigVMPropertyBag : public FInstancedPropertyBag
 	UE_API bool Serialize(FArchive& Ar);
 	UE_API void AddStructReferencedObjects(FReferenceCollector& Collector);
 	UE_API void GetUserDefinedDependencies(TArray<const UObject*>& OutDependencies) const;
-	UE_API void GetRequiredPlugins(TArray<FString>& OutPlugins, TMap<FString,FString>& OutPluginPrefixes) const;
+	UE_API void GetRequiredPlugins(TArray<FString>& OutPlugins) const;
 
 	friend FArchive& operator<<(FArchive& Ar, FRigVMPropertyBag& P)
 	{
diff --git a/Engine/Plugins/Runtime/RigVM/Source/RigVM/Public/RigVMTypeUtils.h b/Engine/Plugins/Runtime/RigVM/Source/RigVM/Public/RigVMTypeUtils.h
--- a/Engine/Plugins/Runtime/RigVM/Source/RigVM/Public/RigVMTypeUtils.h
+++ b/Engine/Plugins/Runtime/RigVM/Source/RigVM/Public/RigVMTypeUtils.h
@@ -49,6 +49,33 @@ private:
 	TMap<FString, FSoftObjectPath> ObjectMap;
 };
 
+class FRigVMPluginsPrefixCache : public TSharedFromThis<FRigVMPluginsPrefixCache>
+{
+
+public:
+	UE_API static TSharedRef<FRigVMPluginsPrefixCache> Get();
+	UE_API virtual ~FRigVMPluginsPrefixCache();
+	
+	UE_API FString GetPluginName(const UObject* InObject) const;
+	UE_API void Refresh();
+
+private:
+	UE_API FRigVMPluginsPrefixCache();
+	
+	FRigVMPluginsPrefixCache(const FRigVMPluginsPrefixCache&) = delete;
+	FRigVMPluginsPrefixCache& operator=(const FRigVMPluginsPrefixCache&) = delete;
+	
+	UE_API void OnPluginMounted(IPlugin& InPlugin);
+	UE_API void OnPluginUnmounted(IPlugin& InPlugin);
+	
+	TMap<FString, FString> CachedPluginPrefixes;
+	bool bInitialized = false;
+	mutable FRWLock CacheLock;
+	
+	template <typename ObjectType, ESPMode Mode>
+	friend class SharedPointerInternals::TIntrusiveReferenceController;
+};
+
 namespace RigVMTypeUtils
 {
 	constexpr TCHAR TArrayPrefix[] = TEXT("TArray<");
@@ -278,7 +305,7 @@ namespace RigVMTypeUtils
 
 	RIGVM_API bool AreCompatible(const FProperty* InSourceProperty, const FProperty* InTargetProperty);
 
-	RIGVM_API FString GetPluginName(const UObject* InObject, TMap<FString,FString>& InOutPluginPrefixes);
+	RIGVM_API FString GetPluginName(const UObject* InObject);
 	RIGVM_API TArray<FString> GetPluginModulePrefixes(const IPlugin& InPlugin);
 
 	RIGVM_API bool CPPTypeFromExternalVariable(const FRigVMExternalVariable& InExternalVariable, FString& OutCPPType, UObject** OutCPPTypeObject);
