Fixed Control Rig function parameter type selection not allowing array types

#Synced-CL 49801191

diff --git a/Engine/Plugins/Runtime/RigVM/Source/RigVMEditor/Private/Editor/RigVMGraphDetailCustomization.cpp b/Engine/Plugins/Runtime/RigVM/Source/RigVMEditor/Private/Editor/RigVMGraphDetailCustomization.cpp
--- a/Engine/Plugins/Runtime/RigVM/Source/RigVMEditor/Private/Editor/RigVMGraphDetailCustomization.cpp
+++ b/Engine/Plugins/Runtime/RigVM/Source/RigVMEditor/Private/Editor/RigVMGraphDetailCustomization.cpp
@@ -810,6 +810,16 @@ void FRigVMFunctionArgumentLayout::OnContainerTypeSelectionChanged(TSharedPtr<TC
 		{
 			if (URigVMController* Controller = RigVMClientHost->GetController(LibraryNode->GetContainedGraph()))
 			{
+				FRigVMControllerCompileBracketScope CompileBracketScope(Controller);
+
+				FEdGraphPinType PinType = URigVMEdGraphNode::GetPinTypeForModelPin(PinPtr.Get());
+				const EPinContainerType NewPinContainerType = InContainerType->Get<0>();
+				if (NewPinContainerType != PinType.ContainerType)
+				{
+					PinType.ContainerType = NewPinContainerType;
+					PinInfoChanged(PinType);
+				}
+
 				(void)Controller->MarkExposedPinAsInputVariable(Pin->GetFName(), InContainerType->Get<1>());
 			}
 		}
