Disable RigVM function parameter customization on aggreage nodes

#Synced-CL 49918759

#Synced-CL 49918759

diff --git a/Engine/Plugins/Runtime/RigVM/Source/RigVMEditor/Private/Editor/RigVMGraphDetailCustomization.cpp b/Engine/Plugins/Runtime/RigVM/Source/RigVMEditor/Private/Editor/RigVMGraphDetailCustomization.cpp
--- a/Engine/Plugins/Runtime/RigVM/Source/RigVMEditor/Private/Editor/RigVMGraphDetailCustomization.cpp
+++ b/Engine/Plugins/Runtime/RigVM/Source/RigVMEditor/Private/Editor/RigVMGraphDetailCustomization.cpp
@@ -550,14 +550,11 @@ bool FRigVMFunctionArgumentLayout::IsPinEditingReadOnly(bool bIsEditingPinType/*
 	{
 		if(const URigVMPin* Pin = PinPtr.Get())
 		{
-			if(Pin->IsExecuteContext())
+			if(const URigVMNode* Node = Pin->GetNode())
 			{
-				if(const URigVMNode* Node = Pin->GetNode())
+				if(Node->IsA<URigVMAggregateNode>())
 				{
-					if(Node->IsA<URigVMAggregateNode>())
-					{
-						return true;
-					}
+					return true;
 				}
 			}
 		}
@@ -686,7 +683,19 @@ bool FRigVMFunctionArgumentLayout::CanChangeContainerType() const
 {
 	if (const URigVMPin* Pin = PinPtr.Get())
 	{
-		return !Pin->IsExecuteContext();
+		if (Pin->IsExecuteContext())
+		{
+			return false;
+		}
+
+		if (const URigVMNode* Node = Pin->GetNode())
+		{
+			if (Node->IsA<URigVMAggregateNode>())
+			{
+				return false;
+			}
+		}
+		return true;
 	}
 	return false;
 }
