UAF: Fix rename/move issues

Fixes reset to default not working correctly for object properties in the details panel.
Adds a UObject ptr to URigVM pins that reference a UObject - this ensures that any asset references that get renamed without a redirector (e.g. those that are not under source control) get correctly tracked. The ptr to the object is exclusive to the text representation (text is empty if pointer is set and vice versa)
Fixes the node title of template nodes not updating when a referenced asset is renamed.

#Synced-CL 50096253

diff --git a/Engine/Plugins/Experimental/UAF/UAF/Source/UAFEditor/Private/Common/AnimNextEdGraphNodeCustomization.cpp b/Engine/Plugins/Experimental/UAF/UAF/Source/UAFEditor/Private/Common/AnimNextEdGraphNodeCustomization.cpp
--- a/Engine/Plugins/Experimental/UAF/UAF/Source/UAFEditor/Private/Common/AnimNextEdGraphNodeCustomization.cpp
+++ b/Engine/Plugins/Experimental/UAF/UAF/Source/UAFEditor/Private/Common/AnimNextEdGraphNodeCustomization.cpp
@@ -513,7 +513,8 @@ void FAnimNextEdGraphNodeCustomization::PopulateCategory(IDetailLayoutBuilder& D
 
 		const TWeakPtr<FTraitStackDetailsData> WeakTraitData = TraitData.ToWeakPtr();
 		const TWeakPtr<IPropertyHandle> WeakPropertyHandle = InPropertyHandle.ToWeakPtr();
-		InPropertyHandle->SetOnPropertyValueChanged(FSimpleDelegate::CreateLambda([WeakTraitData, WeakPropertyHandle]()
+
+		InPropertyHandle->SetOnPropertyValueChangedWithData(TDelegate<void(const FPropertyChangedEvent&)>::CreateLambda([WeakTraitData, WeakPropertyHandle](const FPropertyChangedEvent& InPropertyChangedEvent)
 		{
 			TSharedPtr<IPropertyHandle> PinnedPropertyHandle = WeakPropertyHandle.Pin(); 
 			if(!PinnedPropertyHandle.IsValid())
@@ -531,7 +532,7 @@ void FAnimNextEdGraphNodeCustomization::PopulateCategory(IDetailLayoutBuilder& D
 			URigVMController* Controller = PinnedTraitData->EdGraphNodes[0]->GetController();
 			FRigVMControllerCompileBracketScope CompileScope(Controller);
 
-			PinnedPropertyHandle->EnumerateConstRawData([&PinnedTraitData, &PinnedPropertyHandle, Controller](const void* InRawData, const int32 InDataIndex, const int32 InNumDatas)
+			PinnedPropertyHandle->EnumerateConstRawData([&PinnedTraitData, &PinnedPropertyHandle, Controller, &InPropertyChangedEvent](const void* InRawData, const int32 InDataIndex, const int32 InNumDatas)
 			{
 				check(InNumDatas == PinnedTraitData->ScopedSharedDataInstances.Num());
 				check(InNumDatas == PinnedTraitData->EdGraphNodes.Num());
@@ -543,28 +544,37 @@ void FAnimNextEdGraphNodeCustomization::PopulateCategory(IDetailLayoutBuilder& D
 					return true;
 				}
 
-				const FProperty* Property = PinnedPropertyHandle->GetProperty();
-				const FString ValueStr = FRigVMStruct::ExportToFullyQualifiedText(Property, (uint8*)InRawData, false);
-
 				// Transform property path into pin path
-				FString PropertyPath(PinnedPropertyHandle->GetPropertyPath());
-
-				// Replace delimiter
-				FString PinPath = PropertyPath.Replace(TEXT("->"), TEXT("."));
-				// Replace array element delimiter
-				PinPath = PinPath.Replace(TEXT("["), TEXT("."));
-				PinPath = PinPath.Replace(TEXT("]."), TEXT("."));
-				PinPath = PinPath.Replace(TEXT("]"), TEXT(""));
-				// Replace struct root with trait root
-				const FString RootPropertyPath(PinnedTraitData->RootPropertyHandle->GetPropertyPath());
-				PinPath = PinPath.Replace(*RootPropertyPath, *PinnedTraitData->Name.ToString());
-				// Prefix node path
-				const FString NodePath = EdGraphNode->GetModelNodePath();
-				const bool bIsTrait = PinnedTraitData->Type == FCategoryDetailsData::EType::TraitStack;
-				PinPath = bIsTrait ? NodePath + TEXT(".") + PinnedTraitData->Name.ToString() + TEXT(".") + PinPath : NodePath + TEXT(".") + PinPath;
-				if(!Controller->SetPinDefaultValue(PinPath, ValueStr))
+				auto PropertyPathToPinPath = [](const TSharedPtr<IPropertyHandle>& InPropertyHandle, const TSharedPtr<FTraitStackDetailsData>& InTraitData, const UAnimNextEdGraphNode* InEdGraphNode)
+				{
+					// Transform property path into pin path
+					FString PropertyPath(InPropertyHandle->GetPropertyPath());
+
+					// Replace delimiter
+					FString PinPath = PropertyPath.Replace(TEXT("->"), TEXT("."));
+					// Replace array element delimiter
+					PinPath = PinPath.Replace(TEXT("["), TEXT("."));
+					PinPath = PinPath.Replace(TEXT("]."), TEXT("."));
+					PinPath = PinPath.Replace(TEXT("]"), TEXT(""));
+					// Replace struct root with trait root
+					const FString RootPropertyPath(InTraitData->RootPropertyHandle->GetPropertyPath());
+					PinPath = PinPath.Replace(*RootPropertyPath, *InTraitData->Name.ToString());
+					// Prefix node path
+					const FString NodePath = InEdGraphNode->GetModelNodePath();
+					const bool bIsTrait = InTraitData->Type == FCategoryDetailsData::EType::TraitStack;
+					return bIsTrait ? (NodePath + TEXT(".") + InTraitData->Name.ToString() + TEXT(".") + PinPath) : (NodePath + TEXT(".") + PinPath);
+				};
+
+				const FString PinPath = PropertyPathToPinPath(PinnedPropertyHandle, PinnedTraitData, EdGraphNode.Get());
+				if ((InPropertyChangedEvent.ChangeType & EPropertyChangeType::ResetToDefault) != 0)
+				{
+					Controller->ResetPinDefaultValue(PinPath);
+				}
+				else
 				{
-					UE_LOG(LogAnimation, Warning, TEXT("Could not set default value for pin '%s'"), *PinPath);
+					const FProperty* Property = PinnedPropertyHandle->GetProperty();
+					const FString ValueStr = FRigVMStruct::ExportToFullyQualifiedText(Property, (uint8*)InRawData, false);
+					Controller->SetPinDefaultValue(PinPath, ValueStr);
 				}
 
 				return true;
diff --git a/Engine/Plugins/Experimental/UAF/UAFAnimGraph/Source/UAFAnimGraphUncookedOnly/Internal/AnimNextTraitStackUnitNode.h b/Engine/Plugins/Experimental/UAF/UAFAnimGraph/Source/UAFAnimGraphUncookedOnly/Internal/AnimNextTraitStackUnitNode.h
--- a/Engine/Plugins/Experimental/UAF/UAFAnimGraph/Source/UAFAnimGraphUncookedOnly/Internal/AnimNextTraitStackUnitNode.h
+++ b/Engine/Plugins/Experimental/UAF/UAFAnimGraph/Source/UAFAnimGraphUncookedOnly/Internal/AnimNextTraitStackUnitNode.h
@@ -42,6 +42,7 @@ private:
 	friend UE::UAF::Editor::FAnimationGraphMenuExtensions;
 
 	void HandlePinDefaultValueChanged(UAnimNextController* InController, URigVMPin* InPin);
+	void HandleAssetRenamed(UAnimNextController* InController, const FAssetData& InAssetData, const FString& InOldName);
 
 	// The template that defines our behavior
 	UPROPERTY(VisibleAnywhere, Category = "Template")
diff --git a/Engine/Plugins/Experimental/UAF/UAFAnimGraph/Source/UAFAnimGraphUncookedOnly/Internal/Graph/AnimNextAnimationGraph_EditorData.h b/Engine/Plugins/Experimental/UAF/UAFAnimGraph/Source/UAFAnimGraphUncookedOnly/Internal/Graph/AnimNextAnimationGraph_EditorData.h
--- a/Engine/Plugins/Experimental/UAF/UAFAnimGraph/Source/UAFAnimGraphUncookedOnly/Internal/Graph/AnimNextAnimationGraph_EditorData.h
+++ b/Engine/Plugins/Experimental/UAF/UAFAnimGraph/Source/UAFAnimGraphUncookedOnly/Internal/Graph/AnimNextAnimationGraph_EditorData.h
@@ -66,6 +66,8 @@ class UUAFAnimGraph_EditorData : public UUAFSharedVariables_EditorData
 	friend class FAnimationAnimNextEditorTest_GraphManifest;
 
 public:
+	UE_API UUAFAnimGraph_EditorData();
+
 	/** Adds an animation graph to this asset */
 	UE_API UAnimNextAnimationGraphEntry* AddAnimationGraph(FName InName, bool bSetupUndoRedo = true, bool bPrintPythonCommand = true);
 
@@ -84,6 +86,9 @@ protected:
 	UE_API virtual void OnCompileJobFinished() override;
 	UE_API virtual UClass* GetRigVMEdGraphSchemaClass() const override;
 	UE_API virtual void HandleModifiedEvent(ERigVMGraphNotifType InNotifType, URigVMGraph* InGraph, UObject* InSubject) override;
+
+private:
+	void HandleAssetRenamed(const FAssetData& InAssetData, const FString& InOldName);
 };
 
 #undef UE_API
diff --git a/Engine/Plugins/Experimental/UAF/UAFAnimGraph/Source/UAFAnimGraphUncookedOnly/Private/AnimNextTraitStackUnitNode.cpp b/Engine/Plugins/Experimental/UAF/UAFAnimGraph/Source/UAFAnimGraphUncookedOnly/Private/AnimNextTraitStackUnitNode.cpp
--- a/Engine/Plugins/Experimental/UAF/UAFAnimGraph/Source/UAFAnimGraphUncookedOnly/Private/AnimNextTraitStackUnitNode.cpp
+++ b/Engine/Plugins/Experimental/UAF/UAFAnimGraph/Source/UAFAnimGraphUncookedOnly/Private/AnimNextTraitStackUnitNode.cpp
@@ -66,4 +66,12 @@ void UAnimNextTraitStackUnitNode::HandlePinDefaultValueChanged(UAnimNextControll
 	{
 		Template->GetDefaultObject<UUAFGraphNodeTemplate>()->HandlePinDefaultValueChanged(InController, InPin);
 	}
+}
+
+void UAnimNextTraitStackUnitNode::HandleAssetRenamed(UAnimNextController* InController, const FAssetData& InAssetData, const FString& InOldName)
+{
+	if (Template.Get() != nullptr)
+	{
+		Template->GetDefaultObject<UUAFGraphNodeTemplate>()->HandleAssetRenamed(InController, this, InAssetData, InOldName);
+	}
 }
\ No newline at end of file
diff --git a/Engine/Plugins/Experimental/UAF/UAFAnimGraph/Source/UAFAnimGraphUncookedOnly/Private/Graph/AnimNextAnimationGraph_EditorData.cpp b/Engine/Plugins/Experimental/UAF/UAFAnimGraph/Source/UAFAnimGraphUncookedOnly/Private/Graph/AnimNextAnimationGraph_EditorData.cpp
--- a/Engine/Plugins/Experimental/UAF/UAFAnimGraph/Source/UAFAnimGraphUncookedOnly/Private/Graph/AnimNextAnimationGraph_EditorData.cpp
+++ b/Engine/Plugins/Experimental/UAF/UAFAnimGraph/Source/UAFAnimGraphUncookedOnly/Private/Graph/AnimNextAnimationGraph_EditorData.cpp
@@ -11,6 +11,7 @@
 #include "Compilation/AnimNextGetVariableCompileContext.h"
 #include "Compilation/AnimNextProcessGraphCompileContext.h"
 #include "UncookedOnlyUtils.h"
+#include "AssetRegistry/AssetRegistryModule.h"
 #include "Curves/CurveFloat.h"
 #include "Entries/AnimNextAnimationGraphEntry.h"
 #include "Entries/AnimNextSharedVariablesEntry.h"
@@ -628,6 +629,15 @@ namespace UE::UAF::UncookedOnly::Private
 	}
 }
 
+UUAFAnimGraph_EditorData::UUAFAnimGraph_EditorData()
+{
+	if (!HasAnyFlags(RF_ClassDefaultObject))
+	{
+		IAssetRegistry& AssetRegistry = FModuleManager::Get().LoadModuleChecked<FAssetRegistryModule>(TEXT("AssetRegistry")).Get();
+		AssetRegistry.OnAssetRenamed().AddUObject(this, &UUAFAnimGraph_EditorData::HandleAssetRenamed);
+	}
+}
+
 void UUAFAnimGraph_EditorData::OnPreCompileAsset(FRigVMCompileSettings& InSettings)
 {
 	using namespace UE::UAF::UncookedOnly;
@@ -1028,4 +1038,21 @@ void UUAFAnimGraph_EditorData::HandleModifiedEvent(ERigVMGraphNotifType InNotifT
 	}
 }
 
+void UUAFAnimGraph_EditorData::HandleAssetRenamed(const FAssetData& InAssetData, const FString& InOldName)
+{
+	for (URigVMGraph* Graph : GetAllModels())
+	{
+		if(UAnimNextController* Controller = Cast<UAnimNextController>(GetController(Graph)))
+		{
+			for (URigVMNode* Node : Graph->GetNodes())
+			{
+				if (UAnimNextTraitStackUnitNode* TraitStack = Cast<UAnimNextTraitStackUnitNode>(Node))
+				{
+					TraitStack->HandleAssetRenamed(Controller, InAssetData, InOldName);
+				}
+			}
+		}
+	}
+}
+
 #undef LOCTEXT_NAMESPACE
diff --git a/Engine/Plugins/Experimental/UAF/UAFAnimGraph/Source/UAFAnimGraphUncookedOnly/Private/Templates/UAFGraphNodeTemplate_BlendSpacePlayer.h b/Engine/Plugins/Experimental/UAF/UAFAnimGraph/Source/UAFAnimGraphUncookedOnly/Private/Templates/UAFGraphNodeTemplate_BlendSpacePlayer.h
--- a/Engine/Plugins/Experimental/UAF/UAFAnimGraph/Source/UAFAnimGraphUncookedOnly/Private/Templates/UAFGraphNodeTemplate_BlendSpacePlayer.h
+++ b/Engine/Plugins/Experimental/UAF/UAFAnimGraph/Source/UAFAnimGraphUncookedOnly/Private/Templates/UAFGraphNodeTemplate_BlendSpacePlayer.h
@@ -45,17 +45,24 @@ class UUAFGraphNodeTemplate_BlendSpacePlayer : public UUAFGraphNodeTemplate
 	{
 		Super::HandleAssetDropped_Implementation(Controller, Node, Asset);
 
+		const bool bSetupUndoRedo = !GIsTransacting;
 		if (Asset)
 		{
-			Controller->OpenUndoBracket(LOCTEXT("ConfigureNodeOnDrop", "Configure Node On Drop").ToString());
+			if (bSetupUndoRedo)
+			{
+				Controller->OpenUndoBracket(LOCTEXT("ConfigureNodeOnDrop", "Configure Node On Drop").ToString());
+			}
 
-			Controller->SetNodeTitle(Node, FText::Format(LOCTEXT("NodeTitleFormat", "Play {0}"), FText::FromString(Asset->GetName())).ToString(), true, true, true);
+			Controller->SetNodeTitle(Node, FText::Format(LOCTEXT("NodeTitleFormat", "Play {0}"), FText::FromString(Asset->GetName())).ToString(), bSetupUndoRedo, true, true);
 
 			// Set looping to asset configuration by default
 			UBlendSpace* BlendSpace = CastChecked<UBlendSpace>(Asset);
-			Controller->SetPinDefaultValue(TEXT("AnimNextSequencePlayerTraitSharedData.bLoop"), LexToString(BlendSpace->bLoop), true, true, true, true, true);
+			Controller->SetPinDefaultValue(TEXT("AnimNextBlendSpacePlayerTraitSharedData.bLoop"), LexToString(BlendSpace->bLoop), true, bSetupUndoRedo, true, true, true);
 
-			Controller->CloseUndoBracket();
+			if (bSetupUndoRedo)
+			{
+				Controller->CloseUndoBracket();
+			}
 		}
 	}
 
@@ -63,14 +70,42 @@ class UUAFGraphNodeTemplate_BlendSpacePlayer : public UUAFGraphNodeTemplate
 	{
 		Super::HandlePinDefaultValueChanged_Implementation(Controller, Pin);
 
+		const bool bSetupUndoRedo = !GIsTransacting;
 		if (Pin->GetFName() == GET_MEMBER_NAME_CHECKED(UE::UAF::FBlendSpacePlayerData, BlendSpace))
 		{
-			Controller->OpenUndoBracket(LOCTEXT("SetNodeTitle", "Set Node Title").ToString());
+			if (bSetupUndoRedo)
+			{
+				Controller->OpenUndoBracket(LOCTEXT("SetNodeTitle", "Set Node Title").ToString());
+			}
 
 			FSoftObjectPath Path(Pin->GetDefaultValue());
-			Controller->SetNodeTitle(Pin->GetNode(), FText::Format(LOCTEXT("NodeTitleFormat", "Play {0}"), FText::FromString(Path.GetAssetName())).ToString(), true, true, true);
+			Controller->SetNodeTitle(Pin->GetNode(), FText::Format(LOCTEXT("NodeTitleFormat", "Play {0}"), FText::FromString(Path.GetAssetName())).ToString(), bSetupUndoRedo, true, true);
+
+			if (bSetupUndoRedo)
+			{
+				Controller->CloseUndoBracket();
+			}
+		}
+	}
 
-			Controller->CloseUndoBracket();
+	virtual void HandleAssetRenamed_Implementation(UAnimNextController* Controller, URigVMNode* Node, const FAssetData& AssetData, const FString& OldName) const
+	{
+		Super::HandleAssetRenamed_Implementation(Controller, Node, AssetData, OldName);
+
+		// It the asset is not loaded, it cant be our asset so we can safely early out
+		UObject* Object = AssetData.FastGetAsset();
+		if (Object == nullptr)
+		{
+			return;
+		}
+
+		// NOTE: Not making this undo-able as this is an automated process, not a user action on the node specifically
+
+		URigVMPin* Pin = Node->FindPin(GET_PIN_PATH_STRING_CHECKED(FAnimNextBlendSpacePlayerTraitSharedData, BlendSpace));
+		if (Pin && Object == Pin->GetDefaultValueObject())
+		{
+			FSoftObjectPath Path(Object);
+			Controller->SetNodeTitle(Node, FText::Format(LOCTEXT("NodeTitleFormat", "Play {0}"), FText::FromString(Path.GetAssetName())).ToString(), false, true, true);
 		}
 	}
 };
diff --git a/Engine/Plugins/Experimental/UAF/UAFAnimGraph/Source/UAFAnimGraphUncookedOnly/Private/Templates/UAFGraphNodeTemplate_SequencePlayer.h b/Engine/Plugins/Experimental/UAF/UAFAnimGraph/Source/UAFAnimGraphUncookedOnly/Private/Templates/UAFGraphNodeTemplate_SequencePlayer.h
--- a/Engine/Plugins/Experimental/UAF/UAFAnimGraph/Source/UAFAnimGraphUncookedOnly/Private/Templates/UAFGraphNodeTemplate_SequencePlayer.h
+++ b/Engine/Plugins/Experimental/UAF/UAFAnimGraph/Source/UAFAnimGraphUncookedOnly/Private/Templates/UAFGraphNodeTemplate_SequencePlayer.h
@@ -46,13 +46,20 @@ class UUAFGraphNodeTemplate_SequencePlayer : public UUAFGraphNodeTemplate
 	{
 		Super::HandleAssetDropped_Implementation(Controller, Node, Asset);
 
+		const bool bSetupUndoRedo = !GIsTransacting;
 		if (Asset)
 		{
-			Controller->OpenUndoBracket(LOCTEXT("ConfigureNodeOnDrop", "Configure Node On Drop").ToString());
+			if(bSetupUndoRedo)
+			{
+				Controller->OpenUndoBracket(LOCTEXT("ConfigureNodeOnDrop", "Configure Node On Drop").ToString());
+			}
 
-			Controller->SetNodeTitle(Node, FText::Format(LOCTEXT("NodeTitleFormat", "Play {0}"), FText::FromString(Asset->GetName())).ToString(), true, true, true);
+			Controller->SetNodeTitle(Node, FText::Format(LOCTEXT("NodeTitleFormat", "Play {0}"), FText::FromString(Asset->GetName())).ToString(), bSetupUndoRedo, true, true);
 
-			Controller->CloseUndoBracket();
+			if(bSetupUndoRedo)
+			{
+				Controller->CloseUndoBracket();
+			}
 		}
 	}
 
@@ -60,14 +67,42 @@ class UUAFGraphNodeTemplate_SequencePlayer : public UUAFGraphNodeTemplate
 	{
 		Super::HandlePinDefaultValueChanged_Implementation(Controller, Pin);
 
+		const bool bSetupUndoRedo = !GIsTransacting;
 		if (Pin->GetFName() == GET_MEMBER_NAME_CHECKED(UE::UAF::FSequencePlayerData, AnimSequence))
 		{
-			Controller->OpenUndoBracket(LOCTEXT("SetNodeTitle", "Set Node Title").ToString());
+			if (bSetupUndoRedo)
+			{
+				Controller->OpenUndoBracket(LOCTEXT("SetNodeTitle", "Set Node Title").ToString());
+			}
 
 			FSoftObjectPath Path(Pin->GetDefaultValue());
-			Controller->SetNodeTitle(Pin->GetNode(), FText::Format(LOCTEXT("NodeTitleFormat", "Play {0}"), FText::FromString(Path.GetAssetName())).ToString(), true, true, true);
+			Controller->SetNodeTitle(Pin->GetNode(), FText::Format(LOCTEXT("NodeTitleFormat", "Play {0}"), FText::FromString(Path.GetAssetName())).ToString(), bSetupUndoRedo, true, true);
+
+			if (bSetupUndoRedo)
+			{
+				Controller->CloseUndoBracket();
+			}
+		}
+	}
 
-			Controller->CloseUndoBracket();
+	virtual void HandleAssetRenamed_Implementation(UAnimNextController* Controller, URigVMNode* Node, const FAssetData& AssetData, const FString& OldName) const
+	{
+		Super::HandleAssetRenamed_Implementation(Controller, Node, AssetData, OldName);
+
+		// It the asset is not loaded, it cant be our asset so we can safely early out
+		UObject* Object = AssetData.FastGetAsset();
+		if (Object == nullptr)
+		{
+			return;
+		}
+
+		// NOTE: Not making this undo-able as this is an automated process, not a user action on the node specifically
+
+		URigVMPin* Pin = Node->FindPin(GET_PIN_PATH_STRING_CHECKED(FAnimNextSequencePlayerTraitSharedData, AnimSequence));
+		if (Pin && Object == Pin->GetDefaultValueObject())
+		{
+			FSoftObjectPath Path(Object);
+			Controller->SetNodeTitle(Node, FText::Format(LOCTEXT("NodeTitleFormat", "Play {0}"), FText::FromString(Path.GetAssetName())).ToString(), false, true, true);
 		}
 	}
 };
diff --git a/Engine/Plugins/Experimental/UAF/UAFAnimGraph/Source/UAFAnimGraphUncookedOnly/Public/Templates/UAFGraphNodeTemplate.h b/Engine/Plugins/Experimental/UAF/UAFAnimGraph/Source/UAFAnimGraphUncookedOnly/Public/Templates/UAFGraphNodeTemplate.h
--- a/Engine/Plugins/Experimental/UAF/UAFAnimGraph/Source/UAFAnimGraphUncookedOnly/Public/Templates/UAFGraphNodeTemplate.h
+++ b/Engine/Plugins/Experimental/UAF/UAFAnimGraph/Source/UAFAnimGraphUncookedOnly/Public/Templates/UAFGraphNodeTemplate.h
@@ -123,6 +123,11 @@ public:
 	UE_API void HandlePinDefaultValueChanged(UAnimNextController* Controller, URigVMPin* Pin) const;
 	virtual void HandlePinDefaultValueChanged_Implementation(UAnimNextController* Controller, URigVMPin* Pin) const {}
 
+	// Called in editor when an asset is renamed, e.g. to be used to reconfigure the nodes appearance
+	UFUNCTION(BlueprintNativeEvent, Category = "UAF|Template")
+	UE_API void HandleAssetRenamed(UAnimNextController* Controller, URigVMNode* Node, const FAssetData& AssetData, const FString& OldName) const;
+	virtual void HandleAssetRenamed_Implementation(UAnimNextController* Controller, URigVMNode* Node, const FAssetData& AssetData, const FString& OldName) const {}
+
 	// Get a layout that assigns all trait pins to the default category
 	UFUNCTION(BlueprintCallable, Category = "UAF|Template")
 	UE_API FRigVMNodeLayout GetDefaultCategoryLayout() const;
diff --git a/Engine/Plugins/Experimental/UAF/UAFMirroring/Source/UAFMirroringUncookedOnly/Private/UAFGraphNodeTemplate_Mirror.h b/Engine/Plugins/Experimental/UAF/UAFMirroring/Source/UAFMirroringUncookedOnly/Private/UAFGraphNodeTemplate_Mirror.h
--- a/Engine/Plugins/Experimental/UAF/UAFMirroring/Source/UAFMirroringUncookedOnly/Private/UAFGraphNodeTemplate_Mirror.h
+++ b/Engine/Plugins/Experimental/UAF/UAFMirroring/Source/UAFMirroringUncookedOnly/Private/UAFGraphNodeTemplate_Mirror.h
@@ -48,14 +48,21 @@ class UUAFGraphNodeTemplate_Mirror : public UUAFGraphNodeTemplate
 		Super::HandleAssetDropped_Implementation(Controller, Node, Asset);
 
 		// @todo: Wish we could use a subtitle instead but looks like the subtitle can't be set via the controller as its not instanced per not but rather queried per template.
-		
+
+		const bool bSetupUndoRedo = !GIsTransacting;
 		if (Asset)
 		{
-			Controller->OpenUndoBracket(LOCTEXT("ConfigureNodeOnDrop", "Configure Node On Drop").ToString());
+			if (bSetupUndoRedo)
+			{
+				Controller->OpenUndoBracket(LOCTEXT("ConfigureNodeOnDrop", "Configure Node On Drop").ToString());
+			}
 
-			Controller->SetNodeTitle(Node, FText::Format(LOCTEXT("NodeTitleFormat", "Mirror using {0}"), FText::FromString(Asset->GetName())).ToString(), true, true, true);
-			
-			Controller->CloseUndoBracket();
+			Controller->SetNodeTitle(Node, FText::Format(LOCTEXT("NodeTitleFormat", "Mirror using {0}"), FText::FromString(Asset->GetName())).ToString(), bSetupUndoRedo, true, true);
+
+			if (bSetupUndoRedo)
+			{
+				Controller->CloseUndoBracket();
+			}
 		}
 	}
 
@@ -63,14 +70,42 @@ class UUAFGraphNodeTemplate_Mirror : public UUAFGraphNodeTemplate
 	{
 		Super::HandlePinDefaultValueChanged_Implementation(Controller, Pin);
 
+		const bool bSetupUndoRedo = !GIsTransacting;
 		if (Pin->GetFName() == GET_MEMBER_NAME_CHECKED(UE::UAF::FMirroringTraitSetupParams, MirrorDataTable))
 		{
-			Controller->OpenUndoBracket(LOCTEXT("SetNodeTitle", "Set Node Title").ToString());
+			if (bSetupUndoRedo)
+			{
+				Controller->OpenUndoBracket(LOCTEXT("SetNodeTitle", "Set Node Title").ToString());
+			}
 
 			FSoftObjectPath Path(Pin->GetDefaultValue());
-			Controller->SetNodeTitle(Pin->GetNode(), FText::Format(LOCTEXT("NodeTitleFormat", "Mirror using {0}"), FText::FromString(Path.GetAssetName())).ToString(), true, true, true);
+			Controller->SetNodeTitle(Pin->GetNode(), FText::Format(LOCTEXT("NodeTitleFormat", "Mirror using {0}"), FText::FromString(Path.GetAssetName())).ToString(), bSetupUndoRedo, true, true);
 
-			Controller->CloseUndoBracket();
+			if (bSetupUndoRedo)
+			{
+				Controller->CloseUndoBracket();
+			}
+		}
+	}
+	
+	virtual void HandleAssetRenamed_Implementation(UAnimNextController* Controller, URigVMNode* Node, const FAssetData& AssetData, const FString& OldName) const
+	{
+		Super::HandleAssetRenamed_Implementation(Controller, Node, AssetData, OldName);
+
+		// It the asset is not loaded, it cant be our asset so we can safely early out
+		UObject* Object = AssetData.FastGetAsset();
+		if (Object == nullptr)
+		{
+			return;
+		}
+
+		// NOTE: Not making this undo-able as this is an automated process, not a user action on the node specifically
+
+		URigVMPin* Pin = Node->FindPin("MirroringTraitSharedData.Setup.MirrorDataTable");
+		if (Pin && Object == Pin->GetDefaultValueObject())
+		{
+			FSoftObjectPath Path(Object);
+			Controller->SetNodeTitle(Node, FText::Format(LOCTEXT("NodeTitleFormat", "Mirror using {0}"), FText::FromString(Path.GetAssetName())).ToString(), false, true, true);
 		}
 	}
 };
diff --git a/Engine/Plugins/Runtime/RigVM/Source/RigVMDeveloper/Private/RigVMModel/RigVMController.cpp b/Engine/Plugins/Runtime/RigVM/Source/RigVMDeveloper/Private/RigVMModel/RigVMController.cpp
--- a/Engine/Plugins/Runtime/RigVM/Source/RigVMDeveloper/Private/RigVMModel/RigVMController.cpp
+++ b/Engine/Plugins/Runtime/RigVM/Source/RigVMDeveloper/Private/RigVMModel/RigVMController.cpp
@@ -10350,9 +10350,18 @@ bool URigVMController::SetPinDefaultValue(URigVMPin* InPin, const FString& InDef
 			bSetPinDefaultValueSucceeded = true;
 		}
 		// no need to send notifications if not changing the value
-		if(InPin->DefaultValue != ClampedDefaultValue)
+		if(InPin->GetDefaultValue() != ClampedDefaultValue)
 		{
-			InPin->DefaultValue = ClampedDefaultValue;
+			if (InPin->IsUObject())
+			{
+				InPin->DefaultValueObject = RigVMTypeUtils::FindObjectFromCPPTypeObjectPath(ClampedDefaultValue);
+				InPin->DefaultValue.Empty();
+			}
+			else
+			{
+				InPin->DefaultValue = ClampedDefaultValue;
+				InPin->DefaultValueObject = nullptr;
+			}
 			InPin->IncrementVersion();
 			
 			// root trait pin store their default value in a separate property bag so that
@@ -10362,7 +10371,7 @@ bool URigVMController::SetPinDefaultValue(URigVMPin* InPin, const FString& InDef
 				FRigVMTraitDefaultValueStruct* DefaultValueStructPtr = InPin->GetNode()->TraitDefaultValues.Find(InPin->GetName());
 				if (ensure(DefaultValueStructPtr))
 				{
-					DefaultValueStructPtr->SetValue(InPin->DefaultValue);
+					DefaultValueStructPtr->SetValue(ClampedDefaultValue);
 				}
 			}
 			
@@ -10986,7 +10995,7 @@ bool URigVMController::ResetPinDefaultValue(URigVMPin* InPin, bool bSetupUndoRed
 
 	FString DefaultValue = RigVMNode->GetOriginalPinDefaultValue(InPin);
 	PostProcessDefaultValue(InPin, DefaultValue);
-	if(!DefaultValue.IsEmpty() || InPin->IsWildCard())
+	if(!DefaultValue.IsEmpty() || InPin->IsWildCard() || (InPin->IsUObject() && DefaultValue.IsEmpty()))
 	{
 		SetPinDefaultValue(InPin, DefaultValue, true, bSetupUndoRedo, false);
 		return true;
@@ -11930,7 +11939,17 @@ URigVMPin* URigVMController::InsertArrayPin(URigVMPin* ArrayPin, int32 InIndex,
 		{
 			Pin->DefaultValueType = GetDefaultValueType(Pin, DefaultValue);
 		}
-		Pin->DefaultValue = DefaultValue;
+
+		if (Pin->IsUObject())
+		{
+			Pin->DefaultValueObject = RigVMTypeUtils::FindObjectFromCPPTypeObjectPath(DefaultValue);
+			Pin->DefaultValue.Empty();
+		}
+		else
+		{
+			Pin->DefaultValue = DefaultValue;
+			Pin->DefaultValueObject = nullptr;
+		}
 	}
 
 	// set the array pin's default value type based on the resulting array pin list
@@ -17664,7 +17683,7 @@ FName URigVMController::AddTrait(URigVMNode* InNode, UScriptStruct* InTraitScrip
 	TraitPin->CPPTypeObject = InTraitScriptStruct;
 	TraitPin->CPPTypeObjectPath = *TraitPin->CPPTypeObject->GetPathName();
 	TraitPin->Direction = ERigVMPinDirection::Input;
-	TraitPin->DefaultValue = DefaultValue;
+	TraitPin->DefaultValueObject = TraitPin->CPPTypeObject;
 
 	{
 		const TGuardValue<bool> GuardNotifications(bSuspendNotifications, true);
@@ -18108,10 +18127,20 @@ void URigVMController::AddPinsForStruct(UStruct* InStruct, URigVMNode* InNode, U
 			{
 				FString DefaultValue = *DefaultValuePtr;
 				PostProcessDefaultValue(Pin, DefaultValue);
-				Pin->DefaultValue = DefaultValue;
-				if(Pin->CanProvideDefaultValue() && !Pin->DefaultValue.IsEmpty())
+				if (Pin->IsUObject())
 				{
-					Pin->DefaultValueType = GetDefaultValueType(Pin, Pin->DefaultValue);
+					Pin->DefaultValueObject = RigVMTypeUtils::FindObjectFromCPPTypeObjectPath(DefaultValue);
+					Pin->DefaultValue.Empty();
+				}
+				else
+				{
+					Pin->DefaultValue = DefaultValue;
+					Pin->DefaultValueObject = nullptr;
+				}
+				
+				if(Pin->CanProvideDefaultValue() && !DefaultValue.IsEmpty())
+				{
+					Pin->DefaultValueType = GetDefaultValueType(Pin, DefaultValue);
 				}
 			}
 
@@ -18202,10 +18231,20 @@ void URigVMController::AddPinsForArray(FArrayProperty* InArrayProperty, URigVMNo
 		if (!Pin->IsArray() && !Pin->IsStruct())
 		{
 			PostProcessDefaultValue(Pin, DefaultValue);
-			Pin->DefaultValue = DefaultValue;
-			if(Pin->CanProvideDefaultValue() && !Pin->DefaultValue.IsEmpty())
+			if (Pin->IsUObject())
 			{
-				Pin->DefaultValueType = GetDefaultValueType(Pin, Pin->DefaultValue);
+				Pin->DefaultValueObject = RigVMTypeUtils::FindObjectFromCPPTypeObjectPath(DefaultValue);
+				Pin->DefaultValue.Empty();
+			}
+			else
+			{
+				Pin->DefaultValue = DefaultValue;
+				Pin->DefaultValueObject = nullptr;
+			}
+			
+			if(Pin->CanProvideDefaultValue() && !DefaultValue.IsEmpty())
+			{
+				Pin->DefaultValueType = GetDefaultValueType(Pin, DefaultValue);
 			}
 		}
 	}
@@ -18505,6 +18544,7 @@ void URigVMController::ConfigurePinFromPin(URigVMPin* InOutPin, const URigVMPin*
 	InOutPin->CPPTypeObjectPath = InPin->CPPTypeObjectPath;
 	InOutPin->CPPTypeObject = InPin->CPPTypeObject;
 	InOutPin->DefaultValue = InPin->DefaultValue;
+	InOutPin->DefaultValueObject = InPin->DefaultValueObject;
 	InOutPin->bIsDynamicArray = InPin->bIsDynamicArray;
 	InOutPin->bIsLazy = InPin->bIsLazy;
 	InOutPin->bIsInputVariable = InPin->bIsInputVariable;
@@ -18530,7 +18570,16 @@ void URigVMController::ConfigurePinFromArgument(URigVMPin* InOutPin, const FRigV
 	InOutPin->CPPType = InArgument.CPPType.ToString();
 	InOutPin->CPPTypeObjectPath = *InArgument.CPPTypeObject.ToSoftObjectPath().ToString();
 	InOutPin->CPPTypeObject = InArgument.CPPTypeObject.Get();
-	InOutPin->DefaultValue = InArgument.DefaultValue;
+	if (InOutPin->IsUObject())
+	{
+		InOutPin->DefaultValueObject = RigVMTypeUtils::FindObjectFromCPPTypeObjectPath(InArgument.DefaultValue);
+		InOutPin->DefaultValue.Empty();
+	}
+	else
+	{
+		InOutPin->DefaultValueObject = nullptr;
+		InOutPin->DefaultValue = InArgument.DefaultValue;
+	}
 	InOutPin->bIsDynamicArray = InArgument.bIsArray;
 	if(bCopyDisplayName)
 	{
@@ -20214,7 +20263,14 @@ URigVMPin* URigVMController::CreatePinFromPinInfo(const FRigVMRegistry& InRegist
 	}
 
 	Pin->bIsExpanded = InPinInfo.bIsExpanded;
-	Pin->DefaultValue = InPinInfo.DefaultValue;
+	if (Pin->IsUObject())
+	{
+		Pin->DefaultValueObject = RigVMTypeUtils::FindObjectFromCPPTypeObjectPath(InPinInfo.DefaultValue);
+	}
+	else
+	{
+		Pin->DefaultValue = InPinInfo.DefaultValue;
+	}
 	Pin->DefaultValueType = InPinInfo.DefaultValueType;
 
 	// reuse expansion state and default value
@@ -20223,7 +20279,14 @@ URigVMPin* URigVMController::CreatePinFromPinInfo(const FRigVMRegistry& InRegist
 		if (PreviousPin->TypeIndex == InPinInfo.TypeIndex)
 		{
 			Pin->bIsExpanded = PreviousPin->bIsExpanded;
-			Pin->DefaultValue = PreviousPin->DefaultValue;
+			if (Pin->IsUObject())
+			{
+				Pin->DefaultValueObject = RigVMTypeUtils::FindObjectFromCPPTypeObjectPath(PreviousPin->DefaultValue);
+			}
+			else
+			{
+				Pin->DefaultValue = PreviousPin->DefaultValue;
+			}
 			Pin->DefaultValueType = PreviousPin->DefaultValueType;
 		}
 	}
@@ -20533,6 +20596,21 @@ void URigVMController::PostProcessDefaultValue(const URigVMPin* Pin, FString& Ou
 			OutDefaultValue = FName(NAME_None).ToString();
 		}
 	}
+	else if (OutDefaultValue.IsEmpty())
+	{
+		if (Pin->GetTypeIndex() == RigVMTypeUtils::TypeIndex::Bool)
+		{
+			OutDefaultValue = TEXT("false");
+		}
+		else if (Pin->GetTypeIndex() == RigVMTypeUtils::TypeIndex::Float || Pin->GetTypeIndex() == RigVMTypeUtils::TypeIndex::Double)
+		{
+			OutDefaultValue = TEXT("0.000000");
+		}
+		else if (Pin->GetTypeIndex() == RigVMTypeUtils::TypeIndex::Int32 || Pin->GetTypeIndex() == RigVMTypeUtils::TypeIndex::UInt32)
+		{
+			OutDefaultValue = TEXT("0");
+		}
+	}
 }
 
 void URigVMController::OverrideDefaultValueMember(const FString& InMemberName, const FString& InMemberValue, FString& InOutDefaultValue)
@@ -21805,13 +21883,22 @@ bool URigVMController::ChangePinType(URigVMPin* InPin, TRigVMTypeIndex InTypeInd
 		{
 			if(URigVMTemplateNode* TemplateNode = Cast<URigVMTemplateNode>(InPin->GetNode()))
 			{
-				InPin->DefaultValue = TemplateNode->GetInitialDefaultValueForPin(InPin->GetFName());
+				const FString InitialDefaultValue = TemplateNode->GetInitialDefaultValueForPin(InPin->GetFName());
+				if (InPin->IsUObject())
+				{
+					InPin->DefaultValueObject = RigVMTypeUtils::FindObjectFromCPPTypeObjectPath(InitialDefaultValue);
+				}
+				else
+				{
+					InPin->DefaultValue = InitialDefaultValue;
+				}
 			}
 		}
 
-		if(!InPin->DefaultValue.IsEmpty())
+		const FString DefaultValue = InPin->GetDefaultValue();
+		if(!DefaultValue.IsEmpty())
 		{
-			InPin->DefaultValueType = GetDefaultValueType(InPin, InPin->DefaultValue);
+			InPin->DefaultValueType = GetDefaultValueType(InPin, DefaultValue);
 		}
 	}
 
@@ -22226,14 +22313,15 @@ void URigVMController::AddNodePin(URigVMNode* InNode, URigVMPin* InPin)
 
 	if(InPin->DefaultValueType == ERigVMPinDefaultValueType::AutoDetect)
 	{
-		if(InPin->DefaultValue.IsEmpty())
+		const FString DefaultValue = InPin->GetDefaultValue();
+		if(DefaultValue.IsEmpty())
 		{
 			InPin->DefaultValueType = ERigVMPinDefaultValueType::Unset;
 		}
 		else
 		{
 			FRigVMDefaultValueTypeGuard _(this, ERigVMPinDefaultValueType::AutoDetect, true);
-			InPin->DefaultValueType = GetDefaultValueType(InPin, InPin->DefaultValue);
+			InPin->DefaultValueType = GetDefaultValueType(InPin, DefaultValue);
 		}
 	}
 
@@ -24340,11 +24428,12 @@ FRigVMClientPatchResult URigVMController::PatchPinDefaultValues()
 				{
 					if(Pin->GetDefaultValueType() == ERigVMPinDefaultValueType::AutoDetect)
 					{
-						if(!Pin->DefaultValue.IsEmpty())
+						const FString DefaultValue = Pin->GetDefaultValue();
+						if(!DefaultValue.IsEmpty())
 						{
 							// this will determine based on the delta if the default value
 							// is unchanged or overridden by the user
-							const ERigVMPinDefaultValueType DefaultValueType = GetDefaultValueType(Pin, Pin->DefaultValue);
+							const ERigVMPinDefaultValueType DefaultValueType = GetDefaultValueType(Pin, DefaultValue);
 							if(Pin->DefaultValueType != DefaultValueType)
 							{
 								Pin->DefaultValueType = DefaultValueType;
@@ -24365,7 +24454,8 @@ FRigVMClientPatchResult URigVMController::PatchPinDefaultValues()
 							{
 								if(Pin->IsValidDefaultValue(OriginalDefaultValue))
 								{
-									if(!Pin->DefaultValue.Equals(OriginalDefaultValue, ESearchCase::CaseSensitive))
+									const FString DefaultValue = Pin->GetDefaultValue();
+									if(!DefaultValue.Equals(OriginalDefaultValue, ESearchCase::CaseSensitive))
 									{
 										FRigVMDefaultValueTypeGuard OverrideGuard(this, ERigVMPinDefaultValueType::Unset, true);
 										if(SetPinDefaultValue(Pin, OriginalDefaultValue, true, false, false))
diff --git a/Engine/Plugins/Runtime/RigVM/Source/RigVMDeveloper/Private/RigVMModel/RigVMPin.cpp b/Engine/Plugins/Runtime/RigVM/Source/RigVMDeveloper/Private/RigVMModel/RigVMPin.cpp
--- a/Engine/Plugins/Runtime/RigVM/Source/RigVMDeveloper/Private/RigVMModel/RigVMPin.cpp
+++ b/Engine/Plugins/Runtime/RigVM/Source/RigVMDeveloper/Private/RigVMModel/RigVMPin.cpp
@@ -27,6 +27,7 @@
 #include "RigVMModel/Nodes/RigVMVariableNode.h"
 #include "RigVMStringUtils.h"
 #include "RigVMDeveloperTypeUtils.h"
+#include "UObject/FortniteMainBranchObjectVersion.h"
 
 #include UE_INLINE_GENERATED_CPP_BY_NAME(RigVMPin)
 
@@ -65,6 +66,34 @@ URigVMInjectionInfo::FWeakInfo URigVMInjectionInfo::GetWeakInfo() const
 const URigVMPin::FPinOverrideMap URigVMPin::EmptyPinOverrideMap;
 const URigVMPin::FPinOverride URigVMPin::EmptyPinOverride = URigVMPin::FPinOverride(FRigVMASTProxy(), EmptyPinOverrideMap);
 
+void URigVMPin::Serialize(FArchive& Ar)
+{
+	Ar.UsingCustomVersion(FFortniteMainBranchObjectVersion::GUID);
+
+	Super::Serialize(Ar);
+}
+
+void URigVMPin::PostLoad()
+{
+	Super::PostLoad();
+
+	if (GetLinkerCustomVersion(FFortniteMainBranchObjectVersion::GUID) < FFortniteMainBranchObjectVersion::RigVMPinDefaultValueObjectPtr)
+	{
+		if (IsUObject() && DefaultValueObject == nullptr)
+		{
+			// Patch object default value - value should already be hashed (if not fully loaded) as it is also referenced via the edgraph node
+			FSoftObjectPath SoftObjectPath(DefaultValue.TrimQuotes());
+			DefaultValueObject = SoftObjectPath.ResolveObject();
+			if (!SoftObjectPath.IsNull())
+			{
+				// Make sure we have resolve the object reference (validate the edgraph pin assumption above)
+				ensureAlways(DefaultValueObject != nullptr);
+			}
+			DefaultValue.Empty();
+		}
+	}
+}
+
 bool URigVMPin::SplitPinPathAtStart(const FString& InPinPath, FString& LeftMost, FString& Right)
 {
 	return RigVMStringUtils::SplitPinPathAtStart(InPinPath, LeftMost, Right);
@@ -107,9 +136,9 @@ FString URigVMPin::GetDefaultValueForArray(TConstArrayView<FString> DefaultValue
 	if (DefaultValues.Num())
 	{
 		Builder << DefaultValues[0];
-		for (const FString& DefaultValue : DefaultValues.Slice(1, DefaultValues.Num() - 1))
+		for (const FString& CurrentDefaultValue : DefaultValues.Slice(1, DefaultValues.Num() - 1))
 		{
-			Builder << TCHAR(',') << DefaultValue;
+			Builder << TCHAR(',') << CurrentDefaultValue;
 		}
 	}
 	Builder << TCHAR(')');
@@ -1079,6 +1108,12 @@ FString URigVMPin::GetDefaultValue(const URigVMPin::FPinOverride& InOverride, bo
 		
 		return DefaultValue.IsEmpty() ? EmptyStructDefaultValue : DefaultValue;
 	}
+	else if (IsUObject())
+	{
+		FSoftObjectPath Object(DefaultValueObject);
+		Cache = Object.ToString();
+		return Cache.GetValue();
+	}
 	else if (IsArrayElement() && DefaultValue.IsEmpty())
 	{
 		// array element cannot have an empty default value because when an array pin is
@@ -1096,6 +1131,11 @@ FString URigVMPin::GetDefaultValue(const URigVMPin::FPinOverride& InOverride, bo
 	return DefaultValue;
 }
 
+UObject* URigVMPin::GetDefaultValueObject() const
+{
+	return DefaultValueObject;
+}
+
 FString URigVMPin::GetDefaultValueStoredByUserInterface() const
 {
 	return GetDefaultValue(EmptyPinOverride, false);
diff --git a/Engine/Plugins/Runtime/RigVM/Source/RigVMDeveloper/Public/RigVMModel/RigVMPin.h b/Engine/Plugins/Runtime/RigVM/Source/RigVMDeveloper/Public/RigVMModel/RigVMPin.h
--- a/Engine/Plugins/Runtime/RigVM/Source/RigVMDeveloper/Public/RigVMModel/RigVMPin.h
+++ b/Engine/Plugins/Runtime/RigVM/Source/RigVMDeveloper/Public/RigVMModel/RigVMPin.h
@@ -105,6 +105,9 @@ class URigVMPin : public UObject
 	GENERATED_BODY()
 
 public:
+	// UObject interface
+	virtual void Serialize(FArchive& Ar) override;
+	virtual void PostLoad() override;
 
 	// A struct to store a pin override value
 	struct FPinOverrideValue 
@@ -335,6 +338,9 @@ public:
 	// Returns the default value with an additional override ma
 	UE_API FString GetDefaultValue(const FPinOverride& InOverride, bool bAdaptValueForPinType = true) const;
 
+	// Returns the default value object if this pin represents a UObject
+	UE_API UObject* GetDefaultValueObject() const;
+
 	// Returns the default value as stored by the user.
 	UE_API FString GetDefaultValueStoredByUserInterface() const;
 
@@ -608,6 +614,10 @@ private:
 	UPROPERTY()
 	FString DefaultValue;
 
+	/** If the default value for this pin should be an object, we store a pointer to it */
+	UPROPERTY()
+	TObjectPtr<UObject> DefaultValueObject;
+
 	UPROPERTY()
 	ERigVMPinDefaultValueType DefaultValueType;
 
diff --git a/Engine/Source/Runtime/Core/Public/UObject/FortniteMainBranchObjectVersions.inl b/Engine/Source/Runtime/Core/Public/UObject/FortniteMainBranchObjectVersions.inl
--- a/Engine/Source/Runtime/Core/Public/UObject/FortniteMainBranchObjectVersions.inl
+++ b/Engine/Source/Runtime/Core/Public/UObject/FortniteMainBranchObjectVersions.inl
@@ -757,5 +757,8 @@ FORTNITE_MAIN_VERSION(SplineComponentReparameterizeOnLoad, 247)
 // Deprecating support for external UPackage storing UAF asset entries
 FORTNITE_MAIN_VERSION(DeprecateUAFExternalPackages, 248)
 
+// RigVM object pins hold an object ptr
+FORTNITE_MAIN_VERSION(RigVMPinDefaultValueObjectPtr, 249)
+
 // -----<new versions can be added above this line>-------------------------------------------------
 #undef FORTNITE_MAIN_VERSION
