RigVM: Demote ensure on pin-PostLoad upgrade to a log

Content readily exists that has pin references to unloadable content (e.g. imported content in AnimSandbox)

#Synced-CL 50098746

diff --git a/Engine/Plugins/Runtime/RigVM/Source/RigVMDeveloper/Private/RigVMModel/RigVMPin.cpp b/Engine/Plugins/Runtime/RigVM/Source/RigVMDeveloper/Private/RigVMModel/RigVMPin.cpp
--- a/Engine/Plugins/Runtime/RigVM/Source/RigVMDeveloper/Private/RigVMModel/RigVMPin.cpp
+++ b/Engine/Plugins/Runtime/RigVM/Source/RigVMDeveloper/Private/RigVMModel/RigVMPin.cpp
@@ -86,8 +86,8 @@ void URigVMPin::PostLoad()
 			DefaultValueObject = SoftObjectPath.ResolveObject();
 			if (!SoftObjectPath.IsNull())
 			{
-				// Make sure we have resolve the object reference (validate the edgraph pin assumption above)
-				ensureAlways(DefaultValueObject != nullptr);
+				// Log if we couldnt resolve the object reference - this is usually fine as it means the reference didnt exist anyways
+				UE_CLOG(DefaultValueObject == nullptr, LogRigVMDeveloper, Display, TEXT("Could not resolve legacy text pin reference to '%s'"), *SoftObjectPath.ToString());
 			}
 			DefaultValue.Empty();
 		}
