RigVM: Address always true condition as per PVS

#rnx

#Synced-CL 50157960

diff --git a/Engine/Plugins/Runtime/RigVM/Source/RigVM/Private/RigVMTraceArchive.cpp b/Engine/Plugins/Runtime/RigVM/Source/RigVM/Private/RigVMTraceArchive.cpp
--- a/Engine/Plugins/Runtime/RigVM/Source/RigVM/Private/RigVMTraceArchive.cpp
+++ b/Engine/Plugins/Runtime/RigVM/Source/RigVM/Private/RigVMTraceArchive.cpp
@@ -106,7 +106,6 @@ bool FRigVMTraceArchive::Decompress()
 		CurrentHeader.CompressedSize = INDEX_NONE;
 		CurrentHeader.bIsCompressed = false;
 		
-		verify(CurrentHeader.CompressedSize == INDEX_NONE);
 		verify(CurrentHeader.UncompressedSize != INDEX_NONE);
 		verify(CurrentHeader.UncompressedSize == PayloadNum());
 		return true;
