RigVM: Fix build breaks on Linux RTTI

#rnx

#Synced-CL 50157960

diff --git a/Engine/Plugins/Runtime/RigVM/Source/RigVM/Public/RigVMHost.h b/Engine/Plugins/Runtime/RigVM/Source/RigVM/Public/RigVMHost.h
--- a/Engine/Plugins/Runtime/RigVM/Source/RigVM/Public/RigVMHost.h
+++ b/Engine/Plugins/Runtime/RigVM/Source/RigVM/Public/RigVMHost.h
@@ -231,7 +231,9 @@ protected:
 
 #if RIGVM_TRACE_ENABLED
 	UE_API virtual void TraceConstantData(FRigVMTraceArchiveWriter& InArchive) const;
+public:
 	UE_API virtual bool MatchesTracedConstantData(FRigVMTraceArchiveReader& InArchive) const;
+protected:
 	UE_API virtual void LoadTracedConstantData(FRigVMTraceArchiveReader& InArchive);
 	UE_API virtual void TraceExecuteData(FRigVMTraceArchiveWriter& InArchive) const;
 	UE_API virtual void LoadTracedExecuteData(FRigVMTraceArchiveReader& InArchive);
@@ -644,7 +646,6 @@ protected:
 	friend class URigVMBlueprintGeneratedClass;
 	friend class FRigVMCompileSettingsDetails;
 	friend class FRigVMTrace;
-	friend class FRewindDebuggerForRigVM;
 	friend class URigVM;
 	friend class URigVMRuntimeAsset;
 	friend class URigVMEditorAsset;
diff --git a/Engine/Plugins/Runtime/RigVM/Source/RigVMEditor/Public/Trace/RigVMTraceModule.h b/Engine/Plugins/Runtime/RigVM/Source/RigVMEditor/Public/Trace/RigVMTraceModule.h
--- a/Engine/Plugins/Runtime/RigVM/Source/RigVMEditor/Public/Trace/RigVMTraceModule.h
+++ b/Engine/Plugins/Runtime/RigVM/Source/RigVMEditor/Public/Trace/RigVMTraceModule.h
@@ -11,6 +11,8 @@
 #include "SEventTimelineView.h"
 #include "RigVMTrace.h"
 
+class URigVMHost;
+
 #if RIGVM_TRACE_ENABLED
 
 #define UE_API RIGVMEDITOR_API
