RigVM: mark deprecated memory classes as deprecated in the class flags

#Synced-CL 50157960

diff --git a/Engine/Plugins/Runtime/RigVM/Source/RigVM/Private/RigVMCore/RigVMMemoryStorage.cpp b/Engine/Plugins/Runtime/RigVM/Source/RigVM/Private/RigVMCore/RigVMMemoryStorage.cpp
--- a/Engine/Plugins/Runtime/RigVM/Source/RigVM/Private/RigVMCore/RigVMMemoryStorage.cpp
+++ b/Engine/Plugins/Runtime/RigVM/Source/RigVM/Private/RigVMCore/RigVMMemoryStorage.cpp
@@ -196,6 +196,11 @@ void URigVMMemoryStorageGeneratorClass::PostLoad()
 	// rebuild property list and property path list
 	RefreshLinkedProperties();
 	RefreshPropertyPaths();
+
+	// Mark the class as deprecated so it doesn't show up in asset searches
+	// and UI elements. This also avoids the call into
+	// GetStructPathName / GetClassPathName during asset registry searches.
+	ClassFlags |= CLASS_Deprecated;
 }
 
 const FString& URigVMMemoryStorageGeneratorClass::GetClassName(ERigVMMemoryType InMemoryType)
