Avoid RigVM crash when there is a function with a missing pin in the injection data

#Synced-CL 50205727

diff --git a/Engine/Plugins/Runtime/RigVM/Source/RigVMDeveloper/Private/RigVMModel/RigVMPin.cpp b/Engine/Plugins/Runtime/RigVM/Source/RigVMDeveloper/Private/RigVMModel/RigVMPin.cpp
--- a/Engine/Plugins/Runtime/RigVM/Source/RigVMDeveloper/Private/RigVMModel/RigVMPin.cpp
+++ b/Engine/Plugins/Runtime/RigVM/Source/RigVMDeveloper/Private/RigVMModel/RigVMPin.cpp
@@ -39,7 +39,11 @@ TAutoConsoleVariable<bool> CVarRigVMEnablePinOverrides(TEXT("RigVM.EnablePinOver
 
 URigVMGraph* URigVMInjectionInfo::GetGraph() const
 {
-	return GetPin()->GetGraph();
+	if(URigVMPin* Pin = GetPin())
+	{
+		return Pin->GetGraph();
+	}
+	return nullptr;
 }
 
 URigVMPin* URigVMInjectionInfo::GetPin() const
