RigVM: fix crash when debugging user defined structs

#rnx

#Synced-CL 50215688

diff --git a/Engine/Plugins/Runtime/RigVM/Source/RigVM/Private/RigVMCore/RigVM.cpp b/Engine/Plugins/Runtime/RigVM/Source/RigVM/Private/RigVMCore/RigVM.cpp
--- a/Engine/Plugins/Runtime/RigVM/Source/RigVM/Private/RigVMCore/RigVM.cpp
+++ b/Engine/Plugins/Runtime/RigVM/Source/RigVM/Private/RigVMCore/RigVM.cpp
@@ -939,7 +939,7 @@ FRigVMMemoryStorageStruct* URigVM::CreateDebugMemory(FRigVMExtendedExecuteContex
 		const FProperty* Property = WorkMemory.GetProperty(WorkPropertyIndex);
 		check(Property);
 
-		const FString CPPType = RigVMTypeUtils::ArrayTypeFromBaseType(RigVMTypeUtils::GetCPPTypeFromProperty(Property));
+		FString PropertyCPPType = RigVMTypeUtils::GetCPPTypeFromProperty(Property);
 		UObject* CPPTypeObject = RigVMTypeUtils::GetCPPTypeObjectFromProperty(Property);
 
 		// ignore any execute types
@@ -949,7 +949,9 @@ FRigVMMemoryStorageStruct* URigVM::CreateDebugMemory(FRigVMExtendedExecuteContex
 			{
 				continue;
 			}
+			PropertyCPPType = RigVMTypeUtils::GetUniqueStructTypeName(ScriptStruct);
 		}
+		const FString CPPType = RigVMTypeUtils::ArrayTypeFromBaseType(PropertyCPPType);
 
 		FRigVMPropertyDescription PropertyDescription(Property->GetFName(), CPPType, CPPTypeObject, TEXT("()"), true);
 		PropertyDescriptions.Add(PropertyDescription);
