Blueprint: Mark the generated class as being internal, if it's a skeleton or a reinstancing class, to allow easy filtering.

- Replaces name-based filtering methods.

#rnx

#Synced-CL 50220239

diff --git a/Engine/Plugins/Experimental/UAF/UAF/Source/UAFUncookedOnly/Private/Variables/UniversalObjectLocatorBindingType.cpp b/Engine/Plugins/Experimental/UAF/UAF/Source/UAFUncookedOnly/Private/Variables/UniversalObjectLocatorBindingType.cpp
--- a/Engine/Plugins/Experimental/UAF/UAF/Source/UAFUncookedOnly/Private/Variables/UniversalObjectLocatorBindingType.cpp
+++ b/Engine/Plugins/Experimental/UAF/UAF/Source/UAFUncookedOnly/Private/Variables/UniversalObjectLocatorBindingType.cpp
@@ -12,6 +12,7 @@
 #include "UniversalObjectLocatorFragmentType.h"
 #include "Component/AnimNextComponent.h"
 #include "Components/VerticalBox.h"
+#include "Engine/BlueprintGeneratedClass.h"
 #include "Framework/Application/SlateApplication.h"
 #include "Framework/PropertyViewer/IFieldExpander.h"
 #include "Framework/PropertyViewer/IFieldIterator.h"
@@ -196,11 +197,11 @@ TSharedRef<SWidget> FUniversalObjectLocatorBindingType::CreateEditWidget(const T
 
 					if(LibraryClass->HasAnyClassFlags(CLASS_CompiledFromBlueprint))
 					{
+						const UBlueprintGeneratedClass* BPClass = Cast<UBlueprintGeneratedClass>(LibraryClass);
+						
 						// Make sure we skip any out-of-date or skeleton classes
-						// It appears that the only way to detect skeleton classes is via the SKEL_ prefix
 						if(LibraryClass->HasAnyClassFlags(CLASS_NewerVersionExists) ||
-							LibraryClass->GetName().Contains(TEXT("SKEL_")) ||
-							LibraryClass->GetName().Contains(TEXT("REINST_")))
+							(BPClass && BPClass->IsInternalClass()))
 						{
 							continue;
 						}
diff --git a/Engine/Plugins/Runtime/RigVM/Source/RigVM/Private/RigVMBlueprintGeneratedClass.cpp b/Engine/Plugins/Runtime/RigVM/Source/RigVM/Private/RigVMBlueprintGeneratedClass.cpp
--- a/Engine/Plugins/Runtime/RigVM/Source/RigVM/Private/RigVMBlueprintGeneratedClass.cpp
+++ b/Engine/Plugins/Runtime/RigVM/Source/RigVM/Private/RigVMBlueprintGeneratedClass.cpp
@@ -24,8 +24,8 @@ uint8* URigVMBlueprintGeneratedClass::GetPersistentUberGraphFrame(UObject* Obj,
 
 void URigVMBlueprintGeneratedClass::PostInitInstance(UObject* InObj, FObjectInstancingGraph* InstanceGraph)
 {
-	// Skip SKEL classes.
-	if (InObj->GetName().StartsWith(TEXT("SKEL_")) || InObj->GetName().StartsWith(TEXT("Default__SKEL_")))
+	// Skip internal classes, such as skeleton and reinstancing.
+	if (IsInternalClass())
 	{
 		return;
 	}
@@ -55,22 +55,9 @@ void URigVMBlueprintGeneratedClass::Serialize(FArchive& Ar)
 		return;
 	}
 
-	// If we're the blueprint skeleton class, we don't need the VM. We should be an empty vessel. Ideally there should
-	// be a class flag to check for this, but there isn't, so we're left with name matching. 
-	auto IsSkeletonClass = [](const UClass* InClass)
-	{
-#if WITH_EDITOR 
-		return InClass->ClassGeneratedBy &&
-			(InClass->GetName().StartsWith(TEXT("SKEL_")) || InClass->GetName().StartsWith(TEXT("REINST_SKEL_")));
-#else
-		// For non-editor targets, this is all irrelevant, since the skeleton class will never exist for them.
-		return false;
-#endif
-	};
-
 	URigVM* VM = NewObject<URigVM>(GetTransientPackage());
 
-	if (!IsSkeletonClass(this))
+	if (!IsInternalClass())
 	{
 		if (URigVMHost* CDO = Cast<URigVMHost>(GetDefaultObject(true)))
 		{
@@ -86,7 +73,7 @@ void URigVMBlueprintGeneratedClass::Serialize(FArchive& Ar)
 	VM->Serialize(Ar);
 	UE_RIGVM_ARCHIVETRACE_ENTRY(Ar, TEXT("VM"));
 
-	if (!IsSkeletonClass(this))
+	if (!IsInternalClass())
 	{
 		if (URigVMHost* CDO = Cast<URigVMHost>(GetDefaultObject(false)))
 		{
diff --git a/Engine/Source/Editor/Kismet/Private/BlueprintCompilationManager.cpp b/Engine/Source/Editor/Kismet/Private/BlueprintCompilationManager.cpp
--- a/Engine/Source/Editor/Kismet/Private/BlueprintCompilationManager.cpp
+++ b/Engine/Source/Editor/Kismet/Private/BlueprintCompilationManager.cpp
@@ -3413,6 +3413,9 @@ UClass* FBlueprintCompilationManagerImpl::FastGenerateSkeletonClass(UBlueprint*
 		Ret = CompilerContext.NewClass;
 		Ret->SetFlags(RF_Transient);
 		CompilerContext.NewClass = OriginalNewClass;
+		
+		// Mark this class as being an internal class to ease filtering when enumerating all BP classes.
+		Ret->bIsInternalClass = true;
 	}
 	else
 	{
diff --git a/Engine/Source/Editor/UnrealEd/Private/Kismet2/KismetReinstanceUtilities.cpp b/Engine/Source/Editor/UnrealEd/Private/Kismet2/KismetReinstanceUtilities.cpp
--- a/Engine/Source/Editor/UnrealEd/Private/Kismet2/KismetReinstanceUtilities.cpp
+++ b/Engine/Source/Editor/UnrealEd/Private/Kismet2/KismetReinstanceUtilities.cpp
@@ -2046,6 +2046,12 @@ void FBlueprintCompileReinstancer::MoveDependentSkelToReinst(UClass* const Owner
 		ReinstClass->Bind();
 		ReinstClass->StaticLink(true);
 
+		// Mark this class as being an internal class to ease filtering when enumerating all BP classes. 
+		if (UBlueprintGeneratedClass* BPReinstClass = Cast<UBlueprintGeneratedClass>(ReinstClass))
+		{
+			BPReinstClass->bIsInternalClass = true;
+		}
+		
 		if (!ReinstClass->GetSparseClassDataStruct())
 		{
 			if (UScriptStruct* SparseClassDataStructArchetype = ReinstClass->GetSparseClassDataArchetypeStruct())
@@ -2148,6 +2154,9 @@ UClass* FBlueprintCompileReinstancer::MoveCDOToNewClass(UClass* OwnerClass, cons
 	CopyOfOwnerClass->Bind();
 	CopyOfOwnerClass->StaticLink(true);
 
+	// Mark this class as being an internal class to ease filtering when enumerating all BP classes. 
+	BPGDuplicatedClass->bIsInternalClass = true;
+	
 	// make sure we've bound to our native sparse data - this should have been done in
 	// link but don't want to destabilize early adopters:
 	if (!CopyOfOwnerClass->GetSparseClassDataStruct())
diff --git a/Engine/Source/Runtime/Engine/Classes/Engine/BlueprintGeneratedClass.h b/Engine/Source/Runtime/Engine/Classes/Engine/BlueprintGeneratedClass.h
--- a/Engine/Source/Runtime/Engine/Classes/Engine/BlueprintGeneratedClass.h
+++ b/Engine/Source/Runtime/Engine/Classes/Engine/BlueprintGeneratedClass.h
@@ -455,6 +455,14 @@ private:
 	/** Flag to make sure the custom property list has been initialized */
 	uint8 bCustomPropertyListForPostConstructionInitialized:1;
 
+#if WITH_EDITORONLY_DATA
+	/** Flag to mark this generated class as an internal class, not for public consumption, such as a temporary skeleton or reinstancing class. 
+	 *  Marked as a UPROPERTY so that it carries over during object duplication.
+	 */
+	UPROPERTY()
+	uint8 bIsInternalClass:1;
+#endif
+	
 public:
 	/** Array of objects containing information for dynamically binding delegates to functions in this blueprint */
 	UPROPERTY()
@@ -688,6 +696,19 @@ public:
 
 	/** Whether or not to use "fast path" component instancing. */
 	ENGINE_API bool UseFastPathComponentInstancing();
+	
+	/** Returns true if this is an internal BP class, either a skeleton or a reinstancing class, which are not meant to be shown to the user.
+	 *  @return true if this class is an internal BP class.
+	 */
+	bool IsInternalClass() const
+	{
+#if WITH_EDITORONLY_DATA
+		return bIsInternalClass;
+#else
+		// Internal BP classes will never show up in non-editor builds.
+		return false;
+#endif
+	}
 
 #if WITH_EDITORONLY_DATA
 	FBlueprintDebugData DebugData;
@@ -758,4 +779,8 @@ private:
 	TWeakObjectPtr<UScriptStruct> SparseClassDataPendingConformStruct = nullptr;
 	void* SparseClassDataPendingConform = nullptr;
 #endif // WITH_EDITORONLY_DATA
+	
+	// Also set on UClass, but friend-ness is not an inherited attribute in C++.
+	friend class FBlueprintCompileReinstancer;
+	friend struct FBlueprintCompilationManagerImpl;
 };
diff --git a/Engine/Source/Runtime/Engine/Private/BlueprintGeneratedClass.cpp b/Engine/Source/Runtime/Engine/Private/BlueprintGeneratedClass.cpp
--- a/Engine/Source/Runtime/Engine/Private/BlueprintGeneratedClass.cpp
+++ b/Engine/Source/Runtime/Engine/Private/BlueprintGeneratedClass.cpp
@@ -406,6 +406,7 @@ UBlueprintGeneratedClass::UBlueprintGeneratedClass(const FObjectInitializer& Obj
 	bSupportsDynamicInstancedReference = false;
 	bCustomPropertyListForPostConstructionInitialized = false;
 #if WITH_EDITORONLY_DATA
+	bIsInternalClass = false;
 	bIsSparseClassDataSerializable = false;
 #endif
 }
