Control Rig: Improve optimization of FRigVMPluginsPrefixCache when plugins are mounted/unmounted
#rnx

#Synced-CL 50204154

diff --git a/Engine/Plugins/Runtime/RigVM/Source/RigVM/Private/RigVMTypeUtils.cpp b/Engine/Plugins/Runtime/RigVM/Source/RigVM/Private/RigVMTypeUtils.cpp
--- a/Engine/Plugins/Runtime/RigVM/Source/RigVM/Private/RigVMTypeUtils.cpp
+++ b/Engine/Plugins/Runtime/RigVM/Source/RigVM/Private/RigVMTypeUtils.cpp
@@ -734,9 +734,9 @@ FString FRigVMPluginsPrefixCache::GetPluginName(const UObject* InObject) const
 		int32 PeriodPos = INDEX_NONE;
 		if (Path.FindChar(TEXT('.'), PeriodPos))
 		{
-			if (const FString* PluginName = CachedPluginPrefixes.Find(Path.Left(PeriodPos + 1)))
+			if (const FRigVMPluginReferences* References = CachedPluginPrefixes.Find(Path.Left(PeriodPos + 1)))
 			{
-				return *PluginName;
+				return *References->PluginName;
 			}
 		}
 	}
@@ -745,9 +745,9 @@ FString FRigVMPluginsPrefixCache::GetPluginName(const UObject* InObject) const
 		const int32 SecondSlashPos = Path.Find(TEXT("/"), ESearchCase::IgnoreCase, ESearchDir::FromStart, 1);
 		if (SecondSlashPos != INDEX_NONE)
 		{
-			if (const FString* PluginName = CachedPluginPrefixes.Find(Path.Left(SecondSlashPos + 1)))
+			if (const FRigVMPluginReferences* References = CachedPluginPrefixes.Find(Path.Left(SecondSlashPos + 1)))
 			{
-				return *PluginName;
+				return *References->PluginName;
 			}
 		}
 	}
@@ -767,7 +767,9 @@ void FRigVMPluginsPrefixCache::Refresh()
 		const TArray<FString> Prefixes = RigVMTypeUtils::GetPluginModulePrefixes(Plugin.Get());
 		for (const FString& Prefix : Prefixes)
 		{
-			CachedPluginPrefixes.Add(Prefix, Plugin->GetName());
+			FRigVMPluginReferences& NewPrefix = CachedPluginPrefixes.FindOrAdd(Prefix);
+			NewPrefix.PluginName = Plugin->GetName();
+			NewPrefix.Count++;
 		}
 	}
 	
@@ -789,10 +791,27 @@ FRigVMPluginsPrefixCache::~FRigVMPluginsPrefixCache()
 
 void FRigVMPluginsPrefixCache::OnPluginMounted(IPlugin& InPlugin)
 {
-	Refresh();
+	const TArray<FString> Prefixes = RigVMTypeUtils::GetPluginModulePrefixes(InPlugin);
+	for (const FString& Prefix : Prefixes)
+	{
+		FRigVMPluginReferences& NewPrefix = CachedPluginPrefixes.FindOrAdd(Prefix);
+		NewPrefix.PluginName = InPlugin.GetName();
+		NewPrefix.Count++;
+	}
 }
 
 void FRigVMPluginsPrefixCache::OnPluginUnmounted(IPlugin& InPlugin)
 {
-	Refresh();
-}
+	const TArray<FString> Prefixes = RigVMTypeUtils::GetPluginModulePrefixes(InPlugin);
+	for (const FString& Prefix : Prefixes)
+	{
+		if (FRigVMPluginReferences* Plugin = CachedPluginPrefixes.Find(Prefix))
+		{
+			Plugin->Count--;
+			if (Plugin->Count == 0)
+			{
+				CachedPluginPrefixes.Remove(Prefix);
+			}
+		}
+	}
+}
\ No newline at end of file
diff --git a/Engine/Plugins/Runtime/RigVM/Source/RigVM/Public/RigVMTypeUtils.h b/Engine/Plugins/Runtime/RigVM/Source/RigVM/Public/RigVMTypeUtils.h
--- a/Engine/Plugins/Runtime/RigVM/Source/RigVM/Public/RigVMTypeUtils.h
+++ b/Engine/Plugins/Runtime/RigVM/Source/RigVM/Public/RigVMTypeUtils.h
@@ -68,7 +68,13 @@ private:
 	UE_API void OnPluginMounted(IPlugin& InPlugin);
 	UE_API void OnPluginUnmounted(IPlugin& InPlugin);
 	
-	TMap<FString, FString> CachedPluginPrefixes;
+	struct FRigVMPluginReferences
+	{
+		FString PluginName;
+		uint16 Count = 0;
+	};
+	
+	TMap<FString, FRigVMPluginReferences> CachedPluginPrefixes;
 	bool bInitialized = false;
 	mutable FRWLock CacheLock;
 	
