Reconciled offline work

#rnx

#Synced-CL 50249092

diff --git a/Engine/Plugins/Runtime/RigVM/Source/RigVM/Private/RigVMCore/RigVM.cpp b/Engine/Plugins/Runtime/RigVM/Source/RigVM/Private/RigVMCore/RigVM.cpp
--- a/Engine/Plugins/Runtime/RigVM/Source/RigVM/Private/RigVMCore/RigVM.cpp
+++ b/Engine/Plugins/Runtime/RigVM/Source/RigVM/Private/RigVMCore/RigVM.cpp
@@ -1712,6 +1712,7 @@ ERigVMExecuteResult URigVM::ExecuteVM(FRigVMExtendedExecuteContext& Context, con
 	if(bIsRootEntry)
 	{
 		Context.ResetExecutionState();
+		Context.SetVM(this);
 		Context.SliceOffsets.AddZeroed(Instructions.Num());
 	}
 
diff --git a/Engine/Plugins/Runtime/RigVM/Source/RigVM/Public/RigVMCore/RigVMNativized.h b/Engine/Plugins/Runtime/RigVM/Source/RigVM/Public/RigVMCore/RigVMNativized.h
--- a/Engine/Plugins/Runtime/RigVM/Source/RigVM/Public/RigVMCore/RigVMNativized.h
+++ b/Engine/Plugins/Runtime/RigVM/Source/RigVM/Public/RigVMCore/RigVMNativized.h
@@ -101,7 +101,7 @@ protected:
 		UpdateExternalVariables(Context);
 	
 		Context.ResetExecutionState();
-		Context.VM = this;
+		Context.SetVM(this);
 		Context.SliceOffsets.AddZeroed(InNumberInstructions);
 		Context.GetPublicData<ExecuteContextType>().EventName = InEntryName;
 		return Context.GetPublicData<ExecuteContextType>();
