RigVM: Fix null object array entries pins not generating their default value correctly

Addresses fallout from object ptrs being added to pins

#Synced-CL 50295456

diff --git a/Engine/Plugins/Runtime/RigVM/Source/RigVMDeveloper/Private/RigVMModel/RigVMPin.cpp b/Engine/Plugins/Runtime/RigVM/Source/RigVMDeveloper/Private/RigVMModel/RigVMPin.cpp
--- a/Engine/Plugins/Runtime/RigVM/Source/RigVMDeveloper/Private/RigVMModel/RigVMPin.cpp
+++ b/Engine/Plugins/Runtime/RigVM/Source/RigVMDeveloper/Private/RigVMModel/RigVMPin.cpp
@@ -421,14 +421,15 @@ void URigVMPin::GetExposedPinChainImpl(TArray<const URigVMPin*>& OutExposedPins,
 		// Add target pins connected to the current pin
 		for (URigVMLink* Link : Current->GetTargetLinks())
 		{
-			URigVMPin* TargetPin = Link->GetTargetPin();
-
-			// Variable nodes do not share the operand with their source link
-			if (TargetPin->GetNode()->IsA<URigVMVariableNode>())
+			if (URigVMPin* TargetPin = Link->GetTargetPin())
 			{
-				continue;
+				// Variable nodes do not share the operand with their source link
+				if (TargetPin->GetNode()->IsA<URigVMVariableNode>())
+				{
+					continue;
+				}
+				ToProcess.Push(TargetPin);
 			}
-			ToProcess.Push(TargetPin);
 		}
 
 		// If pin is on a collapse node, add entry pin
@@ -1112,13 +1113,7 @@ FString URigVMPin::GetDefaultValue(const URigVMPin::FPinOverride& InOverride, bo
 		
 		return DefaultValue.IsEmpty() ? EmptyStructDefaultValue : DefaultValue;
 	}
-	else if (IsUObject())
-	{
-		FSoftObjectPath Object(DefaultValueObject);
-		Cache = Object.ToString();
-		return Cache.GetValue();
-	}
-	else if (IsArrayElement() && DefaultValue.IsEmpty())
+	else if (IsArrayElement() && ((!IsUObject() && DefaultValue.IsEmpty()) || (IsUObject() && DefaultValueObject == nullptr)))
 	{
 		// array element cannot have an empty default value because when an array pin is
 		// added as a property to a memory storage class, its default value needs to reflect 
@@ -1131,6 +1126,12 @@ FString URigVMPin::GetDefaultValue(const URigVMPin::FPinOverride& InOverride, bo
 		Cache = URigVMController::GetPinInitialDefaultValue(this);
 		return Cache.GetValue();
 	}
+	else if (IsUObject())
+	{
+		FSoftObjectPath Object(DefaultValueObject);
+		Cache = Object.ToString();
+		return Cache.GetValue();
+	}
 
 	return DefaultValue;
 }
