Update RigVM function reference nodes when adding or removing  variable refrences to the contained graph

#Synced-CL 50205727

diff --git a/Engine/Plugins/Runtime/RigVM/Source/RigVMDeveloper/Private/RigVMModel/RigVMController.cpp b/Engine/Plugins/Runtime/RigVM/Source/RigVMDeveloper/Private/RigVMModel/RigVMController.cpp
--- a/Engine/Plugins/Runtime/RigVM/Source/RigVMDeveloper/Private/RigVMModel/RigVMController.cpp
+++ b/Engine/Plugins/Runtime/RigVM/Source/RigVMDeveloper/Private/RigVMModel/RigVMController.cpp
@@ -2079,6 +2079,13 @@ URigVMVariableNode* URigVMController::AddVariableNode(const FName& InVariableNam
 	Notify(ERigVMGraphNotifType::NodeAdded, Node);
 	Notify(ERigVMGraphNotifType::VariableAdded, Node);
 
+	if (URigVMLibraryNode* LibraryNode = Cast<URigVMNode>(Node)->FindFunctionForNode())
+	{
+		if (LibraryNode->GetGraph()->IsA<URigVMFunctionLibrary>())
+		{
+			PatchExternalVariablesToArguments(bSetupUndoRedo);
+		}
+	}
 	if (bSetupUndoRedo)
 	{
 		GetActionStack()->EndAction(Action);
@@ -7286,6 +7293,14 @@ bool URigVMController::RemoveNodes(TArray<URigVMNode*> InNodes, bool bSetupUndoR
 		if (URigVMVariableNode* VariableNode = Cast<URigVMVariableNode>(InNode))
 		{
 			Notify(ERigVMGraphNotifType::VariableRemoved, VariableNode);
+
+			if (URigVMLibraryNode* LibraryNode = Cast<URigVMNode>(VariableNode)->FindFunctionForNode())
+			{
+				if (LibraryNode->GetGraph()->IsA<URigVMFunctionLibrary>())
+				{
+					PatchExternalVariablesToArguments(bSetupUndoRedo);
+				}
+			}
 		}
 		
 		if (URigVMInjectionInfo* InjectionInfo = InNode->GetInjectionInfo())
@@ -24551,7 +24566,7 @@ FRigVMClientPatchResult URigVMController::PatchLocalVariableTypes()
 	return Result;
 }
 
-FRigVMClientPatchResult URigVMController::PatchExternalVariablesToArguments()
+FRigVMClientPatchResult URigVMController::PatchExternalVariablesToArguments(bool bInSetupUndoRedo)
 {
 	FRigVMClientPatchResult Result;
 	if (const URigVMGraph* Graph = GetGraph())
@@ -24583,7 +24598,7 @@ FRigVMClientPatchResult URigVMController::PatchExternalVariablesToArguments()
 					}
 
 					const FString TypeName = ExternalVariable.GetExtendedCPPType().ToString();
-					const FName PinName = AddExposedPin(ExternalVariable.GetName(), ERigVMPinDirection::Input, TypeName, *ExternalVariable.GetCPPTypeObjectPathName(), FString(), false, false, true);
+					const FName PinName = AddExposedPin(ExternalVariable.GetName(), ERigVMPinDirection::Input, TypeName, *ExternalVariable.GetCPPTypeObjectPathName(), FString(), bInSetupUndoRedo, false, true);
 					Result.bSucceeded = !PinName.IsNone();
 					if (Result.bSucceeded)
 					{
@@ -24610,7 +24625,7 @@ FRigVMClientPatchResult URigVMController::PatchExternalVariablesToArguments()
 						{
 							continue;
 						}
-						if (BindPinToVariable(Pin, MappedVariablePath, false))
+						if (BindPinToVariable(Pin, MappedVariablePath, bInSetupUndoRedo))
 						{
 							Result.bChangedContent = true;
 						}
diff --git a/Engine/Plugins/Runtime/RigVM/Source/RigVMDeveloper/Public/RigVMModel/RigVMController.h b/Engine/Plugins/Runtime/RigVM/Source/RigVMDeveloper/Public/RigVMModel/RigVMController.h
--- a/Engine/Plugins/Runtime/RigVM/Source/RigVMDeveloper/Public/RigVMModel/RigVMController.h
+++ b/Engine/Plugins/Runtime/RigVM/Source/RigVMDeveloper/Public/RigVMModel/RigVMController.h
@@ -1572,7 +1572,7 @@ protected:
 	UE_API FRigVMClientPatchResult PatchPinDefaultValues();
 	UE_API FRigVMClientPatchResult PatchUserDefinedStructPinNames();
 	UE_API FRigVMClientPatchResult PatchLocalVariableTypes();
-	UE_API FRigVMClientPatchResult PatchExternalVariablesToArguments();
+	UE_API FRigVMClientPatchResult PatchExternalVariablesToArguments(bool bInSetupUndoRedo=false);
 
 	UE_API ERigVMPinDefaultValueType GetDefaultValueType(const URigVMPin* InPin, const FString& InDefaultValue) const;
 
