RigVM: Support for callables / slicing for callables.

#rnx

#Synced-CL 50296749

diff --git a/Engine/Plugins/Runtime/RigVM/Source/RigVM/Private/RigVMCore/RigVM.cpp b/Engine/Plugins/Runtime/RigVM/Source/RigVM/Private/RigVMCore/RigVM.cpp
--- a/Engine/Plugins/Runtime/RigVM/Source/RigVM/Private/RigVMCore/RigVM.cpp
+++ b/Engine/Plugins/Runtime/RigVM/Source/RigVM/Private/RigVMCore/RigVM.cpp
@@ -2233,13 +2233,13 @@ ERigVMExecuteResult URigVM::ExecuteInstructions(FRigVMExtendedExecuteContext& Co
 
 				const int32 Count = (*((int32*)CountHandle.GetInputData()));
 				const int32 Index = (*((int32*)IndexHandle.GetInputData()));
-				Context.BeginSlice(Count, Index);
+				Context.BeginLoopSlice(Count, Index);
 				ContextPublicData.InstructionIndex++;
 				break;
 			}
 			case ERigVMOpCode::EndBlock:
 			{
-				Context.EndSlice();
+				Context.EndLoopSlice();
 				ContextPublicData.InstructionIndex++;
 				break;
 			}
diff --git a/Engine/Plugins/Runtime/RigVM/Source/RigVM/Public/RigVMCore/RigVMExecuteContext.h b/Engine/Plugins/Runtime/RigVM/Source/RigVM/Public/RigVMCore/RigVMExecuteContext.h
--- a/Engine/Plugins/Runtime/RigVM/Source/RigVM/Public/RigVMCore/RigVMExecuteContext.h
+++ b/Engine/Plugins/Runtime/RigVM/Source/RigVM/Public/RigVMCore/RigVMExecuteContext.h
@@ -45,15 +45,20 @@ struct FRigVMLogSettings;
 
 extern UE_API TAutoConsoleVariable<bool> CVarRigVMReportAllMessages;
 
-USTRUCT()
+enum class ERigVMSliceType : uint8
+{
+	Invalid,
+	Loop,
+	Callable
+};
+
 struct FRigVMSlice
 {
 public:
 
-	GENERATED_BODY()
-
 	FRigVMSlice()
-	: LowerBound(0)
+    : SliceType(ERigVMSliceType::Invalid)
+	, LowerBound(0)
 	, UpperBound(0)
 	, Index(INDEX_NONE)
 	, Offset(0)
@@ -62,21 +67,44 @@ public:
 		Reset();
 	}
 
-	FRigVMSlice(int32 InCount)
-		: LowerBound(0)
-		, UpperBound(InCount - 1)
-		, Index(INDEX_NONE)
-		, Offset(0)
-		, InstructionIndex(INDEX_NONE)
+	static FRigVMSlice MakeForLoop(int32 InCount)
 	{
-		Reset();
+		FRigVMSlice Slice;
+		Slice.SliceType = ERigVMSliceType::Loop;
+		Slice.UpperBound = InCount - 1;
+		Slice.Reset();
+		return Slice;
+	}
+
+	static FRigVMSlice MakeForCallable()
+	{
+		FRigVMSlice Slice;
+		Slice.SliceType = ERigVMSliceType::Callable;
+		// for callables the upperbound of 0 represents a single slice
+		Slice.Reset();
+		return Slice;
 	}
 
 	bool IsValid() const
 	{ 
-		return Index != INDEX_NONE;
+		return Index != INDEX_NONE && SliceType != ERigVMSliceType::Invalid;
 	}
-	
+
+	ERigVMSliceType GetType() const
+	{
+		return SliceType;
+	}
+
+	bool IsLoop() const
+	{
+		return SliceType == ERigVMSliceType::Loop;
+	}
+
+	bool IsCallable() const
+	{
+		return SliceType == ERigVMSliceType::Callable;
+	}
+
 	bool IsComplete() const
 	{
 		return Index > UpperBound;
@@ -190,11 +218,12 @@ public:
 
 	uint32 GetHash() const
 	{
-		return HashCombine(GetTypeHash(LowerBound), GetTypeHash(Index));
+		return HashCombine(GetTypeHash(SliceType), GetTypeHash(LowerBound), GetTypeHash(Index));
 	}
 
 private:
 
+	ERigVMSliceType SliceType;
 	int32 LowerBound;
 	int32 UpperBound;
 	int32 Index;
@@ -701,35 +730,56 @@ public:
 		return Slices.Last();
 	}
 
-	void BeginSlice(int32 InCount, int32 InRelativeIndex = 0)
+	void BeginLoopSlice(int32 InCount, int32 InRelativeIndex = 0)
 	{
 		ensure(!IsSliceComplete());
-		Slices.Add(FRigVMSlice(InCount));
+		Slices.Add(FRigVMSlice::MakeForLoop(InCount));
 		Slices.Last().SetRelativeIndex(InRelativeIndex);
 		Slices.Last().SetInstructionIndex(GetPublicData<>().InstructionIndex);
 		Slices.Last().SetOffset(SliceOffsets[GetPublicData<>().InstructionIndex]);
 	}
 
-
-	void EndSlice()
+	void EndLoopSlice()
 	{
 		ensure(Slices.Num() > 1);
+		ensure(Slices.Last().IsLoop());
+		PopSlice();
+	}
+
+	void BeginCallableSlice()
+	{
+		ensure(!IsSliceComplete());
+		Slices.Add(FRigVMSlice::MakeForCallable());
+		Slices.Last().SetInstructionIndex(GetPublicData<>().InstructionIndex);
+		Slices.Last().SetOffset(SliceOffsets[GetPublicData<>().InstructionIndex]);
+	}
 
+	void EndCallableSlice()
+	{
+		ensure(Slices.Num() > 0);
+		ensure(Slices.Last().IsCallable());
+		PopSlice();
+	}
+
+private:
+	void PopSlice()
+	{
 		// if this slice has reached its upperbound
-		// we want to make sure to increment the offset
-		// for the instruction the slice originated from,
-		// so that next time around the slice indices are not
-		// reused.
-		const FRigVMSlice PoppedSlice = Slices.Pop();
-		if(PoppedSlice.GetRelativeIndex() == PoppedSlice.Num() - 1)
-		{
-			if(SliceOffsets.IsValidIndex(PoppedSlice.GetInstructionIndex()))
-			{
-				SliceOffsets[PoppedSlice.GetInstructionIndex()] += static_cast<uint16>(PoppedSlice.Num());
-			}
-		}
+    	// we want to make sure to increment the offset
+    	// for the instruction the slice originated from,
+    	// so that next time around the slice indices are not
+    	// reused.
+    	const FRigVMSlice PoppedSlice = Slices.Pop();
+    	if(PoppedSlice.GetRelativeIndex() == PoppedSlice.Num() - 1)
+    	{
+    		if(SliceOffsets.IsValidIndex(PoppedSlice.GetInstructionIndex()))
+    		{
+    			SliceOffsets[PoppedSlice.GetInstructionIndex()] += static_cast<uint16>(PoppedSlice.Num());
+    		}
+    	}
 	}
 
+public:
 	void IncrementSlice()
 	{
 		FRigVMSlice& ActiveSlice = Slices.Last();
diff --git a/Engine/Plugins/Runtime/RigVM/Source/RigVM/Public/RigVMCore/RigVMNativized.h b/Engine/Plugins/Runtime/RigVM/Source/RigVM/Public/RigVMCore/RigVMNativized.h
--- a/Engine/Plugins/Runtime/RigVM/Source/RigVM/Public/RigVMCore/RigVMNativized.h
+++ b/Engine/Plugins/Runtime/RigVM/Source/RigVM/Public/RigVMCore/RigVMNativized.h
@@ -127,14 +127,24 @@ protected:
 		}
 	};
 
-	void BeginSlice(FRigVMExtendedExecuteContext& Context, int32 InCount, int32 InRelativeIndex = 0)
+	void BeginLoopSlice(FRigVMExtendedExecuteContext& Context, int32 InCount, int32 InRelativeIndex = 0)
 	{
-		Context.BeginSlice(InCount, InRelativeIndex);
+		Context.BeginLoopSlice(InCount, InRelativeIndex);
 	}
 
-	void EndSlice(FRigVMExtendedExecuteContext& Context)
+	void EndLoopSlice(FRigVMExtendedExecuteContext& Context)
 	{
-		Context.EndSlice();
+		Context.EndLoopSlice();
+	}
+
+	void BeginCallableSlice(FRigVMExtendedExecuteContext& Context)
+	{
+		Context.BeginCallableSlice();
+	}
+
+	void EndCallableSlice(FRigVMExtendedExecuteContext& Context)
+	{
+		Context.EndCallableSlice();
 	}
 
 	bool IsValidArraySize(const FRigVMExtendedExecuteContext& Context, int32 InArraySize) const
