RigVM: Added Random Int function

#Synced-CL 50310627

#Synced-CL 50310627

diff --git a/Engine/Plugins/Runtime/RigVM/Source/RigVM/Private/RigVMFunctions/Math/RigVMFunction_Random.cpp b/Engine/Plugins/Runtime/RigVM/Source/RigVM/Private/RigVMFunctions/Math/RigVMFunction_Random.cpp
--- a/Engine/Plugins/Runtime/RigVM/Source/RigVM/Private/RigVMFunctions/Math/RigVMFunction_Random.cpp
+++ b/Engine/Plugins/Runtime/RigVM/Source/RigVM/Private/RigVMFunctions/Math/RigVMFunction_Random.cpp
@@ -16,6 +16,28 @@ float FRigVMFunction_Random_Helper(int32& Seed)
 	return FPlatformMath::Fractional(Result.f);
 }
 
+FRigVMFunction_RandomInt_Execute()
+{
+
+	if (Seed != BaseSeed)
+	{
+		LastSeed = BaseSeed = Seed;
+	}
+
+	TimeLeft = TimeLeft - ExecuteContext.GetDeltaTime<float>();
+	if (TimeLeft > 0.f)
+	{
+		Result = LastResult;
+		return;
+	}
+
+	LastSeed = (LastSeed * 196314165) + 907633515;
+
+	Result = Minimum + ((uint32)LastSeed) % (FMath::Max(Maximum - Minimum, 0) + 1);
+	TimeLeft = Duration < -SMALL_NUMBER ? FLT_MAX : Duration;
+	LastResult = Result;
+}
+
 FRigVMFunction_RandomFloat_Execute()
 {
 
diff --git a/Engine/Plugins/Runtime/RigVM/Source/RigVM/Public/RigVMFunctions/Math/RigVMFunction_Random.h b/Engine/Plugins/Runtime/RigVM/Source/RigVM/Public/RigVMFunctions/Math/RigVMFunction_Random.h
--- a/Engine/Plugins/Runtime/RigVM/Source/RigVM/Public/RigVMFunctions/Math/RigVMFunction_Random.h
+++ b/Engine/Plugins/Runtime/RigVM/Source/RigVM/Public/RigVMFunctions/Math/RigVMFunction_Random.h
@@ -5,6 +5,63 @@
 #include "RigVMFunction_MathBase.h"
 #include "RigVMFunction_Random.generated.h"
 
+/**
+ * Generates a random integer between a min and a max (inclusive)
+ */
+USTRUCT(meta = (DisplayName = "Random (Int)", Category = "Math|Random", TemplateName = "Random"))
+struct FRigVMFunction_RandomInt : public FRigVMFunction_MathBase
+{
+	GENERATED_BODY()
+
+	FRigVMFunction_RandomInt()
+	{
+		Seed = BaseSeed = LastSeed = 217;
+		Minimum = Result = LastResult = 0;
+		Maximum = 1;
+		Duration = 0.f;
+		TimeLeft = 0.f;
+	}
+
+	RIGVM_METHOD()
+	RIGVM_API virtual void Execute() override;
+
+	// The seed for the random number generator
+	UPROPERTY(meta = (Input))
+	int32 Seed;
+
+	// The minimum value for the random number range
+	UPROPERTY(meta = (Input))
+	int32 Minimum = 0;
+
+	// The maximum value for the random number range (inclusive)
+	UPROPERTY(meta = (Input))
+	int32 Maximum = 1;
+
+	/**
+	 * The duration at which the number won't change.
+	 * Use 0 for a different number every time.
+	 * A negative number (for ex: -1) results in an infinite duration.
+	 */
+	UPROPERTY(meta = (Input))
+	float Duration;
+
+	// The resulting random number
+	UPROPERTY(meta = (Output))
+	int32 Result;
+
+	UPROPERTY()
+	int32 LastResult;
+
+	UPROPERTY()
+	int32 LastSeed;
+
+	UPROPERTY()
+	int32 BaseSeed;
+
+	UPROPERTY()
+	float TimeLeft;
+};
+
 /**
  * Generates a random float between a min and a max
  */
