[RigVMRegistry/ALT]
- Make sure to call RegisterGCObject/UnregisterGCObject in the most derived class to avoid pure virtual function call from other threads

#Synced-CL 50254652

diff --git a/Engine/Plugins/Runtime/RigVM/Source/RigVM/Private/RigVMCore/RigVMRegistry.cpp b/Engine/Plugins/Runtime/RigVM/Source/RigVM/Private/RigVMCore/RigVMRegistry.cpp
--- a/Engine/Plugins/Runtime/RigVM/Source/RigVM/Private/RigVMCore/RigVMRegistry.cpp
+++ b/Engine/Plugins/Runtime/RigVM/Source/RigVM/Private/RigVMCore/RigVMRegistry.cpp
@@ -74,12 +74,16 @@ static FDelayedAutoRegisterHelper GRigVMRegistry_PostUObjectInit(EDelayedRegiste
 #endif // UE_WITH_CONSTINIT_UOBJECT
 
 FRigVMRegistry_NoLock::FRigVMRegistry_NoLock(bool InIsGlobalRegistry)
-	: ThisHandle(this)
+	: FGCObject(EFlags::RegisterLater)
+	, ThisHandle(this)
 	, bIsGlobalRegistry(InIsGlobalRegistry)
 	, bAvoidTypePropagation(false)
 	, bEverRefreshedEngineTypes(false)
 	, bEverRefreshedDispatchFactoriesAfterEngineInit(false)
 {
+	// Needs to be called in the most derived constructor to avoid still
+	// having pure virtual calls in the virtual table at registration time.
+	RegisterGCObject();
 }
 
 TSharedRef<FRigVMRegistry_NoLock> FRigVMRegistry_NoLock::CreateLocalizedRegistry()
@@ -130,6 +134,10 @@ TSharedRef<FRigVMRegistry_NoLock> FRigVMRegistry_NoLock::CloneLocalizedRegistry(
 
 FRigVMRegistry_NoLock::~FRigVMRegistry_NoLock()
 {
+	// Needs to be called in the most derived destructor to avoid leaving
+	// a registered object with half its virtual table destroyed.
+	UnregisterGCObject();
+
 	FRigVMRegistry_NoLock::Reset_NoLock();
 }
 
