RigVM: Change s-witch dispatch node to accept enums and display the case names in the pins

#Synced-CL 50541888

diff --git a/Engine/Plugins/Animation/ControlRig/Source/ControlRig/Public/Units/Execution/RigUnit_Collection.h b/Engine/Plugins/Animation/ControlRig/Source/ControlRig/Public/Units/Execution/RigUnit_Collection.h
--- a/Engine/Plugins/Animation/ControlRig/Source/ControlRig/Public/Units/Execution/RigUnit_Collection.h
+++ b/Engine/Plugins/Animation/ControlRig/Source/ControlRig/Public/Units/Execution/RigUnit_Collection.h
@@ -652,11 +652,13 @@ struct FRigUnit_CollectionLoop : public FRigUnit_CollectionBaseMutable
 	}
 
 	// FRigVMStruct overrides
+#if WITH_EDITORONLY_DATA
 	virtual const TArray<FName>& GetControlFlowBlocks_Impl() const override
 	{
 		static const TArray<FName> Blocks = {ExecutePinName, ForLoopCompletedPinName};
 		return Blocks;
 	}
+#endif
 	virtual const bool IsControlFlowBlockSliced(const FName& InBlockName) const { return InBlockName == ExecutePinName; }
 	virtual int32 GetNumSlices() const override { return Count; }
 
diff --git a/Engine/Plugins/Animation/ControlRig/Source/ControlRig/Public/Units/Execution/RigUnit_Hierarchy.h b/Engine/Plugins/Animation/ControlRig/Source/ControlRig/Public/Units/Execution/RigUnit_Hierarchy.h
--- a/Engine/Plugins/Animation/ControlRig/Source/ControlRig/Public/Units/Execution/RigUnit_Hierarchy.h
+++ b/Engine/Plugins/Animation/ControlRig/Source/ControlRig/Public/Units/Execution/RigUnit_Hierarchy.h
@@ -835,11 +835,13 @@ struct FRigUnit_PoseLoop : public FRigUnit_HierarchyBaseMutable
 	}
 
 	// FRigVMStruct overrides
+#if WITH_EDITORONLY_DATA
 	virtual const TArray<FName>& GetControlFlowBlocks_Impl() const override
 	{
 		static const TArray<FName> Blocks = {ExecutePinName, ForLoopCompletedPinName};
 		return Blocks;
 	}
+#endif
 	virtual const bool IsControlFlowBlockSliced(const FName& InBlockName) const { return InBlockName == ExecutePinName; }
 	virtual int32 GetNumSlices() const override { return Count; }
 
diff --git a/Engine/Plugins/Runtime/RigVM/Source/RigVM/Private/RigVMCore/RigVMDispatchFactory.cpp b/Engine/Plugins/Runtime/RigVM/Source/RigVM/Private/RigVMCore/RigVMDispatchFactory.cpp
--- a/Engine/Plugins/Runtime/RigVM/Source/RigVM/Private/RigVMCore/RigVMDispatchFactory.cpp
+++ b/Engine/Plugins/Runtime/RigVM/Source/RigVM/Private/RigVMCore/RigVMDispatchFactory.cpp
@@ -37,7 +37,7 @@ FName FRigVMDispatchFactory::GetNextAggregateName(const FName& InLastAggregatePi
 	return FRigVMStruct().GetNextAggregateName(InLastAggregatePinName);
 }
 
-FName FRigVMDispatchFactory::GetDisplayNameForArgument(const FName& InArgumentName) const
+FName FRigVMDispatchFactory::GetDisplayNameForArgument(const FName& InArgumentName, const FRigVMTemplateTypeMap& InTypes) const
 {
 	if(InArgumentName == FRigVMStruct::ExecuteContextName)
 	{
@@ -219,25 +219,29 @@ FName FRigVMDispatchFactory::GetArgumentNameForOperandIndex(int32 InOperandIndex
 	return GetArgumentInfos(InRegistry)[InOperandIndex].Name;
 }
 
-const TArray<FName>& FRigVMDispatchFactory::GetControlFlowBlocks(const FRigVMDispatchContext& InContext) const
+#if WITH_EDITORONLY_DATA
+const TArray<FName>& FRigVMDispatchFactory::GetControlFlowBlocks(const FRigVMDispatchContext& InContext, const FRigVMTemplateTypeMap& InTypeMap) const
 {
-	const TArray<FName>& Blocks = GetControlFlowBlocks_Impl(InContext);
+	const TArray<FName>& Blocks = GetControlFlowBlocks_Impl(InContext, InTypeMap);
 #if WITH_EDITOR
 	FRigVMStruct::ValidateControlFlowBlocks(Blocks);
 #endif
 	return Blocks;
 }
+#endif
 
 bool FRigVMDispatchFactory::SupportsExecuteContextStruct(const UScriptStruct* InExecuteContextStruct) const
 {
 	return InExecuteContextStruct->IsChildOf(GetExecuteContextStruct());
 }
 
-const TArray<FName>& FRigVMDispatchFactory::GetControlFlowBlocks_Impl(const FRigVMDispatchContext& InContext) const
+#if WITH_EDITORONLY_DATA
+const TArray<FName>& FRigVMDispatchFactory::GetControlFlowBlocks_Impl(const FRigVMDispatchContext& InContext, const FRigVMTemplateTypeMap& InTypeMap) const
 {
 	static const TArray<FName> EmptyArray;
 	return EmptyArray;
 }
+#endif
 
 const TArray<FName>* FRigVMDispatchFactory::UpdateArgumentNameCache(int32 InNumberOperands, FRigVMRegistryHandle& InRegistry) const
 {
diff --git a/Engine/Plugins/Runtime/RigVM/Source/RigVM/Private/RigVMCore/RigVMStruct.cpp b/Engine/Plugins/Runtime/RigVM/Source/RigVM/Private/RigVMCore/RigVMStruct.cpp
--- a/Engine/Plugins/Runtime/RigVM/Source/RigVM/Private/RigVMCore/RigVMStruct.cpp
+++ b/Engine/Plugins/Runtime/RigVM/Source/RigVM/Private/RigVMCore/RigVMStruct.cpp
@@ -345,6 +345,8 @@ float FRigVMStruct::GetRatioFromIndex(int32 InIndex, int32 InCount)
 	return ((float)FMath::Clamp<int32>(InIndex, 0, InCount - 1)) / ((float)(InCount - 1));
 }
 
+#if WITH_EDITORONLY_DATA
+
 bool FRigVMStruct::IsForLoop() const
 {
 	UE_LOG(LogRigVM, Warning, TEXT("FRigVMStruct::IsForLoop is deprecated. Please call FRigVMStruct::IsControlFlowNode instead."));
@@ -375,6 +377,7 @@ const TArray<FName>& FRigVMStruct::GetControlFlowBlocks() const
 #endif
 	return Blocks;
 }
+#endif
 
 TArray<FRigVMUserWorkflow> FRigVMStruct::GetWorkflows(ERigVMUserWorkflowType InType, const UObject* InSubject) const
 {
@@ -880,11 +883,13 @@ bool FRigVMStruct::ApplyUpgradeInfo(const FRigVMStructUpgradeInfo& InUpgradeInfo
 	return true;
 }
 
+#if WITH_EDITORONLY_DATA
 const TArray<FName>& FRigVMStruct::GetControlFlowBlocks_Impl() const
 {
 	static const TArray<FName> EmptyArray;
 	return EmptyArray;
 }
+#endif
 
 #if WITH_EDITOR
 
diff --git a/Engine/Plugins/Runtime/RigVM/Source/RigVM/Private/RigVMCore/RigVMTemplate.cpp b/Engine/Plugins/Runtime/RigVM/Source/RigVM/Private/RigVMCore/RigVMTemplate.cpp
--- a/Engine/Plugins/Runtime/RigVM/Source/RigVM/Private/RigVMCore/RigVMTemplate.cpp
+++ b/Engine/Plugins/Runtime/RigVM/Source/RigVM/Private/RigVMCore/RigVMTemplate.cpp
@@ -1334,7 +1334,13 @@ FText FRigVMTemplate::GetDisplayNameForArgument(const FName& InArgumentName, con
 {
 	if(const FRigVMDispatchFactory* Factory = GetDispatchFactory())
 	{
-		const FName DisplayName = Factory->GetDisplayNameForArgument(InArgumentName);
+		FRigVMTemplateTypeMap Types;
+		if (InPermutationIndices.Num() == 1)
+		{
+			Types = GetTypesForPermutation_NoLock(InPermutationIndices[0], FRigVMRegistryReadLock());
+		}
+		
+		const FName DisplayName = Factory->GetDisplayNameForArgument(InArgumentName, Types);
 		if(DisplayName.IsNone())
 		{
 			return FText();
diff --git a/Engine/Plugins/Runtime/RigVM/Source/RigVM/Private/RigVMFunctions/RigVMDispatch_Array.cpp b/Engine/Plugins/Runtime/RigVM/Source/RigVM/Private/RigVMFunctions/RigVMDispatch_Array.cpp
--- a/Engine/Plugins/Runtime/RigVM/Source/RigVM/Private/RigVMFunctions/RigVMDispatch_Array.cpp
+++ b/Engine/Plugins/Runtime/RigVM/Source/RigVM/Private/RigVMFunctions/RigVMDispatch_Array.cpp
@@ -1586,12 +1586,14 @@ const TArray<FRigVMExecuteArgument>& FRigVMDispatch_ArrayIterator::GetExecuteArg
 	return ExecuteArguments;
 }
 
+#if WITH_EDITORONLY_DATA
 const TArray<FName>& FRigVMDispatch_ArrayIterator::GetControlFlowBlocks_Impl(
-	const FRigVMDispatchContext& InContext) const
+	const FRigVMDispatchContext& InContext, const FRigVMTemplateTypeMap& InTypeMap) const
 {
 	static const TArray<FName> Blocks = {FRigVMStruct::ExecuteContextName, FRigVMStruct::ControlFlowCompletedName};
 	return Blocks;
 }
+#endif
 
 const bool FRigVMDispatch_ArrayIterator::IsControlFlowBlockSliced(const FName& InBlockName) const
 {
diff --git a/Engine/Plugins/Runtime/RigVM/Source/RigVM/Private/RigVMFunctions/RigVMDispatch_Switch.cpp b/Engine/Plugins/Runtime/RigVM/Source/RigVM/Private/RigVMFunctions/RigVMDispatch_Switch.cpp
--- a/Engine/Plugins/Runtime/RigVM/Source/RigVM/Private/RigVMFunctions/RigVMDispatch_Switch.cpp
+++ b/Engine/Plugins/Runtime/RigVM/Source/RigVM/Private/RigVMFunctions/RigVMDispatch_Switch.cpp
@@ -9,13 +9,14 @@
 #include UE_INLINE_GENERATED_CPP_BY_NAME(RigVMDispatch_Switch)
 #define LOCTEXT_NAMESPACE "RigVMDispatch_Switch"
 
-FName FRigVMDispatch_SwitchInt32::GetArgumentNameForOperandIndex(int32 InOperandIndex, int32 InTotalOperands, FRigVMRegistryHandle& InRegistry) const
+FName FRigVMDispatch_SwitchBase::GetArgumentNameForOperandIndex(int32 InOperandIndex, int32 InTotalOperands, FRigVMRegistryHandle& InRegistry) const
 {
 	static const FLazyName ArgumentNames[] = {
 		IndexName,
 		FRigVMStruct::ControlFlowBlockToRunName
 	};
 	check(InTotalOperands == UE_ARRAY_COUNT(ArgumentNames));
+	check(InOperandIndex < InTotalOperands);
 	return ArgumentNames[InOperandIndex];
 }
 
@@ -29,7 +30,7 @@ const TArray<FRigVMTemplateArgumentInfo>& FRigVMDispatch_SwitchInt32::GetArgumen
 	return CachedArgumentInfos;
 }
 
-const TArray<FRigVMExecuteArgument>& FRigVMDispatch_SwitchInt32::GetExecuteArguments_Impl(const FRigVMDispatchContext& InContext) const
+const TArray<FRigVMExecuteArgument>& FRigVMDispatch_SwitchBase::GetExecuteArguments_Impl(const FRigVMDispatchContext& InContext) const
 {
 	static const TArray<FRigVMExecuteArgument> Arguments =
 	{
@@ -42,7 +43,7 @@ const TArray<FRigVMExecuteArgument>& FRigVMDispatch_SwitchInt32::GetExecuteArgum
 
 #if WITH_EDITOR
 
-FString FRigVMDispatch_SwitchInt32::GetArgumentMetaData(const FName& InArgumentName, const FName& InMetaDataKey) const
+FString FRigVMDispatch_SwitchBase::GetArgumentMetaData(const FName& InArgumentName, const FName& InMetaDataKey) const
 {
 	if(InArgumentName == CasesName)
 	{
@@ -54,7 +55,7 @@ FString FRigVMDispatch_SwitchInt32::GetArgumentMetaData(const FName& InArgumentN
 	return FRigVMDispatch_CoreBase::GetArgumentMetaData(InArgumentName, InMetaDataKey);
 }
 
-FString FRigVMDispatch_SwitchInt32::GetArgumentDefaultValue(const FName& InArgumentName,
+FString FRigVMDispatch_SwitchBase::GetArgumentDefaultValue(const FName& InArgumentName,
 	TRigVMTypeIndex InTypeIndex) const
 {
 	if(InArgumentName == CasesName)
@@ -65,7 +66,7 @@ FString FRigVMDispatch_SwitchInt32::GetArgumentDefaultValue(const FName& InArgum
 	return FRigVMDispatch_CoreBase::GetArgumentDefaultValue(InArgumentName, InTypeIndex);
 }
 
-FName FRigVMDispatch_SwitchInt32::GetDisplayNameForArgument(const FName& InArgumentName) const
+FName FRigVMDispatch_SwitchInt32::GetDisplayNameForArgument(const FName& InArgumentName, const FRigVMTemplateTypeMap& InTypes) const
 {
 	static const FString CasesPrefix = CasesName.ToString() + TEXT(".");
 	const FString ArgumentNameString = InArgumentName.ToString();
@@ -79,10 +80,10 @@ FName FRigVMDispatch_SwitchInt32::GetDisplayNameForArgument(const FName& InArgum
 			return GetCaseDisplayName(CaseIndex);
 		}
 	}
-	return FRigVMDispatch_CoreBase::GetDisplayNameForArgument(InArgumentName);
+	return FRigVMDispatch_CoreBase::GetDisplayNameForArgument(InArgumentName, InTypes);
 }
 
-FText FRigVMDispatch_SwitchInt32::GetArgumentTooltip(const FName& InArgumentName, TRigVMTypeIndex InTypeIndex) const
+FText FRigVMDispatch_SwitchBase::GetArgumentTooltip(const FName& InArgumentName, TRigVMTypeIndex InTypeIndex) const
 {
 	if (InArgumentName == IndexName)
 	{
@@ -97,7 +98,8 @@ FText FRigVMDispatch_SwitchInt32::GetArgumentTooltip(const FName& InArgumentName
 
 #endif
 
-const TArray<FName>& FRigVMDispatch_SwitchInt32::GetControlFlowBlocks_Impl(const FRigVMDispatchContext& InContext) const
+#if WITH_EDITORONLY_DATA
+const TArray<FName>& FRigVMDispatch_SwitchInt32::GetControlFlowBlocks_Impl(const FRigVMDispatchContext& InContext, const FRigVMTemplateTypeMap& InTypeMap) const
 {
 	static const TArray<FName> DefaultBlocks = {FRigVMStruct::ControlFlowCompletedName};
 	
@@ -133,6 +135,7 @@ const TArray<FName>& FRigVMDispatch_SwitchInt32::GetControlFlowBlocks_Impl(const
 
 	return DefaultBlocks;
 }
+#endif
 
 FRigVMFunctionPtr FRigVMDispatch_SwitchInt32::GetDispatchFunctionImpl(const FRigVMTemplateTypeMap& InTypes, const FRigVMRegistryHandle& InRegistry) const
 {
@@ -159,7 +162,7 @@ void FRigVMDispatch_SwitchInt32::Execute(FRigVMExtendedExecuteContext& InContext
 	}
 }
 
-FName FRigVMDispatch_SwitchInt32::GetCaseName(int32 InIndex)
+FName FRigVMDispatch_SwitchBase::GetCaseName(int32 InIndex)
 {
 	using MapType = TMap<int32, FName>;
 	UE_AUTORTFM_DECLARE_THREAD_LOCAL_VAR(MapType, NamesFromInt);
@@ -190,4 +193,166 @@ FName FRigVMDispatch_SwitchInt32::GetCaseDisplayName(int32 InIndex)
     return Result;
 }
 
+
+const TArray<FRigVMTemplateArgumentInfo>& FRigVMDispatch_SwitchEnum::GetArgumentInfos(FRigVMRegistryHandle& InRegistry) const
+{
+	if(CachedArgumentInfos.IsEmpty())
+	{
+		static const TArray<FRigVMTemplateArgument::ETypeCategory> EnumCategories = {
+			FRigVMTemplateArgument::ETypeCategory_SingleEnumValue
+		};
+		
+		CachedArgumentInfos.Emplace(IndexName, ERigVMPinDirection::Input, EnumCategories);
+		CachedArgumentInfos.Emplace(FRigVMStruct::ControlFlowBlockToRunName, ERigVMPinDirection::Hidden, RigVMTypeUtils::TypeIndex::FName);
+	}
+	return CachedArgumentInfos;
+}
+
+FRigVMTemplateTypeMap FRigVMDispatch_SwitchEnum::OnNewArgumentType(const FName& InArgumentName, TRigVMTypeIndex InTypeIndex, FRigVMRegistryHandle& InRegistry) const
+{
+	return {
+			{IndexName, InTypeIndex},
+			{FRigVMStruct::ControlFlowBlockToRunName, RigVMTypeUtils::TypeIndex::FName}
+	};
+}
+
+#if WITH_EDITOR
+
+FName FRigVMDispatch_SwitchEnum::GetDisplayNameForArgument(const FName& InArgumentName, const FRigVMTemplateTypeMap& InTypes) const
+{
+	static const FString CasesPrefix = CasesName.ToString() + TEXT(".");
+	const FString ArgumentNameString = InArgumentName.ToString();
+	if(ArgumentNameString.StartsWith(CasesPrefix))
+	{
+		FString Left, Right;
+		verify(RigVMStringUtils::SplitPinPathAtEnd(ArgumentNameString, Left, Right));
+		if(Right.IsNumeric())
+		{
+			UEnum* Enum = nullptr;
+			if (const TRigVMTypeIndex* IndexType = InTypes.Find(IndexName))
+			{
+				FRigVMTemplateArgumentType Type = FRigVMRegistry::Get().GetType(*IndexType);
+				if (UEnum* TypeEnum = Cast<UEnum>(Type.CPPTypeObject))
+				{
+					Enum = TypeEnum;
+				}
+			}
+			const int32 CaseIndex = FCString::Atoi(*Right);
+			return GetCaseDisplayName(CaseIndex, Enum);
+		}
+	}
+	return FRigVMDispatch_CoreBase::GetDisplayNameForArgument(InArgumentName, InTypes);
+}
+
+#endif
+
+#if WITH_EDITORONLY_DATA
+const TArray<FName>& FRigVMDispatch_SwitchEnum::GetControlFlowBlocks_Impl(const FRigVMDispatchContext& InContext, const FRigVMTemplateTypeMap& InTypeMap) const
+{
+	static const TArray<FName> DefaultBlocks = {FRigVMStruct::ControlFlowCompletedName};
+
+	UEnum* Enum = nullptr;
+	if (const TRigVMTypeIndex* IndexTypeIndex = InTypeMap.Find(IndexName))
+	{
+		const FRigVMTemplateArgumentType& IndexType = FRigVMRegistry::Get().GetType(*IndexTypeIndex);
+		if (UEnum* IndexEnum = Cast<UEnum>(IndexType.CPPTypeObject))
+		{
+			Enum = IndexEnum;
+		}
+	}
+	
+	if (Enum)
+	{
+		static TMap<FString, TArray<FName>> BlockCache;
+		if(const TArray<FName>* ExistingBlocks = BlockCache.Find(Enum->GetName()))
+		{
+			return *ExistingBlocks;
+		}
+		TArray<FName>& Blocks = BlockCache.Add(Enum->GetName());
+		
+		const int32 NumCases = Enum->NumEnums() - 1; // Exclude 'MAX' enum value.
+		for (int32 i=0; i<NumCases; ++i)
+		{
+			if (Enum->HasMetaData(TEXT("Hidden"), i))
+			{
+				continue;
+			}
+			
+			Blocks.Add(GetCaseName(i));
+		}
+		Blocks.Append(DefaultBlocks);
+		return Blocks;
+	}
+
+	return DefaultBlocks;
+}
+#endif
+
+FRigVMFunctionPtr FRigVMDispatch_SwitchEnum::GetDispatchFunctionImpl(const FRigVMTemplateTypeMap& InTypes, const FRigVMRegistryHandle& InRegistry) const
+{
+	return &FRigVMDispatch_SwitchEnum::Execute;
+}
+
+void FRigVMDispatch_SwitchEnum::Execute(FRigVMExtendedExecuteContext& InContext, FRigVMMemoryHandleArray Handles, FRigVMPredicateBranchArray Predicates)
+{
+#if WITH_EDITOR
+	check(CastField<FEnumProperty>(Handles[0].GetProperty()) || CastField<FByteProperty>(Handles[0].GetProperty()));
+	check(Handles[1].IsName());
+#endif
+	
+	int32 Index = INDEX_NONE;
+	{
+		const FNumericProperty* NumericProperty = nullptr;
+		const FProperty* IndexProperty = Handles[0].GetProperty();
+		UEnum* Enum = nullptr;
+		if(const FEnumProperty* EnumProperty = CastField<FEnumProperty>(IndexProperty))
+		{
+			Enum = EnumProperty->GetEnum();
+			NumericProperty = EnumProperty->GetUnderlyingProperty();
+		}
+		else if(const FByteProperty* ByteProperty = CastField<FByteProperty>(IndexProperty))
+		{
+			Enum = ByteProperty->GetIntPropertyEnum();
+			NumericProperty = ByteProperty;
+		}
+		check(NumericProperty && NumericProperty->IsInteger());
+		check(Enum);
+		
+		const uint8* ValuePtr = Handles[0].GetInputData();
+		const int64 Value = NumericProperty->GetSignedIntPropertyValue(ValuePtr);
+		
+		Index = Enum->GetIndexByValue(Value);
+	}
+
+	FName& BlockToRun = *reinterpret_cast<FName*>(Handles[1].GetOutputData());
+
+	if(BlockToRun.IsNone() && Index != INDEX_NONE)
+	{
+		BlockToRun = GetCaseName(Index);
+	}
+	else
+	{
+		BlockToRun = FRigVMStruct::ControlFlowCompletedName;
+	}
+}
+
+#if WITH_EDITORONLY_DATA
+FName FRigVMDispatch_SwitchEnum::GetCaseDisplayName(int32 InIndex, const UEnum* InEnum)
+{
+	using MapType = TMap<int32, FName>;
+	UE_AUTORTFM_DECLARE_THREAD_LOCAL_VAR(MapType, NamesFromInt);
+
+	FName Result = NAME_None;
+	if (InEnum)
+	{
+		// The input index represents the index of visible values,
+		// so we need to find the actual real index
+		Result = *InEnum->GetDisplayNameTextByIndex(InIndex).ToString();
+	}
+    return Result;
+}
+
+#endif
+
+
 #undef LOCTEXT_NAMESPACE
diff --git a/Engine/Plugins/Runtime/RigVM/Source/RigVM/Public/RigVMCore/RigVMDispatchFactory.h b/Engine/Plugins/Runtime/RigVM/Source/RigVM/Public/RigVMCore/RigVMDispatchFactory.h
--- a/Engine/Plugins/Runtime/RigVM/Source/RigVM/Public/RigVMCore/RigVMDispatchFactory.h
+++ b/Engine/Plugins/Runtime/RigVM/Source/RigVM/Public/RigVMCore/RigVMDispatchFactory.h
@@ -113,7 +113,7 @@ public:
 	RIGVM_API virtual FName GetNextAggregateName(const FName& InLastAggregatePinName) const;
 
 	// Returns the display name text for an argument 
-	RIGVM_API virtual FName GetDisplayNameForArgument(const FName& InArgumentName) const;
+	RIGVM_API virtual FName GetDisplayNameForArgument(const FName& InArgumentName, const FRigVMTemplateTypeMap& InTypes) const;
 
 	// Returns meta data on the property of the permutations 
 	RIGVM_API virtual FString GetArgumentMetaData(const FName& InArgumentName, const FName& InMetaDataKey) const;
@@ -140,11 +140,13 @@ public:
 	// arguments get unrolled.
 	RIGVM_API virtual FName GetArgumentNameForOperandIndex(int32 InOperandIndex, int32 InTotalOperands, FRigVMRegistryHandle& InRegistry) const;
 
+#if WITH_EDITORONLY_DATA
 	// returns true if the dispatch is a control flow dispatch
-	bool IsControlFlowDispatch(const FRigVMDispatchContext& InContext) const { return !GetControlFlowBlocks(InContext).IsEmpty(); }
+	bool IsControlFlowDispatch(const FRigVMDispatchContext& InContext, const FRigVMTemplateTypeMap& InTypeMap) const { return !GetControlFlowBlocks(InContext, InTypeMap).IsEmpty(); }
 
 	// returns the control flow blocks of this dispatch
-	RIGVM_API const TArray<FName>& GetControlFlowBlocks(const FRigVMDispatchContext& InContext) const;
+	RIGVM_API const TArray<FName>& GetControlFlowBlocks(const FRigVMDispatchContext& InContext, const FRigVMTemplateTypeMap& InTypeMap) const;
+#endif
 
 	// returns true if a given control flow block needs to be sliced
 	virtual const bool IsControlFlowBlockSliced(const FName& InBlockName) const { return false; }
@@ -274,8 +276,10 @@ protected:
 		return bCondition;
 	}
 #endif
-
-	RIGVM_API virtual const TArray<FName>& GetControlFlowBlocks_Impl(const FRigVMDispatchContext& InContext) const;
+	
+#if WITH_EDITORONLY_DATA
+	RIGVM_API virtual const TArray<FName>& GetControlFlowBlocks_Impl(const FRigVMDispatchContext& InContext, const FRigVMTemplateTypeMap& InTypeMap) const;
+#endif
 
 	static const inline TCHAR* DispatchPrefix = TEXT("DISPATCH_");
 	static const inline TCHAR* TrueString = TEXT("True");
diff --git a/Engine/Plugins/Runtime/RigVM/Source/RigVM/Public/RigVMCore/RigVMStruct.h b/Engine/Plugins/Runtime/RigVM/Source/RigVM/Public/RigVMCore/RigVMStruct.h
--- a/Engine/Plugins/Runtime/RigVM/Source/RigVM/Public/RigVMCore/RigVMStruct.h
+++ b/Engine/Plugins/Runtime/RigVM/Source/RigVM/Public/RigVMCore/RigVMStruct.h
@@ -205,10 +205,12 @@ public:
 	virtual void Execute() {}
 
 	// control flow related
+	virtual int32 GetNumSlices() const { return 1; }
+#if WITH_EDITORONLY_DATA
 	RIGVM_API bool IsForLoop() const;
 	RIGVM_API bool IsControlFlowNode() const; 
-	virtual int32 GetNumSlices() const { return 1; }
 	RIGVM_API const TArray<FName>& GetControlFlowBlocks() const;
+#endif
 	virtual const bool IsControlFlowBlockSliced(const FName& InBlockName) const { return false; }
 
 	// node creation
@@ -304,7 +306,10 @@ protected:
 	RIGVM_API TMap<FName, FString> GetDefaultValues(UScriptStruct* InScriptStruct) const;
 	RIGVM_API bool ApplyUpgradeInfo(const FRigVMStructUpgradeInfo& InUpgradeInfo);
 	virtual TArray<FRigVMUserWorkflow> GetSupportedWorkflows(const UObject* InSubject) const { return TArray<FRigVMUserWorkflow>(); }
+	
+#if WITH_EDITORONLY_DATA
 	RIGVM_API virtual const TArray<FName>& GetControlFlowBlocks_Impl() const;
+#endif
 
 #if WITH_EDITOR
 	RIGVM_API static void ValidateControlFlowBlocks(const TArray<FName>& InBlocks);
diff --git a/Engine/Plugins/Runtime/RigVM/Source/RigVM/Public/RigVMFunctions/Execution/RigVMFunction_ForLoop.h b/Engine/Plugins/Runtime/RigVM/Source/RigVM/Public/RigVMFunctions/Execution/RigVMFunction_ForLoop.h
--- a/Engine/Plugins/Runtime/RigVM/Source/RigVM/Public/RigVMFunctions/Execution/RigVMFunction_ForLoop.h
+++ b/Engine/Plugins/Runtime/RigVM/Source/RigVM/Public/RigVMFunctions/Execution/RigVMFunction_ForLoop.h
@@ -22,11 +22,13 @@ struct FRigVMFunction_ForLoopCount : public FRigVMStructMutable
 	}
 
 	// FRigVMStruct overrides
+#if WITH_EDITORONLY_DATA
 	virtual const TArray<FName>& GetControlFlowBlocks_Impl() const override
 	{
 		static const TArray<FName> Blocks = {ExecutePinName, ForLoopCompletedPinName};
 		return Blocks;
 	}
+#endif
 	virtual const bool IsControlFlowBlockSliced(const FName& InBlockName) const { return InBlockName == ExecutePinName; }
 	virtual int32 GetNumSlices() const override { return Count; }
 
diff --git a/Engine/Plugins/Runtime/RigVM/Source/RigVM/Public/RigVMFunctions/RigVMDispatch_Array.h b/Engine/Plugins/Runtime/RigVM/Source/RigVM/Public/RigVMFunctions/RigVMDispatch_Array.h
--- a/Engine/Plugins/Runtime/RigVM/Source/RigVM/Public/RigVMFunctions/RigVMDispatch_Array.h
+++ b/Engine/Plugins/Runtime/RigVM/Source/RigVM/Public/RigVMFunctions/RigVMDispatch_Array.h
@@ -502,7 +502,9 @@ public:
 	RIGVM_API virtual FName GetArgumentNameForOperandIndex(int32 InOperandIndex, int32 InTotalOperands, FRigVMRegistryHandle& InRegistry) const override;
 	RIGVM_API virtual const TArray<FRigVMTemplateArgumentInfo>& GetArgumentInfos(FRigVMRegistryHandle& InRegistry) const override;
 	RIGVM_API virtual const TArray<FRigVMExecuteArgument>& GetExecuteArguments_Impl(const FRigVMDispatchContext& InContext) const override;
-	RIGVM_API virtual const TArray<FName>& GetControlFlowBlocks_Impl(const FRigVMDispatchContext& InContext) const override;
+#if WITH_EDITORONLY_DATA
+	RIGVM_API virtual const TArray<FName>& GetControlFlowBlocks_Impl(const FRigVMDispatchContext& InContext, const FRigVMTemplateTypeMap& InTypeMap) const override;
+#endif
 	RIGVM_API virtual const bool IsControlFlowBlockSliced(const FName& InBlockName) const override;
 	RIGVM_API virtual FRigVMTemplateTypeMap OnNewArgumentType(const FName& InArgumentName, TRigVMTypeIndex InTypeIndex, FRigVMRegistryHandle& InRegistry) const override;
 #if WITH_EDITOR
diff --git a/Engine/Plugins/Runtime/RigVM/Source/RigVM/Public/RigVMFunctions/RigVMDispatch_Switch.h b/Engine/Plugins/Runtime/RigVM/Source/RigVM/Public/RigVMFunctions/RigVMDispatch_Switch.h
--- a/Engine/Plugins/Runtime/RigVM/Source/RigVM/Public/RigVMFunctions/RigVMDispatch_Switch.h
+++ b/Engine/Plugins/Runtime/RigVM/Source/RigVM/Public/RigVMFunctions/RigVMDispatch_Switch.h
@@ -6,11 +6,36 @@
 #include "RigVMDispatch_Core.h"
 #include "RigVMDispatch_Switch.generated.h"
 
+USTRUCT(meta=(Abstract, Category = "Switch", NodeColor = "0.762745, 1,0, 0.329412"))
+struct FRigVMDispatch_SwitchBase : public FRigVMDispatch_CoreBase
+{
+	GENERATED_BODY()
+
+public:
+
+	RIGVM_API virtual FName GetArgumentNameForOperandIndex(int32 InOperandIndex, int32 InTotalOperands, FRigVMRegistryHandle& InRegistry) const override;
+	RIGVM_API virtual const TArray<FRigVMExecuteArgument>& GetExecuteArguments_Impl(const FRigVMDispatchContext& InContext) const override;
+	virtual bool IsSingleton() const override { return true; } 
+
+#if WITH_EDITOR
+	RIGVM_API virtual FString GetArgumentMetaData(const FName& InArgumentName, const FName& InMetaDataKey) const override;
+	RIGVM_API virtual FString GetArgumentDefaultValue(const FName& InArgumentName, TRigVMTypeIndex InTypeIndex) const override;
+	RIGVM_API virtual FText GetArgumentTooltip(const FName& InArgumentName, TRigVMTypeIndex InTypeIndex) const override;
+#endif
+	
+protected:
+
+	static RIGVM_API FName GetCaseName(int32 InIndex);
+
+	static inline const FLazyName IndexName = FLazyName(TEXT("Index"));
+	static inline const FLazyName CasesName = FLazyName(TEXT("Cases"));
+};
+
 /*
  * Run a branch based on an integer index
  */
 USTRUCT(meta=(DisplayName = "Switch", Category = "Execution", Keywords = "Case", NodeColor = "0,1,0,1"))
-struct FRigVMDispatch_SwitchInt32 : public FRigVMDispatch_CoreBase
+struct FRigVMDispatch_SwitchInt32 : public FRigVMDispatch_SwitchBase
 {
 	GENERATED_BODY()
 
@@ -21,27 +46,57 @@ public:
 		FactoryScriptStruct = StaticStruct();
 	}
 
-	RIGVM_API virtual FName GetArgumentNameForOperandIndex(int32 InOperandIndex, int32 InTotalOperands, FRigVMRegistryHandle& InRegistry) const override;
 	RIGVM_API virtual const TArray<FRigVMTemplateArgumentInfo>& GetArgumentInfos(FRigVMRegistryHandle& InRegistry) const override;
-	RIGVM_API virtual const TArray<FRigVMExecuteArgument>& GetExecuteArguments_Impl(const FRigVMDispatchContext& InContext) const override;
-	virtual bool IsSingleton() const override { return true; } 
 
 #if WITH_EDITOR
-	RIGVM_API virtual FString GetArgumentMetaData(const FName& InArgumentName, const FName& InMetaDataKey) const override;
-	RIGVM_API virtual FString GetArgumentDefaultValue(const FName& InArgumentName, TRigVMTypeIndex InTypeIndex) const override;
-	RIGVM_API virtual FName GetDisplayNameForArgument(const FName& InArgumentName) const override;
-	RIGVM_API virtual FText GetArgumentTooltip(const FName& InArgumentName, TRigVMTypeIndex InTypeIndex) const override;
+	RIGVM_API virtual FName GetDisplayNameForArgument(const FName& InArgumentName, const FRigVMTemplateTypeMap& InTypes) const override;
 #endif
 	
-	RIGVM_API virtual const TArray<FName>& GetControlFlowBlocks_Impl(const FRigVMDispatchContext& InContext) const override;
+#if WITH_EDITORONLY_DATA
+	RIGVM_API virtual const TArray<FName>& GetControlFlowBlocks_Impl(const FRigVMDispatchContext& InContext, const FRigVMTemplateTypeMap& InTypeMap) const override;
+#endif
 
 protected:
 
 	RIGVM_API virtual FRigVMFunctionPtr GetDispatchFunctionImpl(const FRigVMTemplateTypeMap& InTypes, const FRigVMRegistryHandle& InRegistry) const override;
 	static RIGVM_API void Execute(FRigVMExtendedExecuteContext& InContext, FRigVMMemoryHandleArray Handles, FRigVMPredicateBranchArray Predicates);
-	static RIGVM_API FName GetCaseName(int32 InIndex);
 	static RIGVM_API FName GetCaseDisplayName(int32 InIndex);
+};
 
-	static inline const FLazyName IndexName = FLazyName(TEXT("Index"));
-	static inline const FLazyName CasesName = FLazyName(TEXT("Cases"));
+/*
+ * Run a branch based on an enum index
+ */
+USTRUCT(meta=(DisplayName = "Switch Enum", Category = "Execution", Keywords = "Case", NodeColor = "0,1,0,1"))
+struct FRigVMDispatch_SwitchEnum : public FRigVMDispatch_SwitchBase
+{
+	GENERATED_BODY()
+
+public:
+
+	FRigVMDispatch_SwitchEnum()
+	{
+		FactoryScriptStruct = StaticStruct();
+	}
+
+	RIGVM_API virtual const TArray<FRigVMTemplateArgumentInfo>& GetArgumentInfos(FRigVMRegistryHandle& InRegistry) const override;
+	RIGVM_API virtual FRigVMTemplateTypeMap OnNewArgumentType(const FName& InArgumentName, TRigVMTypeIndex InTypeIndex, FRigVMRegistryHandle& InRegistry) const override;
+
+#if WITH_EDITOR
+	RIGVM_API virtual FName GetDisplayNameForArgument(const FName& InArgumentName, const FRigVMTemplateTypeMap& InTypes) const override;
+#endif
+	
+#if WITH_EDITORONLY_DATA
+	RIGVM_API virtual const TArray<FName>& GetControlFlowBlocks_Impl(const FRigVMDispatchContext& InContext, const FRigVMTemplateTypeMap& InTypeMap) const override;
+#endif
+
+protected:
+
+	RIGVM_API virtual FRigVMFunctionPtr GetDispatchFunctionImpl(const FRigVMTemplateTypeMap& InTypes, const FRigVMRegistryHandle& InRegistry) const override;
+	static RIGVM_API void Execute(FRigVMExtendedExecuteContext& InContext, FRigVMMemoryHandleArray Handles, FRigVMPredicateBranchArray Predicates);
+	
+#if WITH_EDITORONLY_DATA
+	static RIGVM_API FName GetCaseDisplayName(int32 InIndex, const UEnum* InEnum = nullptr);
+#endif
+	
+	friend class URigVMController;
 };
diff --git a/Engine/Plugins/Runtime/RigVM/Source/RigVM/Public/RigVMFunctions/RigVMFunction_ControlFlow.h b/Engine/Plugins/Runtime/RigVM/Source/RigVM/Public/RigVMFunctions/RigVMFunction_ControlFlow.h
--- a/Engine/Plugins/Runtime/RigVM/Source/RigVM/Public/RigVMFunctions/RigVMFunction_ControlFlow.h
+++ b/Engine/Plugins/Runtime/RigVM/Source/RigVM/Public/RigVMFunctions/RigVMFunction_ControlFlow.h
@@ -31,6 +31,7 @@ struct FRigVMFunction_ControlFlowBranch : public FRigVMFunction_ControlFlowBase
 	RIGVM_METHOD()
 	RIGVM_API virtual void Execute();
 
+#if WITH_EDITORONLY_DATA
 	virtual const TArray<FName>& GetControlFlowBlocks_Impl() const override
 	{
 		static const TArray<FName> Blocks = {
@@ -40,6 +41,7 @@ struct FRigVMFunction_ControlFlowBranch : public FRigVMFunction_ControlFlowBase
 		};
 		return Blocks;
 	}
+#endif
 
 	// The input execution pin to hook up to the graph
 	UPROPERTY(meta=(Input, DisplayName="Execute"))
@@ -83,6 +85,7 @@ struct FRigVMFunction_ControlFlowRunOnce : public FRigVMFunction_ControlFlowBase
 	RIGVM_METHOD()
 	RIGVM_API virtual void Execute();
 
+#if WITH_EDITORONLY_DATA
 	virtual const TArray<FName>& GetControlFlowBlocks_Impl() const override
 	{
 		static const TArray<FName> Blocks = {
@@ -91,6 +94,7 @@ struct FRigVMFunction_ControlFlowRunOnce : public FRigVMFunction_ControlFlowBase
 		};
 		return Blocks;
 	}
+#endif
 
 	// The input execution pin to hook up to the graph
 	UPROPERTY(meta=(Input, DisplayName="Execute"))
@@ -132,6 +136,7 @@ struct FRigVMFunction_ControlFlowTimer : public FRigVMFunction_ControlFlowBase
 	RIGVM_METHOD()
 	RIGVM_API virtual void Execute();
 
+#if WITH_EDITORONLY_DATA
 	virtual const TArray<FName>& GetControlFlowBlocks_Impl() const override
 	{
 		static const TArray<FName> Blocks = {
@@ -142,6 +147,7 @@ struct FRigVMFunction_ControlFlowTimer : public FRigVMFunction_ControlFlowBase
 		};
 		return Blocks;
 	}
+#endif
 
 	// The input execution pin to hook up to the graph
 	UPROPERTY(meta=(Input, DisplayName="Execute"))
diff --git a/Engine/Plugins/Runtime/RigVM/Source/RigVMDeveloper/Private/RigVMModel/Nodes/RigVMDispatchNode.cpp b/Engine/Plugins/Runtime/RigVM/Source/RigVMDeveloper/Private/RigVMModel/Nodes/RigVMDispatchNode.cpp
--- a/Engine/Plugins/Runtime/RigVM/Source/RigVMDeveloper/Private/RigVMModel/Nodes/RigVMDispatchNode.cpp
+++ b/Engine/Plugins/Runtime/RigVM/Source/RigVMDeveloper/Private/RigVMModel/Nodes/RigVMDispatchNode.cpp
@@ -67,7 +67,7 @@ const TArray<FName>& URigVMDispatchNode::GetControlFlowBlocks() const
 {
 	if (const FRigVMDispatchFactory* Factory = GetFactory())
 	{
-		return Factory->GetControlFlowBlocks(GetDispatchContext());
+		return Factory->GetControlFlowBlocks(GetDispatchContext(), GetTemplatePinTypeMap());
 	}
 	return Super::GetControlFlowBlocks();
 }
@@ -283,3 +283,12 @@ FRigVMStructUpgradeInfo URigVMDispatchNode::GetUpgradeInfo() const
 	return FRigVMStructUpgradeInfo();
 }
 
+FName URigVMDispatchNode::GetDisplayNameForPin(const URigVMPin* InPin) const
+{
+	if(const FRigVMDispatchFactory* Factory = GetFactory())
+	{
+		return Factory->GetDisplayNameForArgument(*InPin->GetSubPinPath(InPin->GetRootPin(), true), GetTemplatePinTypeMap());
+	}
+	return Super::GetDisplayNameForPin(InPin);
+}
+
diff --git a/Engine/Plugins/Runtime/RigVM/Source/RigVMDeveloper/Private/RigVMModel/Nodes/RigVMTemplateNode.cpp b/Engine/Plugins/Runtime/RigVM/Source/RigVMDeveloper/Private/RigVMModel/Nodes/RigVMTemplateNode.cpp
--- a/Engine/Plugins/Runtime/RigVM/Source/RigVMDeveloper/Private/RigVMModel/Nodes/RigVMTemplateNode.cpp
+++ b/Engine/Plugins/Runtime/RigVM/Source/RigVMDeveloper/Private/RigVMModel/Nodes/RigVMTemplateNode.cpp
@@ -721,7 +721,7 @@ FName URigVMTemplateNode::GetDisplayNameForPin(const URigVMPin* InPin, const TAr
 
 FName URigVMTemplateNode::GetDisplayNameForPin(const URigVMPin* InPin) const
 {
-	return GetDisplayNameForPin(InPin, TArray<int32>());
+	return GetDisplayNameForPin(InPin, GetResolvedPermutationIndices(true));
 }
 
 TRigVMTypeIndex URigVMTemplateNode::TryReduceTypesToSingle(const TArray<TRigVMTypeIndex>& InTypes, const TRigVMTypeIndex PreferredType) const
diff --git a/Engine/Plugins/Runtime/RigVM/Source/RigVMDeveloper/Private/RigVMModel/RigVMController.cpp b/Engine/Plugins/Runtime/RigVM/Source/RigVMDeveloper/Private/RigVMModel/RigVMController.cpp
--- a/Engine/Plugins/Runtime/RigVM/Source/RigVMDeveloper/Private/RigVMModel/RigVMController.cpp
+++ b/Engine/Plugins/Runtime/RigVM/Source/RigVMDeveloper/Private/RigVMModel/RigVMController.cpp
@@ -34,6 +34,7 @@
 #include "StructUtils/UserDefinedStruct.h"
 #include "RigVMFunctions/RigVMDispatch_If.h"
 #include "RigVMFunctions/RigVMDispatch_Select.h"
+#include "RigVMFunctions/RigVMDispatch_Switch.h"
 #include "RigVMModel/RigVMClient.h"
 #include "RigVMModel/Nodes/RigVMBranchNode.h"
 #include "RigVMModel/Nodes/RigVMArrayNode.h"
@@ -20964,7 +20965,7 @@ bool URigVMController::FullyResolveTemplateNode(URigVMTemplateNode* InNode, int3
 
 				if(Factory)
 				{
-					const FName DisplayNameText = Factory->GetDisplayNameForArgument(MissingPin.GetName());
+					const FName DisplayNameText = Factory->GetDisplayNameForArgument(MissingPin.GetName(), InNode->GetTemplatePinTypeMap(true));
 					if(!DisplayNameText.IsNone())
 					{
 						Pin->DisplayName = *DisplayNameText.ToString();
@@ -21161,6 +21162,8 @@ bool URigVMController::ResolveWildCardPin(URigVMPin* InPin, TRigVMTypeIndex InTy
 		Action.SetTitle(TEXT("Resolve Wildcard Pin"));
 		GetActionStack()->BeginAction(Action);
 	}
+	
+	FRigVMControllerCompileBracketScope CompileScope(this);
 
 	if (!InPin->IsWildCard())
 	{
@@ -21180,6 +21183,52 @@ bool URigVMController::ResolveWildCardPin(URigVMPin* InPin, TRigVMTypeIndex InTy
 	}
 
 	UpdateTemplateNodePinTypes(TemplateNode, bSetupUndoRedo);
+	
+	if(const URigVMDispatchNode* DispatchNode = Cast<URigVMDispatchNode>(TemplateNode))
+	{
+		if (DispatchNode->GetFactory()->GetScriptStruct() == FRigVMDispatch_SwitchEnum::StaticStruct())
+		{
+			URigVMPin* IndexPin = DispatchNode->FindPin(FRigVMDispatch_SwitchEnum::IndexName.ToString());
+			URigVMPin* CasesPin = DispatchNode->FindPin(FRigVMDispatch_SwitchEnum::CasesName.ToString());
+			if (IndexPin && CasesPin)
+			{
+				if (UEnum* Enum = Cast<UEnum>(IndexPin->GetCPPTypeObject()))
+				{
+					const int32 NumCases = Enum->NumEnums() - 1; // Exclude 'MAX' enum value.
+					int32 SubPinIndex = 0;
+					
+					ClearArrayPin(CasesPin->GetPinPath(), bSetupUndoRedo);
+					CasesPin->GetSubPins()[0]->Rename(TEXT("OldName"), nullptr); // Make sure no subpins will have a name clash
+					TArray<FString> Names;
+					Names.Reserve(NumCases);
+					for (int32 i=0; i<NumCases; ++i)
+					{
+						if (Enum->HasMetaData(TEXT("Hidden"), i))
+						{
+							continue;
+						}
+					
+						if (SubPinIndex > 0)
+						{
+							AddArrayPin(CasesPin->GetPinPath(), FString(), bSetupUndoRedo);
+						}
+						Names.Add(FString::FromInt(i));
+						SubPinIndex++;
+					}
+					
+					for (int32 i=Names.Num()-1; i>= 0; --i)
+					{
+						URigVMPin* SubPin = CasesPin->GetSubPins()[i];
+						SubPin->Rename(*Names[i], nullptr);
+						SubPin->CachedDisplayName.ResetCachedValue();
+						SubPin->CachedPinPath.ResetCachedValue();
+						SubPin->CachedPinPathWithNodePath.ResetCachedValue();
+						Notify(ERigVMGraphNotifType::PinDisplayNameChanged, SubPin);
+					}
+				}
+			}
+		}
+	}
 
 	if (bSetupUndoRedo)
 	{
diff --git a/Engine/Plugins/Runtime/RigVM/Source/RigVMDeveloper/Private/RigVMModel/RigVMPin.cpp b/Engine/Plugins/Runtime/RigVM/Source/RigVMDeveloper/Private/RigVMModel/RigVMPin.cpp
--- a/Engine/Plugins/Runtime/RigVM/Source/RigVMDeveloper/Private/RigVMModel/RigVMPin.cpp
+++ b/Engine/Plugins/Runtime/RigVM/Source/RigVMDeveloper/Private/RigVMModel/RigVMPin.cpp
@@ -467,12 +467,6 @@ FName URigVMPin::GetDisplayName() const
 
 	if (DisplayName == NAME_None)
 	{
-		if(IsArrayElement())
-		{
-			CachedDisplayName = *FString::FromInt(GetPinIndex());
-			return CachedDisplayName.GetValue();
-		}
-		
 		if(const URigVMNode* Node = GetNode())
 		{
 			const FName DisplayNameFromNode = Node->GetDisplayNameForPin(this);
diff --git a/Engine/Plugins/Runtime/RigVM/Source/RigVMDeveloper/Public/RigVMModel/Nodes/RigVMDispatchNode.h b/Engine/Plugins/Runtime/RigVM/Source/RigVMDeveloper/Public/RigVMModel/Nodes/RigVMDispatchNode.h
--- a/Engine/Plugins/Runtime/RigVM/Source/RigVMDeveloper/Public/RigVMModel/Nodes/RigVMDispatchNode.h
+++ b/Engine/Plugins/Runtime/RigVM/Source/RigVMDeveloper/Public/RigVMModel/Nodes/RigVMDispatchNode.h
@@ -48,6 +48,7 @@ public:
 	UE_API virtual TArray<URigVMPin*> GetAggregateOutputs() const override;
 	UE_API virtual FName GetNextAggregateName(const FName& InLastAggregatePinName) const override;
 	UE_API virtual FRigVMStructUpgradeInfo GetUpgradeInfo() const override;
+	UE_API virtual FName GetDisplayNameForPin(const URigVMPin* InPin) const override;
 
 	UE_API const FRigVMDispatchFactory* GetFactory() const;
 	UE_API virtual bool IsOutDated() const override;
diff --git a/Engine/Plugins/Runtime/RigVM/Source/RigVMEditor/Private/RigVMEditorModule.cpp b/Engine/Plugins/Runtime/RigVM/Source/RigVMEditor/Private/RigVMEditorModule.cpp
--- a/Engine/Plugins/Runtime/RigVM/Source/RigVMEditor/Private/RigVMEditorModule.cpp
+++ b/Engine/Plugins/Runtime/RigVM/Source/RigVMEditor/Private/RigVMEditorModule.cpp
@@ -46,6 +46,7 @@
 #include "Toolkits/ToolkitManager.h"
 #include "ToolMenus.h"
 #include "Editor/RigVMPythonLogDetails.h"
+#include "RigVMFunctions/RigVMDispatch_Switch.h"
 #include "RigVMFunctions/Execution/RigVMFunction_Sequence.h"
 #include "RigVMModel/RigVMControllerActions.h"
 #include "Widgets/SRigVMSwapFunctionsWidget.h"
@@ -1872,6 +1873,14 @@ void FRigVMEditorModule::GetPinDebugContextMenuActions(TScriptInterface<IRigVMCl
 
 void FRigVMEditorModule::GetPinArrayContextMenuActions(TScriptInterface<IRigVMClientHost> RigVMClientHost, const UEdGraphPin* EdGraphPin, URigVMPin* ModelPin, UToolMenu* Menu) const
 {
+	if (URigVMDispatchNode* DispatchNode = Cast<URigVMDispatchNode>(ModelPin->GetNode()))
+	{
+		if (DispatchNode->GetFactory()->GetScriptStruct() == FRigVMDispatch_SwitchEnum::StaticStruct())
+		{
+			return;
+		}
+	}
+	
 	URigVMController* Controller = RigVMClientHost->GetRigVMClient()->GetController(ModelPin->GetGraph());
 	if (ModelPin->IsArray() && !ModelPin->IsExecuteContext())
 	{
diff --git a/Engine/Plugins/Runtime/RigVM/Source/RigVMEditor/Private/Widgets/SRigVMGraphNode.cpp b/Engine/Plugins/Runtime/RigVM/Source/RigVMEditor/Private/Widgets/SRigVMGraphNode.cpp
--- a/Engine/Plugins/Runtime/RigVM/Source/RigVMEditor/Private/Widgets/SRigVMGraphNode.cpp
+++ b/Engine/Plugins/Runtime/RigVM/Source/RigVMEditor/Private/Widgets/SRigVMGraphNode.cpp
@@ -44,6 +44,7 @@
 #include "Widgets/SRigVMVariantTagWidget.h"
 #include "Widgets/SToolTip.h"
 #include "Overrides/OverrideStatusWidgetMenuBuilder.h"
+#include "RigVMFunctions/RigVMDispatch_Switch.h"
 
 #if WITH_EDITOR
 #include "Editor.h"
@@ -623,6 +624,14 @@ void SRigVMGraphNode::CreateAddPinButton()
 		AddPinPadding.Left -= 2.0f;
 		EHorizontalAlignment HorizontalAlignment = HAlign_Left;
 		
+		if (URigVMDispatchNode* DispatchNode = Cast<URigVMDispatchNode>(ModelNode.Get()))
+		{
+			if (DispatchNode->GetFactory()->GetScriptStruct() == FRigVMDispatch_SwitchEnum::StaticStruct())
+			{
+				return;
+			}
+		}
+		
 		if(ModelNode->IsAggregate() || ModelNode->IsA<URigVMAggregateNode>())
 		{
 			const bool bInputAggregate = ModelNode->IsInputAggregate();
