Fixed Control Rig AST generation, to avoid function entry variables being included in scope blocks

#Synced-CL 50592216

diff --git a/Engine/Plugins/Runtime/RigVM/Source/RigVMDeveloper/Private/RigVMCompiler/RigVMAST.cpp b/Engine/Plugins/Runtime/RigVM/Source/RigVMDeveloper/Private/RigVMCompiler/RigVMAST.cpp
--- a/Engine/Plugins/Runtime/RigVM/Source/RigVMDeveloper/Private/RigVMCompiler/RigVMAST.cpp
+++ b/Engine/Plugins/Runtime/RigVM/Source/RigVMDeveloper/Private/RigVMCompiler/RigVMAST.cpp
@@ -336,6 +336,16 @@ TOptional<uint32> FRigVMExprAST::GetBlockCombinationHash() const
 					return TOptional<uint32>();
 				}
 
+				// Entry node output pins may be shared across many blocks,
+				// since they represent data which is not considered scoped.
+				if (Node->IsA<URigVMFunctionEntryNode>())
+				{
+					if (Pin->GetDirection() == ERigVMPinDirection::Output)
+					{
+						return TOptional<uint32>();
+					}
+				}
+
 				// variable node output pins may be shared across many blocks,
 				// since they represent data which is not considered scoped.
 				if(Node->IsA<URigVMVariableNode>())
