Control Rig Blueprints are perpetually dirtied upon reopening

#Synced-CL 50648283

diff --git a/Engine/Plugins/Runtime/RigVM/Source/RigVMDeveloper/Private/RigVMBlueprintLegacy.cpp b/Engine/Plugins/Runtime/RigVM/Source/RigVMDeveloper/Private/RigVMBlueprintLegacy.cpp
--- a/Engine/Plugins/Runtime/RigVM/Source/RigVMDeveloper/Private/RigVMBlueprintLegacy.cpp
+++ b/Engine/Plugins/Runtime/RigVM/Source/RigVMDeveloper/Private/RigVMBlueprintLegacy.cpp
@@ -273,7 +273,23 @@ ERigVMAssetStatus URigVMBlueprint::GetAssetStatus() const
 {
 	switch (Status)
 	{
-		case BS_Dirty: return RVMA_Dirty; break;
+		case BS_Dirty:
+		case BS_Unknown:
+		{
+			// Check if the asset is actually up to date by comparing the hashes.
+			// This handles cases where status is incorrectly set to dirty or defaults to unknown.
+			if (GetRigVMClient())
+			{
+				if (GetRigVMClient()->GetStructureHash() == GetRigVMClient()->GetSerializedStructureHash())
+				{
+					// Hashes match, asset is up to date.
+					// Set the Status property so UI sees the correct state.
+					const_cast<URigVMBlueprint*>(this)->Status = BS_UpToDate;
+					return RVMA_UpToDate;
+				}
+			}
+			return Status == BS_Dirty ? RVMA_Dirty : RVMA_Unknown;
+		}
 		case BS_Error: return RVMA_Error; break;
 		case BS_UpToDate: return RVMA_UpToDate; break;
 		case BS_BeingCreated: return RVMA_BeingCreated; break;
@@ -282,6 +298,7 @@ ERigVMAssetStatus URigVMBlueprint::GetAssetStatus() const
 	}
 }
 
+
 TArray<UObject*> URigVMBlueprint::GetArchetypeInstances(bool bIncludeCDO, bool bIncludeDerivedClasses) const
 {
 	TArray<UObject*> ArchetypeInstances;
diff --git a/Engine/Plugins/Runtime/RigVM/Source/RigVMDeveloper/Private/RigVMEditorAsset.cpp b/Engine/Plugins/Runtime/RigVM/Source/RigVMDeveloper/Private/RigVMEditorAsset.cpp
--- a/Engine/Plugins/Runtime/RigVM/Source/RigVMDeveloper/Private/RigVMEditorAsset.cpp
+++ b/Engine/Plugins/Runtime/RigVM/Source/RigVMDeveloper/Private/RigVMEditorAsset.cpp
@@ -4866,7 +4866,22 @@ ERigVMAssetStatus URigVMEditorAsset::GetAssetStatus() const
 {
 	if (URigVMRuntimeAsset* RuntimeAsset = GetRuntimeAsset())
 	{
-		return RuntimeAsset->GetAssetStatus();
+		ERigVMAssetStatus AssetStatus = RuntimeAsset->GetAssetStatus();
+
+		// When the status is unknown, check if the asset is up to date.
+		// This handles the case where the transient Status property defaults to unknown after reopening.
+		if (AssetStatus == ERigVMAssetStatus::RVMA_Unknown)
+		{
+			if (GetRigVMClient())
+			{
+				if (GetRigVMClient()->GetStructureHash() == GetRigVMClient()->GetSerializedStructureHash())
+				{
+					// Asset is up to date when the hashes match.
+					return ERigVMAssetStatus::RVMA_UpToDate;
+				}
+			}
+		}
+		return AssetStatus;
 	}
 	return ERigVMAssetStatus::RVMA_Unknown;
 }
