Control Rig: Fix crash when loading User defined enums during Async loading
#rnx

#Synced-CL 50592532

diff --git a/Engine/Plugins/Runtime/RigVM/Source/RigVM/Private/RigVMTypeUtils.cpp b/Engine/Plugins/Runtime/RigVM/Source/RigVM/Private/RigVMTypeUtils.cpp
--- a/Engine/Plugins/Runtime/RigVM/Source/RigVM/Private/RigVMTypeUtils.cpp
+++ b/Engine/Plugins/Runtime/RigVM/Source/RigVM/Private/RigVMTypeUtils.cpp
@@ -297,7 +297,9 @@ UObject* RigVMTypeUtils::UserDefinedTypeFromCPPType(FString& InOutCPPType, const
 						const FString GuidBasedName = GetUniqueStructTypeName(AssetData.GetTagValueRef<FGuid>(GuidTag));
 						if(GuidBasedName == InOutCPPType)
 						{
-							if(UUserDefinedStruct* UserDefinedStruct = Cast<UUserDefinedStruct>(AssetData.GetAsset()))
+							// Don't force a load (specially if you are async loading)
+							UObject* Obj = AssetData.FastGetAsset(false);
+							if(UUserDefinedStruct* UserDefinedStruct = Cast<UUserDefinedStruct>(Obj))
 							{
 								CPPTypeObject = UserDefinedStruct;
 								InOutCPPType = PostProcessCPPType(OriginalTypeName, CPPTypeObject);
@@ -308,7 +310,7 @@ UObject* RigVMTypeUtils::UserDefinedTypeFromCPPType(FString& InOutCPPType, const
 				}
 			}
 
-			// second pass - deal with the ones that don't have a tag and force load them
+			// second pass - deal with the ones that don't have a tag and resolve them
 			if(CPPTypeObject == nullptr)
 			{
 				for(const FAssetData& AssetData : AssetDataList)
@@ -318,7 +320,9 @@ UObject* RigVMTypeUtils::UserDefinedTypeFromCPPType(FString& InOutCPPType, const
 						continue;
 					}
 
-					if(UUserDefinedStruct* UserDefinedStruct = Cast<UUserDefinedStruct>(AssetData.GetAsset()))
+					// Don't force a load (specially if you are async loading)
+					UObject* Obj = AssetData.FastGetAsset(false);
+					if(UUserDefinedStruct* UserDefinedStruct = Cast<UUserDefinedStruct>(Obj))
 					{
 						const FString GuidBasedName = GetUniqueStructTypeName(UserDefinedStruct);
 						if(GuidBasedName == InOutCPPType)
@@ -338,9 +342,11 @@ UObject* RigVMTypeUtils::UserDefinedTypeFromCPPType(FString& InOutCPPType, const
 
 			for(const FAssetData& AssetData : AssetDataList)
 			{
-				if(UUserDefinedEnum* UserDefinedEnum = Cast<UUserDefinedEnum>(AssetData.GetAsset()))
+				// Don't force a load (specially if you are async loading)
+				UObject* Obj = AssetData.FastGetAsset(false);
+				if(UUserDefinedEnum* UserDefinedEnum = Cast<UUserDefinedEnum>(Obj))
 				{
-					const FString EnumCPPName = UserDefinedEnum->GetName();
+					const FString EnumCPPName = CPPTypeFromEnum(UserDefinedEnum);
 					if(EnumCPPName == InOutCPPType)
 					{
 						CPPTypeObject = UserDefinedEnum;
