[Backout] - CL50710366
#rnx
Original CL Desc
-----------------------------------------------------------------
Control Rig: Fix crash when loading User defined enums during Async loading
#rnx

#Synced-CL 50708221

diff --git a/Engine/Plugins/Runtime/RigVM/Source/RigVM/Private/RigVMTypeUtils.cpp b/Engine/Plugins/Runtime/RigVM/Source/RigVM/Private/RigVMTypeUtils.cpp
--- a/Engine/Plugins/Runtime/RigVM/Source/RigVM/Private/RigVMTypeUtils.cpp
+++ b/Engine/Plugins/Runtime/RigVM/Source/RigVM/Private/RigVMTypeUtils.cpp
@@ -297,9 +297,7 @@ UObject* RigVMTypeUtils::UserDefinedTypeFromCPPType(FString& InOutCPPType, const
 						const FString GuidBasedName = GetUniqueStructTypeName(AssetData.GetTagValueRef<FGuid>(GuidTag));
 						if(GuidBasedName == InOutCPPType)
 						{
-							// Don't force a load (specially if you are async loading)
-							UObject* Obj = AssetData.FastGetAsset(false);
-							if(UUserDefinedStruct* UserDefinedStruct = Cast<UUserDefinedStruct>(Obj))
+							if(UUserDefinedStruct* UserDefinedStruct = Cast<UUserDefinedStruct>(AssetData.GetAsset()))
 							{
 								CPPTypeObject = UserDefinedStruct;
 								InOutCPPType = PostProcessCPPType(OriginalTypeName, CPPTypeObject);
@@ -310,7 +308,7 @@ UObject* RigVMTypeUtils::UserDefinedTypeFromCPPType(FString& InOutCPPType, const
 				}
 			}
 
-			// second pass - deal with the ones that don't have a tag and resolve them
+			// second pass - deal with the ones that don't have a tag and force load them
 			if(CPPTypeObject == nullptr)
 			{
 				for(const FAssetData& AssetData : AssetDataList)
@@ -320,9 +318,7 @@ UObject* RigVMTypeUtils::UserDefinedTypeFromCPPType(FString& InOutCPPType, const
 						continue;
 					}
 
-					// Don't force a load (specially if you are async loading)
-					UObject* Obj = AssetData.FastGetAsset(false);
-					if(UUserDefinedStruct* UserDefinedStruct = Cast<UUserDefinedStruct>(Obj))
+					if(UUserDefinedStruct* UserDefinedStruct = Cast<UUserDefinedStruct>(AssetData.GetAsset()))
 					{
 						const FString GuidBasedName = GetUniqueStructTypeName(UserDefinedStruct);
 						if(GuidBasedName == InOutCPPType)
@@ -342,11 +338,9 @@ UObject* RigVMTypeUtils::UserDefinedTypeFromCPPType(FString& InOutCPPType, const
 
 			for(const FAssetData& AssetData : AssetDataList)
 			{
-				// Don't force a load (specially if you are async loading)
-				UObject* Obj = AssetData.FastGetAsset(false);
-				if(UUserDefinedEnum* UserDefinedEnum = Cast<UUserDefinedEnum>(Obj))
+				if(UUserDefinedEnum* UserDefinedEnum = Cast<UUserDefinedEnum>(AssetData.GetAsset()))
 				{
-					const FString EnumCPPName = CPPTypeFromEnum(UserDefinedEnum);
+					const FString EnumCPPName = UserDefinedEnum->GetName();
 					if(EnumCPPName == InOutCPPType)
 					{
 						CPPTypeObject = UserDefinedEnum;
