RigVM: Add basic support for trait deprecation/upgrade

This only changes the minimal amount of RigVM core to allow UAF trait deprecation to work, it does not add deprecation/upgrade support to RigVM core.

#Synced-CL 50766172

diff --git a/Engine/Plugins/Runtime/RigVM/Source/RigVM/Private/RigVMCore/RigVMStruct.cpp b/Engine/Plugins/Runtime/RigVM/Source/RigVM/Private/RigVMCore/RigVMStruct.cpp
--- a/Engine/Plugins/Runtime/RigVM/Source/RigVM/Private/RigVMCore/RigVMStruct.cpp
+++ b/Engine/Plugins/Runtime/RigVM/Source/RigVM/Private/RigVMCore/RigVMStruct.cpp
@@ -327,11 +327,11 @@ FString FRigVMStructUpgradeInfo::AddAggregatePin(FString InPinName)
 	return AggregatePins.Last();
 }
 
-void FRigVMStructUpgradeInfo::SetDefaultValues(const FRigVMStruct* InNewStructMemory)
+void FRigVMStructUpgradeInfo::SetDefaultValues(const void* InNewStructMemory)
 {
 	check(NewStruct);
 	check(InNewStructMemory);
-	DefaultValues = InNewStructMemory->GetDefaultValues(NewStruct);
+	DefaultValues = FRigVMStruct::GetDefaultValues(InNewStructMemory, NewStruct);
 }
 
 ////////////////////////////////////////////////////////////////////////////////////////////////////////////////////////
@@ -781,6 +781,11 @@ FName FRigVMStruct::GetNextAggregateName(const FName& InLastAggregatePinName) co
 }
 
 TMap<FName, FString> FRigVMStruct::GetDefaultValues(UScriptStruct* InScriptStruct) const
+{
+	return GetDefaultValues(this, InScriptStruct);
+}
+
+TMap<FName, FString> FRigVMStruct::GetDefaultValues(const void* InValue, UScriptStruct* InScriptStruct)
 {
 	check(InScriptStruct);
 
@@ -805,7 +810,7 @@ TMap<FName, FString> FRigVMStruct::GetDefaultValues(UScriptStruct* InScriptStruc
 #endif
 			
 			const FName PropertyName = It->GetFName();
-			const uint8* StructMemberMemoryPtr = It->ContainerPtrToValuePtr<uint8>(this);
+			const uint8* StructMemberMemoryPtr = It->ContainerPtrToValuePtr<uint8>(InValue);
 			const FString DefaultValue = ExportToFullyQualifiedText(*It, StructMemberMemoryPtr, false);
 			DefaultValues.Add(PropertyName, DefaultValue);
 		}
diff --git a/Engine/Plugins/Runtime/RigVM/Source/RigVM/Public/RigVMCore/RigVMStruct.h b/Engine/Plugins/Runtime/RigVM/Source/RigVM/Public/RigVMCore/RigVMStruct.h
--- a/Engine/Plugins/Runtime/RigVM/Source/RigVM/Public/RigVMCore/RigVMStruct.h
+++ b/Engine/Plugins/Runtime/RigVM/Source/RigVM/Public/RigVMCore/RigVMStruct.h
@@ -304,6 +304,7 @@ protected:
 
 	RIGVM_API static float GetRatioFromIndex(int32 InIndex, int32 InCount);
 	RIGVM_API TMap<FName, FString> GetDefaultValues(UScriptStruct* InScriptStruct) const;
+	RIGVM_API static TMap<FName, FString> GetDefaultValues(const void* InValue, UScriptStruct* InScriptStruct);
 	RIGVM_API bool ApplyUpgradeInfo(const FRigVMStructUpgradeInfo& InUpgradeInfo);
 	virtual TArray<FRigVMUserWorkflow> GetSupportedWorkflows(const UObject* InSubject) const { return TArray<FRigVMUserWorkflow>(); }
 	
diff --git a/Engine/Plugins/Runtime/RigVM/Source/RigVM/Public/RigVMCore/RigVMStructUpgradeInfo.h b/Engine/Plugins/Runtime/RigVM/Source/RigVM/Public/RigVMCore/RigVMStructUpgradeInfo.h
--- a/Engine/Plugins/Runtime/RigVM/Source/RigVM/Public/RigVMCore/RigVMStructUpgradeInfo.h
+++ b/Engine/Plugins/Runtime/RigVM/Source/RigVM/Public/RigVMCore/RigVMStructUpgradeInfo.h
@@ -102,7 +102,7 @@ private:
 	TArray<FString> AggregatePins;
 
 	// sets the default values from the new struct.
-	UE_API void SetDefaultValues(const FRigVMStruct* InNewStructMemory);
+	UE_API void SetDefaultValues(const void* InNewStructMemory);
 
 	friend class URigVMController;
 };
diff --git a/Engine/Plugins/Runtime/RigVM/Source/RigVM/Public/RigVMCore/RigVMTrait.h b/Engine/Plugins/Runtime/RigVM/Source/RigVM/Public/RigVMCore/RigVMTrait.h
--- a/Engine/Plugins/Runtime/RigVM/Source/RigVM/Public/RigVMCore/RigVMTrait.h
+++ b/Engine/Plugins/Runtime/RigVM/Source/RigVM/Public/RigVMCore/RigVMTrait.h
@@ -60,6 +60,7 @@ public:
 		return InProperty->GetFName() != GET_MEMBER_NAME_CHECKED(FRigVMTrait, Name);
 	}
 
+	virtual FRigVMStructUpgradeInfo GetUpgradeInfo() const { return FRigVMStructUpgradeInfo(); }
 #endif
 
 private:
