RigVM: Make fuzzy search in action menu optional

#rnx

#Synced-CL 50707384

diff --git a/Engine/Plugins/Animation/ControlRig/Source/ControlRigDeveloper/Private/Graph/ControlRigGraphSchema.cpp b/Engine/Plugins/Animation/ControlRig/Source/ControlRigDeveloper/Private/Graph/ControlRigGraphSchema.cpp
--- a/Engine/Plugins/Animation/ControlRig/Source/ControlRigDeveloper/Private/Graph/ControlRigGraphSchema.cpp
+++ b/Engine/Plugins/Animation/ControlRig/Source/ControlRigDeveloper/Private/Graph/ControlRigGraphSchema.cpp
@@ -2,6 +2,7 @@
 
 #include "Graph/ControlRigGraphSchema.h"
 #include "Graph/ControlRigGraph.h"
+#include "Settings/ControlRigSettings.h"
 #include "Units/Execution/RigUnit_BeginExecution.h"
 #include "Units/Execution/RigUnit_PrepareForExecution.h"
 #include "Units/Execution/RigUnit_InverseExecution.h"
@@ -17,6 +18,7 @@
 
 UControlRigGraphSchema::UControlRigGraphSchema()
 {
+	EditorSettingsClass = UControlRigEditorSettings::StaticClass();
 }
 
 FLinearColor UControlRigGraphSchema::GetPinTypeColor(const FEdGraphPinType& PinType) const
diff --git a/Engine/Plugins/Runtime/RigVM/Source/RigVM/Private/RigVMSettings.cpp b/Engine/Plugins/Runtime/RigVM/Source/RigVM/Private/RigVMSettings.cpp
--- a/Engine/Plugins/Runtime/RigVM/Source/RigVM/Private/RigVMSettings.cpp
+++ b/Engine/Plugins/Runtime/RigVM/Source/RigVM/Private/RigVMSettings.cpp
@@ -10,6 +10,8 @@ URigVMEditorSettings::URigVMEditorSettings(const FObjectInitializer& ObjectIniti
     : Super(ObjectInitializer)
 #if WITH_EDITORONLY_DATA
 	, bEnableContextMenuTimeSlicing(true)
+	, bEnableFuzzySearchInActionMenu(true)
+	, MinimumFuzzyScore(25.f)
 #endif
 {
 #if WITH_EDITORONLY_DATA
diff --git a/Engine/Plugins/Runtime/RigVM/Source/RigVM/Public/RigVMSettings.h b/Engine/Plugins/Runtime/RigVM/Source/RigVM/Public/RigVMSettings.h
--- a/Engine/Plugins/Runtime/RigVM/Source/RigVM/Public/RigVMSettings.h
+++ b/Engine/Plugins/Runtime/RigVM/Source/RigVM/Public/RigVMSettings.h
@@ -49,6 +49,20 @@ public:
 	/** When the Blueprint graph context menu is invoked (e.g. by right-clicking in the graph or dragging off a pin), do not block the UI while populating the available actions list. */
 	UPROPERTY(EditAnywhere, config, Category = Workflow, meta = (DisplayName = "Enable Non-Blocking Context Menu"))
 	bool bEnableContextMenuTimeSlicing;
+
+	/**
+	 * When right-clicking in a RigVM based graph (for example Control Rig)
+	 * enabling this option will apply a fuzzy string search to find items
+	 */
+	UPROPERTY(EditAnywhere, config, Category = Workflow, meta = (DisplayName = "Enable fuzzy search in action menu"))
+	bool bEnableFuzzySearchInActionMenu;
+
+	/**
+	 * Fuzzy searches will return a weight between 0.0 and 100.0. 
+	 * By specifying the minimum weight you can remove unlikely matches.
+	 */
+	UPROPERTY(EditAnywhere, config, Category = Workflow, meta = (DisplayName = "Minimum Fuzzy Score", UIMin = "0.0", UIMax = "100.0"))
+	float MinimumFuzzyScore;
 #endif
 };
 
diff --git a/Engine/Plugins/Runtime/RigVM/Source/RigVMDeveloper/Private/EdGraph/RigVMEdGraphSchema.cpp b/Engine/Plugins/Runtime/RigVM/Source/RigVMDeveloper/Private/EdGraph/RigVMEdGraphSchema.cpp
--- a/Engine/Plugins/Runtime/RigVM/Source/RigVMDeveloper/Private/EdGraph/RigVMEdGraphSchema.cpp
+++ b/Engine/Plugins/Runtime/RigVM/Source/RigVMDeveloper/Private/EdGraph/RigVMEdGraphSchema.cpp
@@ -713,6 +713,9 @@ TSharedRef<FRigVMFunctionDragDropAction> FRigVMFunctionDragDropAction::New(TShar
 }
 
 URigVMEdGraphSchema::URigVMEdGraphSchema()
+	: UEdGraphSchema()
+	, EditorSettingsClass(URigVMEditorSettings::StaticClass())
+	, bLastPinWasInput(false)
 {
 }
 
@@ -1067,9 +1070,21 @@ bool URigVMEdGraphSchema::SupportsPreviewFromHereOnNode(const UEdGraphNode* InNo
 	return false;
 }
 
+bool URigVMEdGraphSchema::AllowsFuzzyActionFiltering(const FEdGraphSchemaAction& InCurrentAction) const
+{
+	const URigVMEditorSettings* Settings = GetDefault<URigVMEditorSettings>(EditorSettingsClass);
+	return Settings->bEnableFuzzySearchInActionMenu;
+}
+
 float URigVMEdGraphSchema::GetActionFilteredWeight(const FEdGraphSchemaAction& InCurrentAction, const TArray<FString>& InFilterTerms,
-	const TArray<FString>& InSanitizedFilterTerms, const TArray<UEdGraphPin*>& DraggedFromPins) const
+                                                   const TArray<FString>& InSanitizedFilterTerms, const TArray<UEdGraphPin*>& DraggedFromPins) const
 {
+	const URigVMEditorSettings* Settings = GetDefault<URigVMEditorSettings>(EditorSettingsClass);
+	if (!Settings->bEnableFuzzySearchInActionMenu)
+	{
+		return UEdGraphSchema::GetActionFilteredWeight(InCurrentAction, InFilterTerms, InSanitizedFilterTerms, DraggedFromPins);
+	}
+	
 	TArray<FString> SanitizedFilterTerms = InSanitizedFilterTerms;
 	for (FString& SanitizedFilterTerm : SanitizedFilterTerms)
 	{
@@ -1271,16 +1286,25 @@ float URigVMEdGraphSchema::GetActionFilteredWeight(const FEdGraphSchemaAction& I
 	}
 
 	uint64 OverallScore = 0;
+	uint64 MaximumPossibleScore = 0;
 	uint64 Multiplier = 1;
 	for (int32 ScoreIndex = Scores.Num() - 1; ScoreIndex >= 0; --ScoreIndex)
 	{
 		const int32 ScoreInt32 = FMath::Min(static_cast<int32>(Divisor) - 1, FMath::FloorToInt32(Scores[ScoreIndex] * (static_cast<float>(Divisor) - 0.1f)));
 		OverallScore += static_cast<uint64>(ScoreInt32) * Multiplier;
+		MaximumPossibleScore += Divisor * Multiplier;
 		Multiplier *= Divisor;
 	}
 
 	const double OverallScoreDouble = static_cast<double>(OverallScore);
-	return static_cast<float>(OverallScoreDouble);
+	const double MaximumScoreDouble = static_cast<double>(MaximumPossibleScore);
+	const double MinimumScoreDouble = static_cast<double>(Settings->MinimumFuzzyScore);
+	const double WeightAsPercent = 100.0 * OverallScoreDouble / MaximumScoreDouble;
+	if (WeightAsPercent < MinimumScoreDouble)
+	{
+		return -1.f;
+	}
+	return static_cast<float>(WeightAsPercent);
 }
 
 bool URigVMEdGraphSchema::SupportsPinType(const UScriptStruct* ScriptStruct) const
diff --git a/Engine/Plugins/Runtime/RigVM/Source/RigVMDeveloper/Public/EdGraph/RigVMEdGraphSchema.h b/Engine/Plugins/Runtime/RigVM/Source/RigVMDeveloper/Public/EdGraph/RigVMEdGraphSchema.h
--- a/Engine/Plugins/Runtime/RigVM/Source/RigVMDeveloper/Public/EdGraph/RigVMEdGraphSchema.h
+++ b/Engine/Plugins/Runtime/RigVM/Source/RigVMDeveloper/Public/EdGraph/RigVMEdGraphSchema.h
@@ -365,7 +365,7 @@ public:
 	UE_API virtual TSharedPtr<INameValidatorInterface> GetNameValidator(const FRigVMEditorAssetInterfacePtr Asset, const FName& OriginalName, const UStruct* ValidationScope) const;
 	UE_API bool CanPromotePinToVariable(const UEdGraphPin& Pin, bool bCond) const;
 	UE_API virtual bool SupportsPreviewFromHereOnNode(const UEdGraphNode* InNode) const;
-	virtual bool AllowsFuzzyActionFiltering(const FEdGraphSchemaAction& InCurrentAction) const override { return true; }
+	UE_API virtual bool AllowsFuzzyActionFiltering(const FEdGraphSchemaAction& InCurrentAction) const override;
 	UE_API virtual float GetActionFilteredWeight(const FEdGraphSchemaAction& InCurrentAction, const TArray<FString>& InFilterTerms, const TArray<FString>& InSanitizedFilterTerms, const TArray<UEdGraphPin*>& DraggedFromPins) const override;
 
 	/** Returns true if the schema supports the script type */
@@ -400,6 +400,11 @@ public:
 	// Allows to hide an input variable pin from node node
 	UE_API virtual bool ShouldHidePinOnNode(const URigVMPin* InModelPin) const;
 
+protected:
+#if WITH_EDITOR
+	mutable UClass* EditorSettingsClass;
+#endif
+
 private:
 
 	const UEdGraphPin* LastPinForCompatibleCheck = nullptr;
