Enable AutoRTFM in VersePrototypes.

#Synced-CL 50777426

diff --git a/Engine/Plugins/Runtime/RigVM/Source/RigVMDeveloper/RigVMDeveloper.Build.cs b/Engine/Plugins/Runtime/RigVM/Source/RigVMDeveloper/RigVMDeveloper.Build.cs
--- a/Engine/Plugins/Runtime/RigVM/Source/RigVMDeveloper/RigVMDeveloper.Build.cs
+++ b/Engine/Plugins/Runtime/RigVM/Source/RigVMDeveloper/RigVMDeveloper.Build.cs
@@ -48,5 +48,7 @@ public class RigVMDeveloper : ModuleRules
             PrivateIncludePathModuleNames.Add("RigVMEditor");
             DynamicallyLoadedModuleNames.Add("RigVMEditor");
         }
-    }
+
+		bDisableAutoRTFMInstrumentation = false;
+	}
 }
