[Backout] - CL50812809
#rnx
Original CL Desc
-----------------------------------------------------------------
Enable AutoRTFM in VersePrototypes.

#Synced-CL 50814076

diff --git a/Engine/Plugins/Runtime/RigVM/Source/RigVMDeveloper/RigVMDeveloper.Build.cs b/Engine/Plugins/Runtime/RigVM/Source/RigVMDeveloper/RigVMDeveloper.Build.cs
--- a/Engine/Plugins/Runtime/RigVM/Source/RigVMDeveloper/RigVMDeveloper.Build.cs
+++ b/Engine/Plugins/Runtime/RigVM/Source/RigVMDeveloper/RigVMDeveloper.Build.cs
@@ -48,7 +48,5 @@ public class RigVMDeveloper : ModuleRules
             PrivateIncludePathModuleNames.Add("RigVMEditor");
             DynamicallyLoadedModuleNames.Add("RigVMEditor");
         }
-
-		bDisableAutoRTFMInstrumentation = false;
-	}
+    }
 }
