RigVM: Introducing callables to RigVM Bytecode

#rnx

Full cook preflight https://horde.devtools.epicgames.com/job/698d8d4f976487e0773563ed

#Synced-CL 50813099

diff --git a/Engine/Plugins/Animation/ControlRig/Source/ControlRigEditor/Private/ControlRigCompilerDetails.cpp b/Engine/Plugins/Animation/ControlRig/Source/ControlRigEditor/Private/ControlRigCompilerDetails.cpp
--- a/Engine/Plugins/Animation/ControlRig/Source/ControlRigEditor/Private/ControlRigCompilerDetails.cpp
+++ b/Engine/Plugins/Animation/ControlRig/Source/ControlRigEditor/Private/ControlRigCompilerDetails.cpp
@@ -11,7 +11,6 @@
 #include "HAL/PlatformApplicationMisc.h"
 #include "ControlRigVisualGraphUtils.h"
 #include "Units/Execution/RigUnit_BeginExecution.h"
-#include "RigVMCompiler/RigVMCodeGenerator.h"
 
 #define LOCTEXT_NAMESPACE "ControlRigCompilerDetails"
 
@@ -136,17 +135,13 @@ void FRigVMCompileSettingsDetails::CustomizeChildren(TSharedRef<IPropertyHandle>
 				+ SVerticalBox::Slot()
 				[
 					SNew(SButton)
-					.OnClicked(this, &FRigVMCompileSettingsDetails::OnCopyGeneratedCodeClicked)
+					.OnClicked(this, &FRigVMCompileSettingsDetails::OnCopyAnalyticsClicked)
 					.ContentPadding(FMargin(2))
-					.Visibility_Lambda([]()
-					{
-						return UControlRig::AreNativizedVMsDisabled() ? EVisibility::Collapsed : EVisibility::Visible;
-					})
 					.Content()
 					[
 						SNew(STextBlock)
 						.Justification(ETextJustify::Center)
-						.Text(LOCTEXT("CopyGeneratedCodeToClipboard", "Copy Nativized C++ Code"))
+						.Text(LOCTEXT("CopyAnalyticsToClipboard", "Copy Analytics"))
 					]
 				]
 				+ SVerticalBox::Slot()
@@ -227,50 +222,91 @@ FReply FRigVMCompileSettingsDetails::OnCopyByteCodeClicked()
 	return FReply::Handled();
 }
 
-FReply FRigVMCompileSettingsDetails::OnCopyHierarchyGraphClicked()
+class FRigVMCompileSettingsDetailsCounterArchive final : public FArchiveUObject
+{
+public:
+	FRigVMCompileSettingsDetailsCounterArchive()
+		: Size(0)
+	{
+		SetIsSaving(true);
+		SetIsPersistent(true);
+		ArIsCountingMemory = true;
+	}
+
+	virtual void Serialize(void*, int64 Length) override { Size += Length; }
+	virtual int64 TotalSize() override { return Size; }
+	using FArchiveUObject::operator<<; // For visibility of the overloads we don't override
+	virtual FArchive& operator<<(UObject*& Obj) override
+	{
+		Size += sizeof(UObject*);
+		return *this;
+	}
+	virtual FArchive& operator<<(FName& Value) override
+	{
+		Size += sizeof(FName);
+		return *this;
+	}
+
+private:
+	int64 Size;
+};
+
+FReply FRigVMCompileSettingsDetails::OnCopyAnalyticsClicked()
 {
 	if (BlueprintBeingCustomized)
 	{
-		if(UControlRig* ControlRig = Cast<UControlRig>(BlueprintBeingCustomized->GetObjectBeingDebugged()))
+		if(URigVMHost* Host = Cast<URigVMHost>(BlueprintBeingCustomized->GetObjectBeingDebugged()))
 		{
-			FName EventName = FRigUnit_BeginExecution::EventName;
-			if(!ControlRig->GetEventQueue().IsEmpty())
+			FString Content;
+			URigVM* VM = Host->GetVM();
+
+			int32 ByteCodeFootPrint = 0;
 			{
-				EventName = ControlRig->GetEventQueue()[0];
+				FRigVMCompileSettingsDetailsCounterArchive Archive;
+				Archive << VM->GetByteCode();
+				ByteCodeFootPrint = IntCastChecked<int32>(Archive.TotalSize());
 			}
-			
-			const FString DotGraphContent = FControlRigVisualGraphUtils::DumpRigHierarchyToDotGraph(ControlRig->GetHierarchy(), EventName);
-			FPlatformApplicationMisc::ClipboardCopy(*DotGraphContent);
+			int32 LiteralsFootPrint = 0;
+			{
+				FRigVMCompileSettingsDetailsCounterArchive Archive;
+				Archive << VM->GetDefaultLiteralMemory();
+				LiteralsFootPrint = IntCastChecked<int32>(Archive.TotalSize());
+			}
+			int32 WorkStateFootPrint = 0;
+			{
+				FRigVMCompileSettingsDetailsCounterArchive Archive;
+				Archive << VM->GetDefaultWorkMemory();
+				WorkStateFootPrint = IntCastChecked<int32>(Archive.TotalSize());
+			}
+
+			Content += FString::Printf(TEXT("ByteCode.Instructions: %04d\n"), VM->GetByteCode().GetNumInstructions());
+			Content += FString::Printf(TEXT("ByteCode.Branches: %04d\n"), VM->GetByteCode().NumBranches());
+			Content += FString::Printf(TEXT("ByteCode.Callables: %04d\n"), VM->GetByteCode().NumCallables());
+			Content += FString::Printf(TEXT("ByteCode.Footprint: %04d bytes\n"), ByteCodeFootPrint);
+			Content += FString::Printf(TEXT("Literals.Properties: %04d\n"), VM->GetDefaultLiteralMemory().Num());
+			Content += FString::Printf(TEXT("Literals.Footprint: %04d bytes\n"), LiteralsFootPrint);
+			Content += FString::Printf(TEXT("WorkState.Properties: %04d\n"), VM->GetDefaultWorkMemory().Num());
+			Content += FString::Printf(TEXT("WorkState.Footprint: %04d bytes\n"), WorkStateFootPrint);
+			FPlatformApplicationMisc::ClipboardCopy(*Content);
 		}
 	}
 	return FReply::Handled();
 }
 
-FReply FRigVMCompileSettingsDetails::OnCopyGeneratedCodeClicked()
+FReply FRigVMCompileSettingsDetails::OnCopyHierarchyGraphClicked()
 {
 	if (BlueprintBeingCustomized)
 	{
-		const FString ClassName = FString::Printf(TEXT("%sVM"), *BlueprintBeingCustomized.GetObject()->GetName());
-
-		if (URigVMHost* VMHost = Cast<URigVMHost>(BlueprintBeingCustomized->GetObjectBeingDebugged()))
+		if(UControlRig* ControlRig = Cast<UControlRig>(BlueprintBeingCustomized->GetObjectBeingDebugged()))
 		{
-			if (URigVM* VM = VMHost->GetVM())
+			FName EventName = FRigUnit_BeginExecution::EventName;
+			if(!ControlRig->GetEventQueue().IsEmpty())
 			{
-				VM->ClearExternalVariables(VMHost->GetRigVMExtendedExecuteContext());
-				TArray<FRigVMExternalVariable> ExternalVariables = VMHost->GetExternalVariablesImpl(false);
-				for(const FRigVMExternalVariable& ExternalVariable : ExternalVariables)
-				{
-					VM->AddExternalVariable(VMHost->GetRigVMExtendedExecuteContext(), ExternalVariable);
-				}
-			
-				FRigVMExtendedExecuteContext& CDOContext = VMHost->GetRigVMExtendedExecuteContext();
-
-				FRigVMCodeGenerator CodeGenerator(ClassName,
-					TEXT("TestModule"), BlueprintBeingCustomized->GetDefaultModel(), VM, CDOContext,
-					VMHost->GetPublicContextStruct(), BlueprintBeingCustomized->PinToOperandMap);
-				const FString Content = CodeGenerator.DumpHeader(CDOContext) + TEXT("\r\n\r\n") + CodeGenerator.DumpSource(CDOContext);
-				FPlatformApplicationMisc::ClipboardCopy(*Content);
+				EventName = ControlRig->GetEventQueue()[0];
 			}
+			
+			const FString DotGraphContent = FControlRigVisualGraphUtils::DumpRigHierarchyToDotGraph(ControlRig->GetHierarchy(), EventName);
+			FPlatformApplicationMisc::ClipboardCopy(*DotGraphContent);
 		}
 	}
 	return FReply::Handled();
diff --git a/Engine/Plugins/Animation/ControlRig/Source/ControlRigEditor/Private/ControlRigCompilerDetails.h b/Engine/Plugins/Animation/ControlRig/Source/ControlRigEditor/Private/ControlRigCompilerDetails.h
--- a/Engine/Plugins/Animation/ControlRig/Source/ControlRigEditor/Private/ControlRigCompilerDetails.h
+++ b/Engine/Plugins/Animation/ControlRig/Source/ControlRigEditor/Private/ControlRigCompilerDetails.h
@@ -27,8 +27,8 @@ protected:
 	FReply OnInspectMemory(ERigVMMemoryType InMemoryType);
 	FReply OnCopyASTClicked();
 	FReply OnCopyByteCodeClicked();
+	FReply OnCopyAnalyticsClicked();
 	FReply OnCopyHierarchyGraphClicked();
-	FReply OnCopyGeneratedCodeClicked();
 #if UE_RIGVM_DEBUG_EXECUTION
 	FReply OnCopyVMExecutionClicked();
 #endif
diff --git a/Engine/Plugins/Experimental/UAF/UAF/Source/UAF/Private/Script/UAFRigVMComponent.cpp b/Engine/Plugins/Experimental/UAF/UAF/Source/UAF/Private/Script/UAFRigVMComponent.cpp
--- a/Engine/Plugins/Experimental/UAF/UAF/Source/UAF/Private/Script/UAFRigVMComponent.cpp
+++ b/Engine/Plugins/Experimental/UAF/UAF/Source/UAF/Private/Script/UAFRigVMComponent.cpp
@@ -132,7 +132,7 @@ TConstArrayView<UE::UAF::FScriptEventInfo> FUAFRigVMComponent::GetScriptEvents()
 		const FRigVMByteCodeEntry& Entry = ByteCode.GetEntry(EntryIndex);
 		const FRigVMInstruction& Instruction = Instructions[Entry.InstructionIndex];
 		const FRigVMExecuteOp& Op = ByteCode.GetOpAt<FRigVMExecuteOp>(Instruction);
-		const FRigVMFunction* Function = Functions[Op.FunctionIndex];
+		const FRigVMFunction* Function = Functions[Op.CallableIndex];
 		check(Function != nullptr);
 
 		if (Function->Struct->IsChildOf(FRigUnit_AnimNextModuleEventBase::StaticStruct()))
@@ -157,7 +157,7 @@ TConstArrayView<UE::UAF::FScriptEventInfo> FUAFRigVMComponent::GetScriptEvents()
 			if (Function->Struct->IsChildOf(FRigUnit_AnimNextUserEvent::StaticStruct()))
 			{
 				// Pull the values out of the literal memory
-				FRigVMOperandArray Operands = ByteCode.GetOperandsForExecuteOp(Instruction);
+				FRigVMOperandArray Operands = ByteCode.GetOperandsForCallableOp(Instruction);
 				check(Function->ArgumentNames.Num() == Operands.Num());
 				int32 NumOperands = Operands.Num();
 				for (int32 OperandIndex = 0; OperandIndex < NumOperands; ++OperandIndex)
diff --git a/Engine/Plugins/Experimental/UAF/UAF/Source/UAFEditor/Private/Graph/SModuleEventPicker.cpp b/Engine/Plugins/Experimental/UAF/UAF/Source/UAFEditor/Private/Graph/SModuleEventPicker.cpp
--- a/Engine/Plugins/Experimental/UAF/UAF/Source/UAFEditor/Private/Graph/SModuleEventPicker.cpp
+++ b/Engine/Plugins/Experimental/UAF/UAF/Source/UAFEditor/Private/Graph/SModuleEventPicker.cpp
@@ -120,7 +120,7 @@ void SModuleEventPicker::RefreshEntries()
 			const FRigVMByteCodeEntry& Entry = ByteCode.GetEntry(EntryIndex);
 			const FRigVMInstruction& Instruction = Instructions[Entry.InstructionIndex];
 			const FRigVMExecuteOp& Op = ByteCode.GetOpAt<FRigVMExecuteOp>(Instruction);
-			const FRigVMFunction* Function = Functions[Op.FunctionIndex];
+			const FRigVMFunction* Function = Functions[Op.CallableIndex];
 			check(Function != nullptr);
 
 			if (Function->Struct->IsChildOf(FRigUnit_AnimNextModuleEventBase::StaticStruct()) &&
diff --git a/Engine/Plugins/Runtime/RigVM/Source/RigVM/Private/RigVMCore/RigVM.cpp b/Engine/Plugins/Runtime/RigVM/Source/RigVM/Private/RigVMCore/RigVM.cpp
--- a/Engine/Plugins/Runtime/RigVM/Source/RigVM/Private/RigVMCore/RigVM.cpp
+++ b/Engine/Plugins/Runtime/RigVM/Source/RigVM/Private/RigVMCore/RigVM.cpp
@@ -488,14 +488,14 @@ void URigVM::RefreshArgumentNameCaches()
 		if (Instruction.OpCode == ERigVMOpCode::Execute)
 		{
 			const FRigVMExecuteOp& Op = ByteCodeStorage.GetOpAt<FRigVMExecuteOp>(Instruction);
-			if (Functions.IsValidIndex(Op.FunctionIndex))
+			if (Functions.IsValidIndex(Op.CallableIndex))
 			{
-				if (!Functions[Op.FunctionIndex])
+				if (!Functions[Op.CallableIndex])
 				{
 					continue;
 				}
 				
-				if (const FRigVMDispatchFactory* Factory = Functions[Op.FunctionIndex]->Factory)
+				if (const FRigVMDispatchFactory* Factory = Functions[Op.CallableIndex]->Factory)
 				{
 					if(Factories.IsEmpty())
 					{
@@ -506,7 +506,7 @@ void URigVM::RefreshArgumentNameCaches()
 					{
 						Factories[Factory->GetFactoryIndex()] = Factory;
 						
-						FRigVMOperandArray Operands = ByteCodeStorage.GetOperandsForExecuteOp(Instruction);
+						FRigVMOperandArray Operands = ByteCodeStorage.GetOperandsForCallableOp(Instruction);
 						if(LocalizedRegistry.IsValid())
 						{
 							(void)Factory->UpdateArgumentNameCache_NoLock(Operands.Num(), LocalizedRegistry->GetHandle_NoLock());
@@ -650,11 +650,28 @@ bool URigVM::ValidateBytecode()
 			case ERigVMOpCode::Execute:
 			{
 				const FRigVMExecuteOp& Op = ByteCodeStorage.GetOpAt<FRigVMExecuteOp>(ByteCodeInstruction);
-				if (!Functions.IsValidIndex(Op.FunctionIndex))
+				if (!Functions.IsValidIndex(Op.CallableIndex))
 				{
 					return false;
 				}
-				FRigVMOperandArray Operands = ByteCodeStorage.GetOperandsForExecuteOp(ByteCodeInstruction);
+				FRigVMOperandArray Operands = ByteCodeStorage.GetOperandsForCallableOp(ByteCodeInstruction);
+				for (const FRigVMOperand& Arg : Operands)
+				{
+					if (!CheckOperandValidity(Arg))
+					{
+						return false;
+					}
+				}
+				break;
+			}
+			case ERigVMOpCode::InvokeCallable:
+			{
+				const FRigVMInvokeCallableOp& Op = ByteCodeStorage.GetOpAt<FRigVMInvokeCallableOp>(ByteCodeInstruction);
+				if (!ByteCodeStorage.IsValidCallableIndex(Op.CallableIndex))
+				{
+					return false;
+				}
+				FRigVMOperandArray Operands = ByteCodeStorage.GetOperandsForCallableOp(ByteCodeInstruction);
 				for (const FRigVMOperand& Arg : Operands)
 				{
 					if (!CheckOperandValidity(Arg))
@@ -819,13 +836,27 @@ int32 URigVM::AddRigVMFunction(UScriptStruct* InRigVMStruct, const FName& InMeth
 
 int32 URigVM::AddRigVMFunction(const FString& InFunctionName)
 {
-	const FName FunctionName = *InFunctionName;
-	const int32 FunctionIndex = GetFunctionNames().Find(FunctionName);
+	const FName FunctionFName = *InFunctionName;
+	const int32 FunctionIndex = GetFunctionNames().Find(FunctionFName);
+	if (FunctionIndex != INDEX_NONE)
+	{
+		return FunctionIndex;
+	}
+	return AddRigVMFunctionImpl(FunctionFName, InFunctionName);
+}
+
+int32 URigVM::AddRigVMFunction(const FName& InFunctionFName)
+{
+	const int32 FunctionIndex = GetFunctionNames().Find(InFunctionFName);
 	if (FunctionIndex != INDEX_NONE)
 	{
 		return FunctionIndex;
 	}
+	return AddRigVMFunctionImpl(InFunctionFName, InFunctionFName.ToString());
+}
 
+int32 URigVM::AddRigVMFunctionImpl(const FName& InFunctionFName, const FString& InFunctionName)
+{
 	const FRigVMFunction* Function = nullptr;
 	if(LocalizedRegistry.IsValid())
 	{
@@ -840,7 +871,7 @@ int32 URigVM::AddRigVMFunction(const FString& InFunctionName)
 		return INDEX_NONE;
 	}
 
-	GetFunctionNames().Add(FunctionName);
+	GetFunctionNames().Add(InFunctionFName);
 	GetFactories().Add(Function->Factory);
 	return GetFunctions().Add(Function);
 }
@@ -1342,18 +1373,28 @@ void URigVM::InstructionOpEval(
 	const FRigVMByteCode& ByteCode = GetByteCode();
 	const TArray<const FRigVMFunction*>& Functions = GetFunctions();
 
-	const ERigVMOpCode OpCode = Instructions[InstructionIndex].OpCode;
+	const FRigVMInstruction& Instruction = Instructions[InstructionIndex]; 
+	const ERigVMOpCode OpCode = Instruction.OpCode;
 
 	if (OpCode == ERigVMOpCode::Execute)
 	{
-		const FRigVMExecuteOp& Op = ByteCode.GetOpAt<FRigVMExecuteOp>(Instructions[InstructionIndex]);
-		FRigVMOperandArray Operands = ByteCode.GetOperandsForExecuteOp(Instructions[InstructionIndex]);
-		const FRigVMFunction* Function = Functions[Op.FunctionIndex];
+		const FRigVMExecuteOp& Op = ByteCode.GetOpAt<FRigVMExecuteOp>(Instruction);
+		FRigVMOperandArray Operands = ByteCode.GetOperandsForCallableOp(Instruction);
+		const FRigVMFunction* Function = Functions[Op.CallableIndex];
 		if(Function == nullptr)
 		{
 			return;
 		}
 
+		if (Function->Factory)
+		{
+			checkf(Function->Arguments.Num() <= Operands.Num(), TEXT("%s: invalid number of operands (%d) for dispatch '%s' - expected at least (%d)."), *GetPathName(), Operands.Num(), *Function->GetName(), Function->Arguments.Num());
+		}
+		else
+		{
+			checkf(Function->Arguments.Num() == Operands.Num(), TEXT("%s: invalid number of operands (%d) for function '%s' - expected (%d)."), *GetPathName(), Operands.Num(), *Function->GetName(), Function->Arguments.Num());
+		}
+		
 		for (int32 ArgIndex = 0; ArgIndex < Operands.Num(); ArgIndex++)
 		{
 			InOpFunc(
@@ -1365,6 +1406,30 @@ void URigVM::InstructionOpEval(
 				Operands[ArgIndex]);
 		}
 	}
+	else if (OpCode == ERigVMOpCode::InvokeCallable)
+	{
+		const FRigVMInvokeCallableOp& Op = ByteCode.GetOpAt<FRigVMInvokeCallableOp>(Instruction);
+		FRigVMOperandArray Operands = ByteCode.GetOperandsForCallableOp(Instruction);
+		const FRigVMCallableInfo* Callable = ByteCode.GetCallable(Op.CallableIndex);
+		if(Callable == nullptr)
+		{
+			return;
+		}
+
+		checkf(Callable->Arguments.Num() == Operands.Num(), TEXT("%s: invalid number of operands (%d) for callable '%s' - expected (%d)."), *GetPathName(), Operands.Num(), *Callable->Name.ToString(), Callable->Arguments.Num());
+		
+		int32 TotalArgIndex = 0;
+		for (int32 ArgIndex = 0; ArgIndex < Operands.Num(); ArgIndex++, TotalArgIndex++)
+		{
+			InOpFunc(
+				Context,
+				RegistryHandle,
+				InstructionIndex,
+				InHandleBaseIndex++,
+				{ InstructionIndex, TotalArgIndex, Callable->Arguments[ArgIndex].Name },
+				Operands[ArgIndex]);
+		}
+	}
 	else
 	{
 		switch (OpCode)
@@ -1375,13 +1440,13 @@ void URigVM::InstructionOpEval(
 		case ERigVMOpCode::Increment:
 		case ERigVMOpCode::Decrement:
 		{
-			const FRigVMUnaryOp& Op = ByteCode.GetOpAt<FRigVMUnaryOp>(Instructions[InstructionIndex]);
+			const FRigVMUnaryOp& Op = ByteCode.GetOpAt<FRigVMUnaryOp>(Instruction);
 			InOpFunc(Context, RegistryHandle, InstructionIndex, InHandleBaseIndex, {}, Op.Arg);
 			break;
 		}
 		case ERigVMOpCode::Copy:
 		{
-			const FRigVMCopyOp& Op = ByteCode.GetOpAt<FRigVMCopyOp>(Instructions[InstructionIndex]);
+			const FRigVMCopyOp& Op = ByteCode.GetOpAt<FRigVMCopyOp>(Instruction);
 			InOpFunc(Context, RegistryHandle, InstructionIndex, InHandleBaseIndex + 0, {}, Op.Source);
 			InOpFunc(Context, RegistryHandle, InstructionIndex, InHandleBaseIndex + 1, {}, Op.Target);
 			break;
@@ -1389,7 +1454,7 @@ void URigVM::InstructionOpEval(
 		case ERigVMOpCode::Equals:
 		case ERigVMOpCode::NotEquals:
 		{
-			const FRigVMComparisonOp& Op = ByteCode.GetOpAt<FRigVMComparisonOp>(Instructions[InstructionIndex]);
+			const FRigVMComparisonOp& Op = ByteCode.GetOpAt<FRigVMComparisonOp>(Instruction);
 			FRigVMOperand Arg = Op.A;
 			InOpFunc(Context, RegistryHandle, InstructionIndex, InHandleBaseIndex + 0, {}, Arg);
 			Arg = Op.B;
@@ -1408,14 +1473,14 @@ void URigVM::InstructionOpEval(
 		case ERigVMOpCode::JumpForwardIf:
 		case ERigVMOpCode::JumpBackwardIf:
 		{
-			const FRigVMJumpIfOp& Op = ByteCode.GetOpAt<FRigVMJumpIfOp>(Instructions[InstructionIndex]);
+			const FRigVMJumpIfOp& Op = ByteCode.GetOpAt<FRigVMJumpIfOp>(Instruction);
 			const FRigVMOperand& Arg = Op.Arg;
 			InOpFunc(Context, RegistryHandle, InstructionIndex, InHandleBaseIndex, {}, Arg);
 			break;
 		}
 		case ERigVMOpCode::ChangeType:
 		{
-			const FRigVMChangeTypeOp& Op = ByteCode.GetOpAt<FRigVMChangeTypeOp>(Instructions[InstructionIndex]);
+			const FRigVMChangeTypeOp& Op = ByteCode.GetOpAt<FRigVMChangeTypeOp>(Instruction);
 			const FRigVMOperand& Arg = Op.Arg;
 			InOpFunc(Context, RegistryHandle, InstructionIndex, InHandleBaseIndex, {}, Arg);
 			break;
@@ -1426,7 +1491,7 @@ void URigVM::InstructionOpEval(
 		}
 		case ERigVMOpCode::BeginBlock:
 		{
-			const FRigVMBinaryOp& Op = ByteCode.GetOpAt<FRigVMBinaryOp>(Instructions[InstructionIndex]);
+			const FRigVMBinaryOp& Op = ByteCode.GetOpAt<FRigVMBinaryOp>(Instruction);
 			InOpFunc(Context, RegistryHandle, InstructionIndex, InHandleBaseIndex + 0, {}, Op.ArgA);
 			InOpFunc(Context, RegistryHandle, InstructionIndex, InHandleBaseIndex + 1, {}, Op.ArgB);
 			break;
@@ -1441,21 +1506,21 @@ void URigVM::InstructionOpEval(
 		}
 		case ERigVMOpCode::JumpToBranch:
 		{
-			const FRigVMJumpToBranchOp& Op = ByteCode.GetOpAt<FRigVMJumpToBranchOp>(Instructions[InstructionIndex]);
+			const FRigVMJumpToBranchOp& Op = ByteCode.GetOpAt<FRigVMJumpToBranchOp>(Instruction);
 			const FRigVMOperand& Arg = Op.Arg;
 			InOpFunc(Context, RegistryHandle, InstructionIndex, InHandleBaseIndex, {}, Arg);
 			break;
 		}
 		case ERigVMOpCode::RunInstructions:
 		{
-			const FRigVMRunInstructionsOp& Op = ByteCode.GetOpAt<FRigVMRunInstructionsOp>(Instructions[InstructionIndex]);
+			const FRigVMRunInstructionsOp& Op = ByteCode.GetOpAt<FRigVMRunInstructionsOp>(Instruction);
 			const FRigVMOperand& Arg = Op.Arg;
 			InOpFunc(Context, RegistryHandle, InstructionIndex, InHandleBaseIndex, {}, Arg);
 			break;
 		}
 		case ERigVMOpCode::SetupTraits:
 		{
-			const FRigVMSetupTraitsOp& Op = ByteCode.GetOpAt<FRigVMSetupTraitsOp>(Instructions[InstructionIndex]);
+			const FRigVMSetupTraitsOp& Op = ByteCode.GetOpAt<FRigVMSetupTraitsOp>(Instruction);
 			const FRigVMOperand& Arg = Op.Arg;
 			InOpFunc(Context, RegistryHandle, InstructionIndex, InHandleBaseIndex, {}, Arg);
 			break;
@@ -1699,10 +1764,10 @@ ERigVMExecuteResult URigVM::ExecuteVM(FRigVMExtendedExecuteContext& Context, con
 		{
 			if (RigVMInstructionVisitInfo->GetFirstEntryEventInEventQueue() == NAME_None || RigVMInstructionVisitInfo->GetFirstEntryEventInEventQueue() == InEntryName)
 			{
-				RigVMInstructionVisitInfo->SetupInstructionTracking(Instructions.Num());
+				RigVMInstructionVisitInfo->SetupInstructionTracking(Instructions.Num(), ByteCode.NumCallables());
 				if (FRigVMProfilingInfo* RigVMProfilingInfo = Context.GetRigVMProfilingInfo())
 				{
-					RigVMProfilingInfo->SetupInstructionTracking(Instructions.Num(), true);
+					RigVMProfilingInfo->SetupInstructionTracking(Instructions.Num(), ByteCode.NumCallables(), true);
 				}
 			}
 		}
@@ -1713,7 +1778,8 @@ ERigVMExecuteResult URigVM::ExecuteVM(FRigVMExtendedExecuteContext& Context, con
 	{
 		Context.ResetExecutionState();
 		Context.SetVM(this);
-		Context.SliceOffsets.AddZeroed(Instructions.Num());
+		Context.SliceOffsetsPerInstruction.AddZeroed(Instructions.Num());
+		Context.SliceOffsetsPerCallable.AddZeroed(ByteCode.NumCallables());
 	}
 
 	ensure(Context.GetVM() == this);
@@ -1742,9 +1808,22 @@ ERigVMExecuteResult URigVM::ExecuteVM(FRigVMExtendedExecuteContext& Context, con
 	}
 	else
 	{
+		int32 FirstInstructionIndex = 0;
+		if (bIsRootEntry)
+		{
+			// move the instruction index to the first instruction not owned by callable
+			for (int32 CallableIndex = 0; CallableIndex < ByteCode.NumCallables(); ++CallableIndex)
+			{
+				const FRigVMCallableInfo* Callable = ByteCode.GetCallable(CallableIndex);
+				check(Callable);
+				FirstInstructionIndex = FMath::Max(FirstInstructionIndex, Callable->LastInstruction + 1);
+			}
+			ContextPublicData.InstructionIndex = FirstInstructionIndex;
+		}
+		
 		for(int32 EntryIndex = 0; EntryIndex < ByteCode.NumEntries(); EntryIndex++)
 		{
-			if(ByteCode.GetEntry(EntryIndex).InstructionIndex == 0)
+			if(ByteCode.GetEntry(EntryIndex).InstructionIndex == FirstInstructionIndex)
 			{
 				EntryIndexToPush = EntryIndex;
 				break;
@@ -1866,6 +1945,41 @@ ERigVMExecuteResult URigVM::ExecuteVM(FRigVMExtendedExecuteContext& Context, con
 	return Context.CurrentExecuteResult;
 }
 
+struct FRigVMCallstackGuard
+{
+public:
+	FRigVMCallstackGuard(FRigVMExtendedExecuteContext& InContext, const FRigVMByteCode& ByteCode, int32 InCallableIndex)
+		: Context(InContext)
+		, PreviousHash(InContext.CallstackHash)
+		, bIsValid(false)
+	{
+		const int32 InstructionIndex = Context.GetInstructionIndex();
+		bIsValid = InstructionIndex != INDEX_NONE && ByteCode.IsValidCallableIndex(InCallableIndex);
+		if (bIsValid)
+		{
+			const FRigVMCallableInfo* Callable = ByteCode.GetCallable(InCallableIndex);
+			Context.Callstack.Emplace(InstructionIndex, Callable);
+			Context.CallstackHash = HashCombine(InstructionIndex, InCallableIndex, Context.CallstackHash);
+			Context.BeginCallableSlice(Callable->Index);
+		}
+	}
+
+	~FRigVMCallstackGuard()
+	{
+		if (bIsValid)
+		{
+			Context.EndCallableSlice();
+			Context.CallstackHash = PreviousHash;
+			Context.Callstack.Pop();
+		}
+	}
+		
+private:
+	FRigVMExtendedExecuteContext& Context;
+	uint32 PreviousHash;
+	bool bIsValid;
+};
+
 ERigVMExecuteResult URigVM::ExecuteInstructions(FRigVMExtendedExecuteContext& Context, int32 InFirstInstruction, int32 InLastInstruction)
 {
 	FRigVMExecuteContext& ContextPublicData = Context.GetPublicData<>();
@@ -1879,8 +1993,12 @@ ERigVMExecuteResult URigVM::ExecuteInstructions(FRigVMExtendedExecuteContext& Co
 	
 	// make we are already executing this VM
 	check(!Context.CurrentVMMemory.IsEmpty());
+	
+
+	FRigVMByteCode& ByteCode = GetByteCode();
 
 	TGuardValue<int32> InstructionIndexGuard(ContextPublicData.InstructionIndex, InFirstInstruction);
+	
 #if WITH_EDITOR
 	FInstructionBracketGuard InstructionBracket(Context, InFirstInstruction, InLastInstruction);
 	if(!InstructionBracket.ErrorMessage.IsEmpty())
@@ -1890,7 +2008,6 @@ ERigVMExecuteResult URigVM::ExecuteInstructions(FRigVMExtendedExecuteContext& Co
 	}
 #endif
 	
-	FRigVMByteCode& ByteCode = GetByteCode();
 	TArray<const FRigVMFunction*>& Functions = GetFunctions();
 	TArray<const FRigVMDispatchFactory*>& Factories = GetFactories();
 
@@ -1949,7 +2066,7 @@ ERigVMExecuteResult URigVM::ExecuteInstructions(FRigVMExtendedExecuteContext& Co
 			if (Instruction.OpCode == ERigVMOpCode::Execute)
 			{
 				const FRigVMExecuteOp& Op = ByteCode.GetOpAt<FRigVMExecuteOp>(Instruction);
-				FRigVMOperandArray Operands = ByteCode.GetOperandsForExecuteOp(Instructions[ContextPublicData.InstructionIndex]);
+				FRigVMOperandArray Operands = ByteCode.GetOperandsForCallableOp(Instructions[ContextPublicData.InstructionIndex]);
 
 				TArray<FString> Labels;
 				for (const FRigVMOperand& Operand : Operands)
@@ -1998,15 +2115,15 @@ ERigVMExecuteResult URigVM::ExecuteInstructions(FRigVMExtendedExecuteContext& Co
 					Predicates = FRigVMPredicateBranchArray(&ByteCode.PredicateBranches[Op.FirstPredicateIndex], Op.PredicateCount);
 				}
 #if WITH_EDITOR
-				ContextPublicData.FunctionName = FunctionNames[Op.FunctionIndex];
+				ContextPublicData.FunctionName = FunctionNames[Op.CallableIndex];
 #endif
-				Context.Factory = Factories[Op.FunctionIndex];
-				(*Functions[Op.FunctionIndex]->FunctionPtr)(Context, Handles, Predicates);
+				Context.Factory = Factories[Op.CallableIndex];
+				(*Functions[Op.CallableIndex]->FunctionPtr)(Context, Handles, Predicates);
 
 #if WITH_EDITOR
 				if(GetDebugMemory(Context) && GetDebugMemory(Context)->Num() > 0)
 				{
-					const FRigVMOperandArray Operands = ByteCode.GetOperandsForExecuteOp(Instruction);
+					const FRigVMOperandArray Operands = ByteCode.GetOperandsForCallableOp(Instruction);
 					for(int32 OperandIndex = 0, HandleIndex = 0; OperandIndex < Operands.Num() && HandleIndex < Handles.Num(); HandleIndex++)
 					{
 						CopyOperandForDebuggingIfNeeded(Context, Operands[OperandIndex++], Handles[HandleIndex]);
@@ -2017,6 +2134,84 @@ ERigVMExecuteResult URigVM::ExecuteInstructions(FRigVMExtendedExecuteContext& Co
 				ContextPublicData.InstructionIndex++;
 				break;
 			}
+			case ERigVMOpCode::InvokeCallable:
+			{
+				const FRigVMInvokeCallableOp& Op = ByteCode.GetOpAt<FRigVMInvokeCallableOp>(Instruction);
+				const FRigVMCallableInfo* Callable = ByteCode.GetCallable(Op.CallableIndex);
+				check(Callable);
+
+				// it's possible to have empty callables, so just step over this 
+				if (Callable->GetNumInstructions() == 0)
+				{
+					ContextPublicData.InstructionIndex++;
+					break;
+				}
+					
+				const int32 OperandCount = FirstHandleForInstruction[ContextPublicData.InstructionIndex + 1] - FirstHandleForInstruction[ContextPublicData.InstructionIndex];
+				const int32 ArgumentCount = Callable->Arguments.Num();
+				check(ArgumentCount == OperandCount);
+
+				if(OperandCount > 0)
+				{
+					// generate patches handles for the call site memory
+					HandlesMemory.SetNumUninitialized(ArgumentCount, EAllowShrinking::No);
+					TArrayView<FRigVMMemoryHandle> Handles = TArrayView<FRigVMMemoryHandle>(HandlesMemory.GetData(), ArgumentCount);
+					GeneratePatchedHandles(Handles, Context.ExternalVariableRuntimeData, VMMemoryPtrs, PreCachedMemoryHandles, PreCachedMemoryHandlesMemoryType, FirstHandleForInstruction[ContextPublicData.InstructionIndex], ArgumentCount);
+
+					// bend the content of the callable's interface operand memory to point to this callsite
+					for (int32 ArgumentIndex = 0; ArgumentIndex < ArgumentCount; ArgumentIndex++)
+					{
+						const FRigVMOperand& ForwardedOperand = Callable->Arguments[ArgumentIndex].ForwardedOperand;
+#if WITH_EDITOR
+						check(ForwardedOperand.IsValid() && ForwardedOperand.GetMemoryType() == ERigVMMemoryType::Work);
+#endif
+						FRigVMMemoryStorageStruct* PropertyBag = GetMemoryByType(Context, ForwardedOperand.GetMemoryType());
+#if WITH_EDITOR
+						const FProperty* ArgumentProperty = PropertyBag->GetProperty(ForwardedOperand.GetRegisterIndex());
+						check(ArgumentProperty);
+						check(CastFieldChecked<FStructProperty>(ArgumentProperty)->Struct == FRigVMForwardedMemoryHandle::StaticStruct());
+#endif
+						FRigVMForwardedMemoryHandle* ForwardHandle = PropertyBag->GetData<FRigVMForwardedMemoryHandle>(ForwardedOperand.GetRegisterIndex());
+						ForwardHandle->Memory = Handles[ArgumentIndex].GetData_Internal(true, INDEX_NONE);
+						ForwardHandle->Property = Handles[ArgumentIndex].GetResolvedProperty();
+#if WITH_EDITOR
+						const FRigVMOperand& InterfaceOperand = Callable->Arguments[ArgumentIndex].InterfaceOperand;
+						check(InterfaceOperand.IsValid() && InterfaceOperand.GetMemoryType() == ERigVMMemoryType::Work);
+						const FProperty* InterfaceProperty = PropertyBag->GetProperty(InterfaceOperand.GetRegisterIndex());
+						check(InterfaceProperty);
+						check(InterfaceProperty->SameType(ForwardHandle->Property));
+#endif
+					}
+				}
+
+				// invoke the range of the callable
+				const FRigVMCallstackGuard CallstackGuard(Context, ByteCode, Callable->Index);
+#if WITH_EDITOR
+				if (RigVMInstructionVisitInfo != nullptr)
+				{
+					RigVMInstructionVisitInfo->SetCallableVisitedDuringLastRun(Callable->Index);
+					RigVMInstructionVisitInfo->AddCallableIndexToVisitOrder(Callable->Index);
+				}
+#endif
+				const ERigVMExecuteResult CallableResult = ExecuteInstructions(Context, Callable->FirstInstruction, Callable->LastInstruction);
+
+				// reset the forwarded handles
+				for (int32 ArgumentIndex = 0; ArgumentIndex < ArgumentCount; ArgumentIndex++)
+				{
+					const FRigVMOperand& ForwardedOperand = Callable->Arguments[ArgumentIndex].ForwardedOperand;
+					FRigVMMemoryStorageStruct* PropertyBag = GetMemoryByType(Context, ForwardedOperand.GetMemoryType());
+					FRigVMForwardedMemoryHandle* ForwardHandle = PropertyBag->GetData<FRigVMForwardedMemoryHandle>(ForwardedOperand.GetRegisterIndex());
+					ForwardHandle->Reset();
+				}
+					
+				if (CallableResult != ERigVMExecuteResult::Succeeded)
+				{
+					return CallableResult;
+				}
+
+				ContextPublicData.InstructionIndex++;
+				break;
+			}
 			case ERigVMOpCode::Zero:
 			{
 				const int32 FirstHandleIndex = FirstHandleForInstruction[ContextPublicData.InstructionIndex];
@@ -2129,19 +2324,19 @@ ERigVMExecuteResult URigVM::ExecuteInstructions(FRigVMExtendedExecuteContext& Co
 			case ERigVMOpCode::JumpAbsolute:
 			{
 				const FRigVMJumpOp& Op = ByteCode.GetOpAt<FRigVMJumpOp>(Instruction);
-				ContextPublicData.InstructionIndex = Op.InstructionIndex;
+				ContextPublicData.InstructionIndex = FRigVMJumpOp::GetTargetInstruction(Instruction.OpCode, ContextPublicData.InstructionIndex, Op.InstructionIndex);
 				break;
 			}
 			case ERigVMOpCode::JumpForward:
 			{
 				const FRigVMJumpOp& Op = ByteCode.GetOpAt<FRigVMJumpOp>(Instruction);
-				ContextPublicData.InstructionIndex += Op.InstructionIndex;
+				ContextPublicData.InstructionIndex = FRigVMJumpOp::GetTargetInstruction(Instruction.OpCode, ContextPublicData.InstructionIndex, Op.InstructionIndex);
 				break;
 			}
 			case ERigVMOpCode::JumpBackward:
 			{
 				const FRigVMJumpOp& Op = ByteCode.GetOpAt<FRigVMJumpOp>(Instruction);
-				ContextPublicData.InstructionIndex -= Op.InstructionIndex;
+				ContextPublicData.InstructionIndex = FRigVMJumpOp::GetTargetInstruction(Instruction.OpCode, ContextPublicData.InstructionIndex, Op.InstructionIndex);
 				break;
 			}
 			case ERigVMOpCode::JumpAbsoluteIf:
@@ -2154,7 +2349,7 @@ ERigVMExecuteResult URigVM::ExecuteInstructions(FRigVMExtendedExecuteContext& Co
 				const bool Condition = *(bool*)Handle.GetInputData();
 				if (Condition == Op.Condition)
 				{
-					ContextPublicData.InstructionIndex = Op.InstructionIndex;
+					ContextPublicData.InstructionIndex = FRigVMJumpOp::GetTargetInstruction(Instruction.OpCode, ContextPublicData.InstructionIndex, Op.InstructionIndex);
 				}
 				else
 				{
@@ -2172,7 +2367,7 @@ ERigVMExecuteResult URigVM::ExecuteInstructions(FRigVMExtendedExecuteContext& Co
 				const bool Condition = *(bool*)Handle.GetInputData();
 				if (Condition == Op.Condition)
 				{
-					ContextPublicData.InstructionIndex += Op.InstructionIndex;
+					ContextPublicData.InstructionIndex = FRigVMJumpOp::GetTargetInstruction(Instruction.OpCode, ContextPublicData.InstructionIndex, Op.InstructionIndex);
 				}
 				else
 				{
@@ -2190,7 +2385,7 @@ ERigVMExecuteResult URigVM::ExecuteInstructions(FRigVMExtendedExecuteContext& Co
 				const bool Condition = *(bool*)Handle.GetInputData();
 				if (Condition == Op.Condition)
 				{
-					ContextPublicData.InstructionIndex -= Op.InstructionIndex;
+					ContextPublicData.InstructionIndex = FRigVMJumpOp::GetTargetInstruction(Instruction.OpCode, ContextPublicData.InstructionIndex, Op.InstructionIndex);
 				}
 				else
 				{
@@ -2329,20 +2524,15 @@ ERigVMExecuteResult URigVM::ExecuteInstructions(FRigVMExtendedExecuteContext& Co
 				const int32 FirstHandleIndex = FirstHandleForInstruction[ContextPublicData.InstructionIndex];
 				FRigVMMemoryHandle Handle = GeneratePatchedHandle(Context.ExternalVariableRuntimeData, VMMemoryPtrs, PreCachedMemoryHandlesMemoryType[FirstHandleIndex], PreCachedMemoryHandles[FirstHandleIndex]);
 
-				FRigVMInstructionSetExecuteState& ExecutionState = *(FRigVMInstructionSetExecuteState*)Handle.GetPrivateData(INDEX_NONE);
+				FRigVMInstructionSetExecuteState& ExecuteState = *(FRigVMInstructionSetExecuteState*)Handle.GetPrivateData(INDEX_NONE);
 
 				if((Op.StartInstruction != INDEX_NONE) &&
 					(Op.EndInstruction != INDEX_NONE) &&
 					(Op.EndInstruction >= Op.StartInstruction))
 				{
-					const uint32 SliceHash = Context.GetSliceHash();
-					uint32& StoredHash = ExecutionState.SliceHashToNumInstruction.FindOrAdd(SliceHash, UINT32_MAX);
-						
-					const uint32 Hash = GetTypeHash(ContextPublicData.GetNumExecutions());
-					if(StoredHash != Hash)
+					if (ExecuteState.RequiresExecute(Context.GetSliceHash(), Context.GetHashForLazyBranch()))
 					{
 						ExecuteInstructions(Context, Op.StartInstruction, Op.EndInstruction);
-						StoredHash = Hash;
 					}
 				}
 
@@ -2424,6 +2614,22 @@ ERigVMExecuteResult URigVM::ExecuteInstructions(FRigVMExtendedExecuteContext& Co
 					RigVMProfilingInfo->AddInstructionCyclesDuringLastRun(CurrentInstructionIndex, Cycles);
 				}
 
+				if (Instruction.OpCode == ERigVMOpCode::InvokeCallable)
+				{
+					const FRigVMInvokeCallableOp& Op = ByteCode.GetOpAt<FRigVMInvokeCallableOp>(Instruction);
+					const FRigVMCallableInfo* Callable = ByteCode.GetCallable(Op.CallableIndex);
+					check(Callable);
+
+					if (RigVMProfilingInfo->GetCallableCyclesDuringLastRun(Callable->Index) == UINT64_MAX)
+					{
+						RigVMProfilingInfo->SetCallableCyclesDuringLastRun(Callable->Index, Cycles);
+					}
+					else
+					{
+						RigVMProfilingInfo->AddCallableCyclesDuringLastRun(Callable->Index, Cycles);
+					}
+				}
+
 				RigVMProfilingInfo->SetStartCycles(EndCycles);
 				RigVMProfilingInfo->AddOverallCycles(Cycles);
 			}
@@ -2965,8 +3171,8 @@ TArray<FString> URigVM::DumpByteCodeAsTextArray(FRigVMExtendedExecuteContext& Co
 			case ERigVMOpCode::Execute:
 			{
 				const FRigVMExecuteOp& Op = ByteCode.GetOpAt<FRigVMExecuteOp>(Instructions[InstructionIndex]);
-				FString FunctionName = FunctionNames[Op.FunctionIndex].ToString();
-				FRigVMOperandArray Operands = ByteCode.GetOperandsForExecuteOp(Instructions[InstructionIndex]);
+				FString FunctionName = FunctionNames[Op.CallableIndex].ToString();
+				FRigVMOperandArray Operands = ByteCode.GetOperandsForCallableOp(Instructions[InstructionIndex]);
 
 				TArray<FString> Labels;
 				for (const FRigVMOperand& Operand : Operands)
@@ -2977,6 +3183,23 @@ TArray<FString> URigVM::DumpByteCodeAsTextArray(FRigVMExtendedExecuteContext& Co
 				ResultLine = FString::Printf(TEXT("%s(%s)"), *FunctionName, *RigVMStringUtils::JoinStrings(Labels, TEXT(",")));
 				break;
 			}
+			case ERigVMOpCode::InvokeCallable:
+			{
+				const FRigVMInvokeCallableOp& Op = ByteCode.GetOpAt<FRigVMInvokeCallableOp>(Instructions[InstructionIndex]);
+				const FRigVMCallableInfo* CallableInfo = ByteCode.GetCallable(Op.CallableIndex);
+				check(CallableInfo);
+				const FString CallableName = CallableInfo->Name.ToString();
+				FRigVMOperandArray Operands = ByteCode.GetOperandsForCallableOp(Instructions[InstructionIndex]);
+
+				TArray<FString> Labels;
+				for (const FRigVMOperand& Operand : Operands)
+				{
+					Labels.Add(GetOperandLabel(Context, Operand, OperandFormatFunction));
+				}
+
+				ResultLine = FString::Printf(TEXT("%s(%s)"), *CallableName, *RigVMStringUtils::JoinStrings(Labels, TEXT(",")));
+				break;
+			}
 			case ERigVMOpCode::Zero:
 			{
 				const FRigVMUnaryOp& Op = ByteCode.GetOpAt<FRigVMUnaryOp>(Instructions[InstructionIndex]);
@@ -3433,7 +3656,19 @@ void URigVM::RefreshExternalPropertyPaths()
 	{
 		ExternalPropertyPaths[PropertyPathIndex] = FRigVMPropertyPath();
 
-		const int32 PropertyIndex = ExternalPropertyPathDescriptions[PropertyPathIndex].PropertyIndex;
+		int32 PropertyIndex = ExternalPropertyPathDescriptions[PropertyPathIndex].PropertyIndex;
+		if(ExternalVariables.IsValidIndex(PropertyIndex))
+		{
+			// it's possible for the external variables in a different order
+			const FString& ExpectedHeadCPPType = ExternalPropertyPathDescriptions[PropertyPathIndex].HeadCPPType;
+			if (ExternalVariables[PropertyIndex].GetExtendedCPPType() != ExpectedHeadCPPType)
+			{
+				PropertyIndex = ExternalVariables.IndexOfByPredicate([ExpectedHeadCPPType](const FRigVMExternalVariableDef& ExternalVariableDef) -> bool
+				{
+					return ExternalVariableDef.GetExtendedCPPType() == ExpectedHeadCPPType;	
+				});
+			}
+		}
 		if(ExternalVariables.IsValidIndex(PropertyIndex))
 		{
 			check(ExternalVariables[PropertyIndex].GetProperty());
@@ -3441,6 +3676,12 @@ void URigVM::RefreshExternalPropertyPaths()
 			ExternalPropertyPaths[PropertyPathIndex] = FRigVMPropertyPath(
 				ExternalVariables[PropertyIndex].GetProperty(),
 				ExternalPropertyPathDescriptions[PropertyPathIndex].SegmentPath);
+#if WITH_EDITOR
+			if (!ExternalPropertyPaths[PropertyPathIndex].IsValid())
+			{
+				UE_LOG(LogRigVM, Warning, TEXT("Unable to resolve property path '%s' for property '%s'."), *ExternalPropertyPathDescriptions[PropertyPathIndex].SegmentPath, *ExternalVariables[PropertyIndex].GetProperty()->GetPathName());
+			}
+#endif
 		}
 	}
 }
@@ -3521,12 +3762,12 @@ void URigVM::GetRequiredPlugins(TArray<FString>& OutPlugins)
 	}
 }
 
-void URigVM::SetupInstructionTracking(FRigVMExtendedExecuteContext& Context, int32 InInstructionCount)
+void URigVM::SetupInstructionTracking(FRigVMExtendedExecuteContext& Context, int32 InInstructionCount, int32 InCallableCount)
 {
 #if WITH_EDITOR
 	if (FRigVMProfilingInfo* RigVMProfilingInfo = Context.GetRigVMProfilingInfo())
 	{
-		RigVMProfilingInfo->SetupInstructionTracking(InInstructionCount, true);
+		RigVMProfilingInfo->SetupInstructionTracking(InInstructionCount, InCallableCount, true);
 	}
 #endif
 }
diff --git a/Engine/Plugins/Runtime/RigVM/Source/RigVM/Private/RigVMCore/RigVMByteCode.cpp b/Engine/Plugins/Runtime/RigVM/Source/RigVM/Private/RigVMCore/RigVMByteCode.cpp
--- a/Engine/Plugins/Runtime/RigVM/Source/RigVM/Private/RigVMCore/RigVMByteCode.cpp
+++ b/Engine/Plugins/Runtime/RigVM/Source/RigVM/Private/RigVMCore/RigVMByteCode.cpp
@@ -22,10 +22,45 @@ ERigVMExecuteResult FRigVMPredicateBranch::Execute(FRigVMExtendedExecuteContext&
 	return ERigVMExecuteResult::Failed;
 }
 
-void FRigVMExecuteOp::Serialize(FArchive& Ar)
+void FRigVMCallableArgument::Serialize(FArchive& Ar)
+{
+	Ar << Name;
+	Ar << TypeString;
+	Ar << InterfaceOperand;
+	Ar << ForwardedOperand;
+	Ar << Direction;
+}
+
+int32 FRigVMCallableInfo::GetArgumentIndex(const FName& InName) const
+{
+	return Arguments.IndexOfByPredicate([InName](const FRigVMCallableArgument& Argument) -> bool
+	{
+		return Argument.Name == InName;
+	});
+}
+
+const FRigVMCallableArgument* FRigVMCallableInfo::FindArgument(const FName& InName) const
+{
+	return Arguments.FindByPredicate([InName](const FRigVMCallableArgument& Argument) -> bool
+	{
+		return Argument.Name == InName;
+	});
+}
+
+void FRigVMCallableInfo::Serialize(FArchive& Ar)
+{
+	Ar << Index;
+	Ar << Name;
+	Ar << Hash;
+	Ar << Arguments;
+	Ar << FirstInstruction;
+	Ar << LastInstruction;
+}
+
+void FRigVMInvokeCallableOp::Serialize(FArchive& Ar)
 {
 	Ar << OpCode;
-	Ar << FunctionIndex;
+	Ar << CallableIndex;
 
 	if(Ar.IsLoading())
 	{
@@ -37,7 +72,7 @@ void FRigVMExecuteOp::Serialize(FArchive& Ar)
 		}
 		else
 		{
-			check(OpCode == ERigVMOpCode::Execute);
+			check(FRigVMByteCode::IsCallableOpCode(OpCode));
 			Ar << ArgumentCount;
 		}
 	}
@@ -45,6 +80,11 @@ void FRigVMExecuteOp::Serialize(FArchive& Ar)
 	{
 		Ar << ArgumentCount;
 	}
+}
+
+void FRigVMExecuteOp::Serialize(FArchive& Ar)
+{
+	FRigVMInvokeCallableOp::Serialize(Ar);
 
 	if (Ar.CustomVer(FRigVMObjectVersion::GUID) >= FRigVMObjectVersion::PredicatesAddedToExecuteOps)
 	{
@@ -222,9 +262,9 @@ FRigVMInstructionArray::FRigVMInstructionArray(const FRigVMByteCode& InByteCode,
 				}
 			}
 
-			if (OpCode == ERigVMOpCode::Execute)
+			if (FRigVMByteCode::IsCallableOpCode(OpCode))
 			{
-				const int32 OperandByteIndex = ByteIndex + sizeof(FRigVMExecuteOp);
+				const int32 OperandByteIndex = ByteIndex + FRigVMByteCode::GetCallableOpSize(OpCode);
 
 				Alignment = InByteCode.GetOperandAlignment();
 				if (Alignment > 0)
@@ -323,7 +363,7 @@ void FRigVMByteCode::Save(FArchive& Ar)
 				FRigVMExecuteOp Op = GetOpAt<FRigVMExecuteOp>(Instruction.ByteCodeIndex);
 				Ar << Op;
 
-				FRigVMOperandArray Operands = GetOperandsForExecuteOp(Instruction);
+				FRigVMOperandArray Operands = GetOperandsForCallableOp(Instruction);
 				int32 OperandCount = (int32)Op.GetOperandCount();
 				ensure(OperandCount == Operands.Num());
 
@@ -334,6 +374,22 @@ void FRigVMByteCode::Save(FArchive& Ar)
 				}
 				break;
 			}
+			case ERigVMOpCode::InvokeCallable:
+			{
+				FRigVMInvokeCallableOp Op = GetOpAt<FRigVMInvokeCallableOp>(Instruction.ByteCodeIndex);
+				Ar << Op;
+
+				FRigVMOperandArray Operands = GetOperandsForCallableOp(Instruction);
+				int32 OperandCount = (int32)Op.GetOperandCount();
+				ensure(OperandCount == Operands.Num());
+					
+				for (int32 OperandIndex = 0; OperandIndex < OperandCount; OperandIndex++)
+				{
+					FRigVMOperand Operand = Operands[OperandIndex];
+					Ar << Operand;
+				}
+				break;
+			}
 			case ERigVMOpCode::Copy:
 			{
 				FRigVMCopyOp Op = GetOpAt<FRigVMCopyOp>(Instruction.ByteCodeIndex);
@@ -475,6 +531,9 @@ void FRigVMByteCode::Save(FArchive& Ar)
 	Ar << TempBranchInfos;
 	UE_RIGVM_ARCHIVETRACE_ENTRY(Ar, TEXT("BranchInfos"));
 
+	Ar << CallableInfos;	
+	UE_RIGVM_ARCHIVETRACE_ENTRY(Ar, TEXT("CallableInfos"));
+
 	Ar << PublicContextAssetPath;
 	UE_RIGVM_ARCHIVETRACE_ENTRY(Ar, TEXT("PublicContextAssetPath"));
 }
@@ -521,7 +580,24 @@ void FRigVMByteCode::Load(FArchive& Ar)
 					Operands.Add(Operand);
 				}
 
-				AddExecuteOp(Op.FunctionIndex, Operands, Op.FirstPredicateIndex, Op.PredicateCount);
+				AddExecuteOp(Op.CallableIndex, Operands, Op.FirstPredicateIndex, Op.PredicateCount);
+				break;
+			}
+			case ERigVMOpCode::InvokeCallable:
+			{
+				FRigVMInvokeCallableOp Op;
+				Ar << Op;
+
+				int32 OperandCount = Op.GetOperandCount();
+				TArray<FRigVMOperand> Operands;
+				for (int32 OperandIndex = 0; OperandIndex < OperandCount; OperandIndex++)
+				{
+					FRigVMOperand Operand;
+					Ar << Operand;
+					Operands.Add(Operand);
+				}
+
+				AddInvokeCallableOpByIndex(Op.CallableIndex, Operands);
 				break;
 			}
 			case ERigVMOpCode::Copy:
@@ -684,6 +760,15 @@ void FRigVMByteCode::Load(FArchive& Ar)
 		BranchInfoLookup.Reset();
 	}
 
+	if (Ar.CustomVer(FRigVMObjectVersion::GUID) >= FRigVMObjectVersion::RigVMCallables)
+	{
+		Ar << CallableInfos;
+	}
+	else
+	{
+		CallableInfos.Reset();
+	}
+
 	if (Ar.CustomVer(FRigVMObjectVersion::GUID) >= FRigVMObjectVersion::VMBytecodeStorePublicContextPathAsTopLevelAssetPath)
 	{
 		Ar << PublicContextAssetPath;
@@ -707,6 +792,7 @@ void FRigVMByteCode::Reset()
 	Entries.Reset();
 	BranchInfos.Reset();
 	BranchInfoLookup.Reset();
+	CallableInfos.Reset();
 	PublicContextAssetPath.Reset();
 
 #if WITH_EDITORONLY_DATA
@@ -728,6 +814,9 @@ void FRigVMByteCode::Empty()
 	bByteCodeIsAligned = false;
 	NumInstructions = 0;
 	Entries.Empty();
+	BranchInfos.Empty();
+	BranchInfoLookup.Empty();
+	CallableInfos.Empty();
 	PublicContextAssetPath.Reset();
 
 #if WITH_EDITORONLY_DATA
@@ -743,6 +832,27 @@ void FRigVMByteCode::Empty()
 #endif
 }
 
+void FRigVMByteCode::Shrink()
+{
+	ByteCode.Shrink();
+	Entries.Shrink();
+	BranchInfos.Shrink();
+	BranchInfoLookup.Shrink();
+	CallableInfos.Shrink();
+
+#if WITH_EDITORONLY_DATA
+	SubjectPerInstruction.Shrink();
+	SubjectToInstructions.Shrink();
+	CallPathPerInstruction.Shrink();
+	CallPathToInstructions.Shrink();
+	CallstackPerInstruction.Shrink();
+	CallstackHashToInstructions.Shrink();
+	CallstackHashPerInstruction.Shrink();
+	InputOperandsPerInstruction.Shrink();
+	OutputOperandsPerInstruction.Shrink();
+#endif
+}
+
 uint32 FRigVMByteCode::GetByteCodeHash() const
 {
 	uint32 Hash = 0;
@@ -763,6 +873,11 @@ uint32 FRigVMByteCode::GetByteCodeHash() const
 		Hash = HashCombine(Hash, GetTypeHash(BranchInfo));
 	}
 
+	for(const FRigVMCallableInfo& CallableInfo : CallableInfos)
+	{
+		Hash = HashCombine(Hash, GetTypeHash(CallableInfo));
+	}
+
 	return Hash;
 }
 
@@ -774,7 +889,17 @@ uint32 FRigVMByteCode::GetOperatorHash(const FRigVMInstruction& InInstruction) c
 		{
 			const FRigVMExecuteOp& Op = GetOpAt<FRigVMExecuteOp>(InInstruction);
 			uint32 Hash = GetTypeHash(Op);
-			for(const FRigVMOperand& Operand : GetOperandsForExecuteOp(InInstruction))
+			for(const FRigVMOperand& Operand : GetOperandsForCallableOp(InInstruction))
+			{
+				Hash = HashCombine(Hash, GetTypeHash(Operand));
+			}
+			return Hash;
+		}
+		case ERigVMOpCode::InvokeCallable:
+		{
+			const FRigVMInvokeCallableOp& Op = GetOpAt<FRigVMInvokeCallableOp>(InInstruction);
+			uint32 Hash = GetTypeHash(Op);
+			for(const FRigVMOperand& Operand : GetOperandsForCallableOp(InInstruction))
 			{
 				Hash = HashCombine(Hash, GetTypeHash(Operand));
 			}
@@ -907,9 +1032,87 @@ int32 FRigVMByteCode::FindEntryIndex(const FName& InEntryName) const
 	return INDEX_NONE;
 }
 
+int32 FRigVMByteCode::NumCallables() const
+{
+	return CallableInfos.Num();
+}
+
+bool FRigVMByteCode::IsValidCallableIndex(int32 InCallableIndex) const
+{
+	return CallableInfos.IsValidIndex(InCallableIndex);
+}
+
+bool FRigVMByteCode::IsValidCallableIndex(uint16 InCallableIndex) const
+{
+	return IsValidCallableIndex(IntCastChecked<int32, uint16>(InCallableIndex));
+}
+
+const FRigVMCallableInfo* FRigVMByteCode::GetCallable(int32 InCallableIndex) const
+{
+	if (ensure(CallableInfos.IsValidIndex(InCallableIndex)))
+	{
+		ensure(CallableInfos[InCallableIndex].Index == InCallableIndex);
+		return &CallableInfos[InCallableIndex];
+	}
+	return nullptr;
+}
+
+FRigVMCallableInfo* FRigVMByteCode::GetCallable(int32 InCallableIndex)
+{
+	if (ensure(CallableInfos.IsValidIndex(InCallableIndex)))
+	{
+		ensure(CallableInfos[InCallableIndex].Index == InCallableIndex);
+		return &CallableInfos[InCallableIndex];
+	}
+	return nullptr;
+}
+
+const FRigVMCallableInfo* FRigVMByteCode::GetCallable(uint16 InCallableIndex) const
+{
+	return GetCallable(IntCastChecked<int32, uint16>(InCallableIndex));
+}
+
+FRigVMCallableInfo* FRigVMByteCode::GetCallable(uint16 InCallableIndex)
+{
+	return GetCallable(IntCastChecked<int32, uint16>(InCallableIndex));
+}
+
+int32 FRigVMByteCode::GetCallableIndex(const uint32& InHash) const
+{
+	return CallableInfos.IndexOfByPredicate([InHash](const FRigVMCallableInfo& InCallableInfo) -> bool
+	{
+		return InCallableInfo.Hash == InHash;
+	});
+}
+
+const FRigVMCallableInfo* FRigVMByteCode::FindCallable(const uint32& InHash) const
+{
+	return CallableInfos.FindByPredicate([InHash](const FRigVMCallableInfo& InCallableInfo) -> const FRigVMCallableInfo*
+	{
+		if (InCallableInfo.Hash == InHash)
+		{
+			return &InCallableInfo;
+		}
+		return nullptr;
+	});
+}
+
+const FRigVMCallableInfo* FRigVMByteCode::FindCallable(const FName& InName) const
+{
+	return CallableInfos.FindByPredicate([InName](const FRigVMCallableInfo& InCallableInfo) -> const FRigVMCallableInfo*
+	{
+		if (InCallableInfo.Name == InName)
+		{
+			return &InCallableInfo;
+		}
+		return nullptr;
+	});
+}
+
 int32 FRigVMByteCode::GetOpNumBytesAt(int32 InByteCodeIndex, bool bIncludeOperands) const
 {
-	switch (GetOpCodeAt(InByteCodeIndex))
+	const ERigVMOpCode OpCode = GetOpCodeAt(InByteCodeIndex);
+	switch (OpCode)
 	{
 		case ERigVMOpCode::Execute_0_Operands:
 		case ERigVMOpCode::Execute_1_Operands:
@@ -977,15 +1180,17 @@ int32 FRigVMByteCode::GetOpNumBytesAt(int32 InByteCodeIndex, bool bIncludeOperan
 		case ERigVMOpCode::Execute_63_Operands:
 		case ERigVMOpCode::Execute_64_Operands:
 		case ERigVMOpCode::Execute:
+		case ERigVMOpCode::InvokeCallable:
 		{
-			int32 NumBytes = sizeof(FRigVMExecuteOp);
+			const int32 OpSize = GetCallableOpSize(OpCode); 
+			int32 NumBytes = OpSize;
 			if(bIncludeOperands)
 			{
-				FRigVMExecuteOp ExecuteOp;
-				uint8* ExecuteOpPtr = reinterpret_cast<uint8*>(&ExecuteOp);
-				for (int32 ByteIndex = 0; ByteIndex < sizeof(FRigVMExecuteOp); ByteIndex++)
+				FRigVMInvokeCallableOp CallableOp;
+				uint8* CallableOpPtr = reinterpret_cast<uint8*>(&CallableOp);
+				for (int32 ByteIndex = 0; ByteIndex < sizeof(FRigVMInvokeCallableOp); ByteIndex++)
 				{
-					ExecuteOpPtr[ByteIndex] = ByteCode[InByteCodeIndex + ByteIndex];
+					CallableOpPtr[ByteIndex] = ByteCode[InByteCodeIndex + ByteIndex];
 				}
 
 				if (bByteCodeIsAligned)
@@ -999,7 +1204,7 @@ int32 FRigVMByteCode::GetOpNumBytesAt(int32 InByteCodeIndex, bool bIncludeOperan
 						}
 					}
 				}
-				NumBytes += ExecuteOp.GetOperandCount() * sizeof(FRigVMOperand);
+				NumBytes += CallableOp.GetOperandCount() * sizeof(FRigVMOperand);
 			}
 			return NumBytes;
 		}
@@ -1212,6 +1417,13 @@ int32 FRigVMByteCode::AddExecuteOp(int32 InFunctionIndex, const FRigVMOperandArr
 	}
 	int32 OpByteIndex = AddOp(Op);
 
+#if WITH_EDITOR
+	for (const FRigVMOperand& Operand : InOperands)
+	{
+		verify(Operand.IsValid());
+	}
+#endif
+
 	const int32 OperandsByteIndex = ByteCode.AddZeroed(sizeof(FRigVMOperand) * InOperands.Num());
 	FMemory::Memcpy(ByteCode.GetData() + OperandsByteIndex, InOperands.GetData(), sizeof(FRigVMOperand) * InOperands.Num());
 
@@ -1224,13 +1436,39 @@ int32 FRigVMByteCode::AddExecuteOp(int32 InFunctionIndex, const FRigVMOperandArr
 	return OpByteIndex;
 }
 
-int32 FRigVMByteCode::InlineFunction(const FRigVMByteCode* FunctionByteCode, const FRigVMOperandArray& InOperands)
+int32 FRigVMByteCode::InlineFunction(const FRigVMByteCode* FunctionByteCode, int32 InFirstIndex, int32 InLastIndex)
 {
 	check(FunctionByteCode);
 	check(!FunctionByteCode->bByteCodeIsAligned);
 	const int32 OpByteIndex = ByteCode.Num();
-	ByteCode.Append(FunctionByteCode->ByteCode);
-	NumInstructions += FunctionByteCode->NumInstructions;
+	if (FunctionByteCode->Num() == 0)
+	{
+		return OpByteIndex;
+	}
+
+	if (InFirstIndex == INDEX_NONE && InLastIndex == INDEX_NONE)
+	{
+		ByteCode.Append(FunctionByteCode->ByteCode);
+		NumInstructions += FunctionByteCode->NumInstructions;
+	}
+	else
+	{
+		const FRigVMInstructionArray FunctionInstructions = FunctionByteCode->GetInstructions();
+		check(FunctionInstructions.IsValidIndex(InFirstIndex) && FunctionInstructions.IsValidIndex(InLastIndex));
+
+		const TArrayView<const uint8> BytesView(FunctionByteCode->ByteCode);
+		const int32 NumInstructionsToAdd = InLastIndex - InFirstIndex + 1;
+		const int32 FirstByteIndex = FunctionInstructions[InFirstIndex].ByteCodeIndex;
+		const int32 LastByteIndex =
+			InLastIndex == FunctionInstructions.Num() - 1 ?
+			FunctionByteCode->Num() - 1 :
+			FunctionInstructions[InLastIndex + 1].ByteCodeIndex - 1;
+		const int32 NumBytesToAdd = LastByteIndex - FirstByteIndex + 1;
+		
+		const TArrayView<const uint8> BytesToAdd(BytesView.Slice(FirstByteIndex, NumBytesToAdd));
+		ByteCode.Append(BytesToAdd);
+		NumInstructions += NumInstructionsToAdd;
+	}
 	
 	return OpByteIndex;
 }
@@ -1258,9 +1496,32 @@ FString FRigVMByteCode::DumpToText() const
 			case ERigVMOpCode::Execute:
 			{
 				const FRigVMExecuteOp& Op = GetOpAt<FRigVMExecuteOp>(Instruction.ByteCodeIndex);
-				Line += FString::Printf(TEXT(", FunctionIndex %u"), Op.FunctionIndex);
+				Line += FString::Printf(TEXT(", FunctionIndex %u"), Op.CallableIndex);
+
+				FRigVMOperandArray Operands = GetOperandsForCallableOp(Instruction);
+				if (Operands.Num() > 0)
+				{
+					TArray<FString> OperandsContent;
+					for (const FRigVMOperand& Operand : Operands)
+					{
+						FString OperandContent;
+						FRigVMOperand::StaticStruct()->ExportText(OperandContent, &Operand, &Operand, nullptr, PPF_None, nullptr);
+						OperandsContent.Add(FString::Printf(TEXT("\t%s"), *OperandContent));
+					}
+
+					Line += RigVMStringUtils::JoinStrings(OperandsContent, TEXT("(\n"), TEXT("\n"), TEXT("\n)"));
+				}
+				break;
+			}
+			case ERigVMOpCode::InvokeCallable:
+			{
+				const FRigVMInvokeCallableOp& Op = GetOpAt<FRigVMInvokeCallableOp>(Instruction.ByteCodeIndex);
+				const FRigVMCallableInfo* CallableInfo = GetCallable(Op.CallableIndex);
+				check(CallableInfo);
+				Line += FString::Printf(TEXT(", CallableIndex %u"), Op.CallableIndex);
+				Line += FString::Printf(TEXT(", Function '%s'"), *CallableInfo->Name.ToString());
 
-				FRigVMOperandArray Operands = GetOperandsForExecuteOp(Instruction);
+				FRigVMOperandArray Operands = GetOperandsForCallableOp(Instruction);
 				if (Operands.Num() > 0)
 				{
 					TArray<FString> OperandsContent;
@@ -1432,6 +1693,37 @@ int32 FRigVMByteCode::AddInvokeEntryOp(const FName& InEntryName)
 	return AddOp(FRigVMInvokeEntryOp(InEntryName));
 }
 
+int32 FRigVMByteCode::AddInvokeCallableOp(const uint32& InCallableHash, const FRigVMOperandArray& InOperands)
+{
+	const FRigVMCallableInfo* CallableInfo = FindCallable(InCallableHash);
+	check(CallableInfo);
+	return AddInvokeCallableOpByIndex(IntCastChecked<uint16>(CallableInfo->Index), InOperands);
+}
+
+int32 FRigVMByteCode::AddInvokeCallableOpByIndex(const uint16& InCallableIndex, const FRigVMOperandArray& InOperands)
+{
+	FRigVMInvokeCallableOp Op(InCallableIndex, IntCastChecked<uint16>(InOperands.Num()));
+	int32 OpByteIndex = AddOp(Op);
+
+#if WITH_EDITOR
+	for (const FRigVMOperand& Operand : InOperands)
+	{
+		verify(Operand.IsValid());
+	}
+#endif
+
+	const int32 OperandsByteIndex = ByteCode.AddZeroed(sizeof(FRigVMOperand) * InOperands.Num());
+	FMemory::Memcpy(ByteCode.GetData() + OperandsByteIndex, InOperands.GetData(), sizeof(FRigVMOperand) * InOperands.Num());
+
+	for(int32 Index = 0; Index < InOperands.Num(); Index++)
+	{
+		FRigVMOperand* Operand = reinterpret_cast<FRigVMOperand*>(ByteCode.GetData() + OperandsByteIndex + sizeof(FRigVMOperand) * Index);
+		FRigVMOperand::ZeroPaddedMemoryIfNeeded(Operand);
+	}
+
+	return OpByteIndex;
+}
+
 int32 FRigVMByteCode::AddJumpToBranchOp(FRigVMOperand InBranchNameArg, int32 InFirstBranchInfoIndex)
 {
 	return AddOp(FRigVMJumpToBranchOp(InBranchNameArg, InFirstBranchInfoIndex));
@@ -1478,8 +1770,9 @@ FRigVMOperandArray FRigVMByteCode::GetOperandsForOp(const FRigVMInstruction& InI
 	switch(InInstruction.OpCode)
 	{
 		case ERigVMOpCode::Execute:
+		case ERigVMOpCode::InvokeCallable:
 		{
-			return GetOperandsForExecuteOp(InInstruction);
+			return GetOperandsForCallableOp(InInstruction);
 		}
 		case ERigVMOpCode::Copy:
 		{
@@ -1574,12 +1867,13 @@ FRigVMOperandArray FRigVMByteCode::GetOperandsForOp(const FRigVMInstruction& InI
 
 int32 FRigVMByteCode::GetFirstOperandByteIndex(const FRigVMInstruction& InInstruction) const
 {
-	if (InInstruction.OpCode == ERigVMOpCode::Execute)
+	if (IsCallableOpCode(InInstruction.OpCode))
 	{
+		const int32 OpSize = GetCallableOpSize(InInstruction.OpCode);
 		const int32 ByteCodeIndex = InInstruction.ByteCodeIndex;
 		// if the bytecode is not aligned the OperandAlignment needs to be 0
 		ensure(bByteCodeIsAligned || InInstruction.OperandAlignment == 0);
-		return ByteCodeIndex + sizeof(FRigVMExecuteOp) + InInstruction.OperandAlignment;
+		return ByteCodeIndex + OpSize + InInstruction.OperandAlignment;
 	}
 
 	switch(InInstruction.OpCode)
@@ -1725,6 +2019,11 @@ int32 FRigVMByteCode::GetOpAlignment(ERigVMOpCode InOpCode)
 			static const int32 Alignment = FRigVMExecuteOp::StaticStruct()->GetCppStructOps()->GetAlignment();
 			return Alignment;
 		}
+		case ERigVMOpCode::InvokeCallable:
+		{
+			static const int32 Alignment = FRigVMInvokeCallableOp::StaticStruct()->GetCppStructOps()->GetAlignment();
+			return Alignment;
+		}
 		case ERigVMOpCode::Copy:
 		{
 			static const int32 Alignment = FRigVMCopyOp::StaticStruct()->GetCppStructOps()->GetAlignment();
@@ -1865,7 +2164,7 @@ void FRigVMByteCode::AlignByteCode()
 		const FRigVMInstruction& Instruction = Instructions[InstructionIndex];
 		BytesToReserve += GetOpAlignment(Instruction.OpCode);
 
-		if (Instruction.OpCode == ERigVMOpCode::Execute)
+		if (IsCallableOpCode(Instruction.OpCode))
 		{
 			BytesToReserve += GetOperandAlignment();
 		}
@@ -1900,7 +2199,7 @@ void FRigVMByteCode::AlignByteCode()
 			AlignedByteCode[AlignedByteCodeIndex + ByteIndex] = ByteCode[OriginalByteCodeIndex + ByteIndex];
 		}
 
-		if (Instruction.OpCode == ERigVMOpCode::Execute)
+		if (IsCallableOpCode(Instruction.OpCode))
 		{
 			AlignedByteCodeIndex += NumBytesToCopy;
 
@@ -1915,15 +2214,9 @@ void FRigVMByteCode::AlignByteCode()
 				}
 			}
 
-			FRigVMExecuteOp ExecuteOp;
-			uint8* ExecuteOpPtr = reinterpret_cast<uint8*>(&ExecuteOp);
-			for (int32 ByteIndex = 0; ByteIndex < sizeof(FRigVMExecuteOp); ByteIndex++)
-			{
-				ExecuteOpPtr[ByteIndex] = ByteCode[OriginalByteCodeIndex + ByteIndex];
-			}
-
+			const int32 OperandCount = GetOpAt<FRigVMInvokeCallableOp>(OriginalByteCodeIndex).GetOperandCount();
 			OriginalByteCodeIndex += NumBytesToCopy;
-			NumBytesToCopy = sizeof(FRigVMOperand) * ExecuteOp.GetOperandCount();
+			NumBytesToCopy = sizeof(FRigVMOperand) * OperandCount;
 
 			for (int32 ByteIndex = 0; ByteIndex < NumBytesToCopy; ByteIndex++)
 			{
diff --git a/Engine/Plugins/Runtime/RigVM/Source/RigVM/Private/RigVMCore/RigVMExecuteContext.cpp b/Engine/Plugins/Runtime/RigVM/Source/RigVM/Private/RigVMCore/RigVMExecuteContext.cpp
--- a/Engine/Plugins/Runtime/RigVM/Source/RigVM/Private/RigVMCore/RigVMExecuteContext.cpp
+++ b/Engine/Plugins/Runtime/RigVM/Source/RigVM/Private/RigVMCore/RigVMExecuteContext.cpp
@@ -107,6 +107,9 @@ void FRigVMExtendedExecuteContext::Reset()
 		GetPublicData<>().NumExecutions = 0;
 	}
 
+	Callstack.Reset();
+	CallstackHash = 0;
+
 	ExecutingThreadId = INDEX_NONE;
 
 	EntriesBeingExecuted.Reset();
@@ -135,8 +138,9 @@ void FRigVMExtendedExecuteContext::ResetExecutionState()
 
 	VM = nullptr;
 	Slices.Reset();
-	Slices.Add(FRigVMSlice());
-	SliceOffsets.Reset();
+	Slices.Add(FRigVMSlice::MakeForCallable());
+	SliceOffsetsPerInstruction.Reset();
+	SliceOffsetsPerCallable.Reset();
 	Factory = nullptr;
 }
 
@@ -169,7 +173,8 @@ FRigVMExtendedExecuteContext& FRigVMExtendedExecuteContext::operator =(const FRi
 	}
 
 	Slices = Other.Slices;
-	SliceOffsets = Other.SliceOffsets;
+	SliceOffsetsPerInstruction = Other.SliceOffsetsPerInstruction;
+	SliceOffsetsPerCallable = Other.SliceOffsetsPerCallable;
 
 	LazyBranchExecuteState = Other.LazyBranchExecuteState;
 	ExternalVariableRuntimeData = Other.ExternalVariableRuntimeData;
diff --git a/Engine/Plugins/Runtime/RigVM/Source/RigVM/Private/RigVMCore/RigVMGraphFunctionDefinition.cpp b/Engine/Plugins/Runtime/RigVM/Source/RigVM/Private/RigVMCore/RigVMGraphFunctionDefinition.cpp
--- a/Engine/Plugins/Runtime/RigVM/Source/RigVM/Private/RigVMCore/RigVMGraphFunctionDefinition.cpp
+++ b/Engine/Plugins/Runtime/RigVM/Source/RigVM/Private/RigVMCore/RigVMGraphFunctionDefinition.cpp
@@ -3,9 +3,16 @@
 #include "RigVMCore/RigVMGraphFunctionDefinition.h"
 #include "RigVMCore/RigVMGraphFunctionHost.h"
 #include "RigVMStringUtils.h"
+#include "RigVMCore/RigVMTrait.h"
 
 #include UE_INLINE_GENERATED_CPP_BY_NAME(RigVMGraphFunctionDefinition)
 
+TAutoConsoleVariable<bool> CVarRigVMCompileFunctionsToCallables(
+	TEXT("RigVM.Compiler.CompileFunctionsToCallables"),
+	false,
+	TEXT("Turn functions into callables to save on memory footprint.")
+);
+
 FRigVMPropertyDescription FRigVMFunctionCompilationPropertyDescription::ToPropertyDescription() const
 {
 	return FRigVMPropertyDescription(Name, CPPType, CPPTypeObject.LoadSynchronous(), DefaultValue);
@@ -23,6 +30,147 @@ TArray<FRigVMPropertyDescription> FRigVMFunctionCompilationPropertyDescription::
 	return Result;
 }
 
+void FRigVMFunctionCompilationData::Reset()
+{
+	FScopeLock Lock(&TaskMutex);
+	
+	ByteCode.Reset();
+	FunctionNames.Reset();
+	WorkPropertyDescriptions.Reset();
+	WorkPropertyPathDescriptions.Reset();
+	LiteralPropertyDescriptions.Reset();
+	LiteralPropertyPathDescriptions.Reset();
+	DebugPropertyDescriptions.Reset();
+	DebugPropertyPathDescriptions.Reset();
+	ExternalPropertyDescriptions.Reset();
+	ExternalPropertyPathDescriptions.Reset();
+	ExternalRegisterIndexToVariable.Reset();
+	Operands.Reset();
+	InterfaceOperands.Reset();
+	Hash = 0;
+	bEncounteredSurpressedErrors = false;
+	bSupportsCallable.Reset();
+}
+
+FRigVMFunctionCompilationData& FRigVMFunctionCompilationData::operator=(const FRigVMFunctionCompilationData& InOther)
+{
+	FScopeLock Lock(&TaskMutex);
+	FScopeLock OtherLock(&InOther.TaskMutex);
+	
+	ByteCode = InOther.ByteCode;
+	FunctionNames = InOther.FunctionNames;
+	WorkPropertyDescriptions = InOther.WorkPropertyDescriptions;
+	WorkPropertyPathDescriptions = InOther.WorkPropertyPathDescriptions;
+	LiteralPropertyDescriptions = InOther.LiteralPropertyDescriptions;
+	LiteralPropertyPathDescriptions = InOther.LiteralPropertyPathDescriptions;
+	DebugPropertyDescriptions = InOther.DebugPropertyDescriptions;
+	DebugPropertyPathDescriptions = InOther.DebugPropertyPathDescriptions;
+	ExternalPropertyDescriptions = InOther.ExternalPropertyDescriptions;
+	ExternalPropertyPathDescriptions = InOther.ExternalPropertyPathDescriptions;
+	ExternalRegisterIndexToVariable = InOther.ExternalRegisterIndexToVariable;
+	Operands = InOther.Operands;
+	InterfaceOperands = InOther.InterfaceOperands;
+	Hash = InOther.Hash;
+	bEncounteredSurpressedErrors = InOther.bEncounteredSurpressedErrors;
+	bSupportsCallable = InOther.bSupportsCallable;
+
+	return *this;
+}
+
+bool FRigVMFunctionCompilationData::SupportsCallable() const
+{
+	if (!CVarRigVMCompileFunctionsToCallables->GetBool())
+	{
+		return false;
+	}
+	
+	FScopeLock Lock(&TaskMutex);
+
+	if (bSupportsCallable.IsSet())
+	{
+		return bSupportsCallable.GetValue();
+	}
+
+	bSupportsCallable = true;
+
+	auto Failed = [this]() -> bool
+	{
+		bSupportsCallable = false;
+		return false;
+	};
+
+	// for now we disable functions for nodes that have traits on them.
+	for (const FRigVMFunctionCompilationPropertyDescription& PropertyDescription : WorkPropertyDescriptions)
+	{
+		const UObject* CPPTypeObject = PropertyDescription.CPPTypeObject.Get();
+		if (!CPPTypeObject)
+		{
+			continue;
+		}
+		const UScriptStruct* CPPTypeStruct = Cast<UScriptStruct>(CPPTypeObject);
+		if (!CPPTypeStruct)
+		{
+			continue;
+		}
+
+		if (CPPTypeStruct->IsChildOf(FRigVMTrait::StaticStruct()))
+		{
+			return Failed();
+		}
+	}
+
+	return bSupportsCallable.GetValue();
+}
+
+const TArray<FRigVMFunctionCompilationPropertyDescription>& FRigVMFunctionCompilationData::GetPropertyDescriptions(
+	ERigVMMemoryType InMemoryType) const
+{
+	switch (InMemoryType)
+	{
+		case ERigVMMemoryType::Literal:
+		{
+			return LiteralPropertyDescriptions;
+		}
+		case ERigVMMemoryType::External:
+		{
+				return ExternalPropertyDescriptions;
+		}
+		case ERigVMMemoryType::Debug:
+		{
+			return DebugPropertyDescriptions;
+		}
+		case ERigVMMemoryType::Work:
+		default:
+		{
+			return WorkPropertyDescriptions;
+		}
+	}
+}
+
+const TArray<FRigVMFunctionCompilationPropertyPath>& FRigVMFunctionCompilationData::GetPropertyPathDescriptions(ERigVMMemoryType InMemoryType) const
+{
+	switch (InMemoryType)
+	{
+		case ERigVMMemoryType::Literal:
+		{
+			return LiteralPropertyPathDescriptions;
+		}
+		case ERigVMMemoryType::External:
+		{
+			return ExternalPropertyPathDescriptions;
+		}
+		case ERigVMMemoryType::Debug:
+		{
+			return DebugPropertyPathDescriptions;
+		}
+		case ERigVMMemoryType::Work:
+		default:
+		{
+			return WorkPropertyPathDescriptions;
+		}
+	}
+}
+
 FRigVMGraphFunctionArgument::FRigVMGraphFunctionArgument(const FRigVMExternalVariable& InExternalVariable)
 : Name(InExternalVariable.GetName())
 , DisplayName(InExternalVariable.GetName())
@@ -102,6 +250,8 @@ FRigVMExternalVariable FRigVMGraphFunctionArgument::ToExternalVariable() const
 
 FArchive& operator<<(FArchive& Ar, FRigVMFunctionCompilationData& Data)
 {
+	FScopeLock Lock(&Data.TaskMutex);
+	
 	Ar.UsingCustomVersion(FUE5ReleaseStreamObjectVersion::GUID);
 	Ar.UsingCustomVersion(FFortniteMainBranchObjectVersion::GUID);
 		
@@ -117,6 +267,16 @@ FArchive& operator<<(FArchive& Ar, FRigVMFunctionCompilationData& Data)
 	Ar << Data.ExternalPropertyPathDescriptions;
 	Ar << Data.ExternalRegisterIndexToVariable;
 	Ar << Data.Operands;
+
+	if (Ar.IsLoading() && Ar.CustomVer(FRigVMObjectVersion::GUID) < FRigVMObjectVersion::RigVMCallables)
+	{
+		Data.InterfaceOperands.Reset();
+	}
+	else
+	{
+		Ar << Data.InterfaceOperands;
+	}
+
 	Ar << Data.Hash;
 
 	if(Ar.IsLoading())
@@ -153,6 +313,12 @@ FArchive& operator<<(FArchive& Ar, FRigVMFunctionCompilationData& Data)
 			}
 		}
 	}
+
+	if (Ar.IsLoading())
+	{
+		Data.bSupportsCallable.Reset();
+	}
+	
 	return Ar;
 }
 
@@ -456,6 +622,22 @@ FRigVMGraphFunctionHeader FRigVMGraphFunctionHeader::FindGraphFunctionHeader(con
 	return FRigVMGraphFunctionHeader(); 
 }
 
+FRigVMGraphFunctionHeader FRigVMGraphFunctionHeader::FindGraphFunctionHeaderFromHash(const FString& InFunctionHash, bool* bOutIsPublic, FString* OutErrorMessage)
+{
+	FString HostObjectPathString, FunctionName;
+	if (!InFunctionHash.Split(TEXT(":"), &HostObjectPathString, &FunctionName, ESearchCase::IgnoreCase, ESearchDir::FromEnd))
+	{
+		if (OutErrorMessage)
+		{
+			*OutErrorMessage = FString::Printf(TEXT("Failed to find the Host object for hash '%s'."), *InFunctionHash);
+		}
+		return FRigVMGraphFunctionHeader();
+	}
+
+	FSoftObjectPath HostObjectPath(HostObjectPathString);
+	return FindGraphFunctionHeader(HostObjectPath, *FunctionName, bOutIsPublic, OutErrorMessage);
+}
+
 TArray<FRigVMExternalVariable> FRigVMGraphFunctionHeader::GetExternalVariables() const
 {
 	TArray<FRigVMExternalVariable> Result;
diff --git a/Engine/Plugins/Runtime/RigVM/Source/RigVM/Private/RigVMCore/RigVMMemoryCommon.cpp b/Engine/Plugins/Runtime/RigVM/Source/RigVM/Private/RigVMCore/RigVMMemoryCommon.cpp
--- a/Engine/Plugins/Runtime/RigVM/Source/RigVM/Private/RigVMCore/RigVMMemoryCommon.cpp
+++ b/Engine/Plugins/Runtime/RigVM/Source/RigVM/Private/RigVMCore/RigVMMemoryCommon.cpp
@@ -11,6 +11,19 @@
 
 ////////////////////////////////////////////////////////////////////////////////////////////////////////////////////////
 
+bool FRigVMInstructionSetExecuteState::RequiresExecute(const uint32& InSliceHash, const uint32& InHashForLazyBranch)
+{
+	uint32& StoredHash = SliceHashToCallStackHash.FindOrAdd(InSliceHash, UINT32_MAX);
+	if(StoredHash != InHashForLazyBranch)
+	{
+		StoredHash = InHashForLazyBranch;
+		return true;
+	}
+	return false;
+}
+
+////////////////////////////////////////////////////////////////////////////////////////////////////////////////////////
+
 void FRigVMOperand::Serialize(FArchive& Ar)
 {
 	Ar << MemoryType;
@@ -18,6 +31,12 @@ void FRigVMOperand::Serialize(FArchive& Ar)
 	Ar << RegisterOffset;
 }
 
+FString FRigVMOperand::ToString() const
+{
+	const FString MemoryPrefix = StaticEnum<ERigVMMemoryType>()->GetDisplayNameTextByValue((int64)GetMemoryType()).ToString();
+	return FString::Printf(TEXT("%s(%d,%d)"), *MemoryPrefix, GetRegisterIndex(), GetRegisterOffset());
+}
+
 void FRigVMMemoryStorageImportErrorContext::Serialize(const TCHAR* V, ELogVerbosity::Type Verbosity, const class FName& Category)
 {
 	if (bLogErrors)
diff --git a/Engine/Plugins/Runtime/RigVM/Source/RigVM/Private/RigVMCore/RigVMMemoryStorage.cpp b/Engine/Plugins/Runtime/RigVM/Source/RigVM/Private/RigVMCore/RigVMMemoryStorage.cpp
--- a/Engine/Plugins/Runtime/RigVM/Source/RigVM/Private/RigVMCore/RigVMMemoryStorage.cpp
+++ b/Engine/Plugins/Runtime/RigVM/Source/RigVM/Private/RigVMCore/RigVMMemoryStorage.cpp
@@ -103,13 +103,10 @@ ERigVMExecuteResult FRigVMLazyBranch::ExecuteIfRequired(FRigVMExtendedExecuteCon
 {
 	check(VM);
 
-	const uint32 Hash = GetTypeHash(Context.GetNumExecutions());
-	uint32& StoredHash = Context.LazyBranchExecuteState[BranchInfo.Index].SliceHashToNumInstruction.FindOrAdd(InSliceHash, UINT32_MAX);
-	if(Hash != StoredHash)
+	FRigVMInstructionSetExecuteState& ExecuteState = Context.LazyBranchExecuteState[BranchInfo.Index];
+	if(ExecuteState.RequiresExecute(InSliceHash, Context.GetHashForLazyBranch()))
 	{
-		const ERigVMExecuteResult Result = Execute(Context);
-		StoredHash = Hash;
-		return Result;
+		return Execute(Context);
 	}
 	
 	return ERigVMExecuteResult::Succeeded;
diff --git a/Engine/Plugins/Runtime/RigVM/Source/RigVM/Private/RigVMCore/RigVMMemoryStorageStruct.cpp b/Engine/Plugins/Runtime/RigVM/Source/RigVM/Private/RigVMCore/RigVMMemoryStorageStruct.cpp
--- a/Engine/Plugins/Runtime/RigVM/Source/RigVM/Private/RigVMCore/RigVMMemoryStorageStruct.cpp
+++ b/Engine/Plugins/Runtime/RigVM/Source/RigVM/Private/RigVMCore/RigVMMemoryStorageStruct.cpp
@@ -612,6 +612,12 @@ void FRigVMMemoryStorageStruct::RefreshPropertyPaths()
 			PropertyPaths[PropertyPathIndex] = FRigVMPropertyPath(
 				LinkedProperties[PropertyIndex],
 				PropertyPathDescriptions[PropertyPathIndex].SegmentPath);
+#if WITH_EDITOR
+			if (!PropertyPaths[PropertyPathIndex].IsValid())
+			{
+				UE_LOG(LogRigVM, Warning, TEXT("Unable to resolve property path '%s' for property '%s'."), *PropertyPathDescriptions[PropertyPathIndex].SegmentPath, *LinkedProperties[PropertyIndex]->GetPathName());
+			}
+#endif
 		}
 	}
 }
diff --git a/Engine/Plugins/Runtime/RigVM/Source/RigVM/Private/RigVMCore/RigVMProfilingInfo.cpp b/Engine/Plugins/Runtime/RigVM/Source/RigVM/Private/RigVMCore/RigVMProfilingInfo.cpp
--- a/Engine/Plugins/Runtime/RigVM/Source/RigVM/Private/RigVMCore/RigVMProfilingInfo.cpp
+++ b/Engine/Plugins/Runtime/RigVM/Source/RigVM/Private/RigVMCore/RigVMProfilingInfo.cpp
@@ -4,22 +4,27 @@
 
 #include UE_INLINE_GENERATED_CPP_BY_NAME(RigVMProfilingInfo)
 
-void FRigVMInstructionVisitInfo::SetupInstructionTracking(int32 InInstructionCount)
+void FRigVMInstructionVisitInfo::SetupInstructionTracking(int32 InInstructionCount, int32 InCallableCount)
 {
 #if WITH_EDITOR
 	ResetInstructionVisitedDuringLastRun(InInstructionCount);
 	SetNumInstructionVisitedDuringLastRunZeroed(InInstructionCount);
 	ResetInstructionVisitOrder(InInstructionCount);
+	ResetCallableVisitedDuringLastRun(InCallableCount);
+	SetNumCallableVisitedDuringLastRunZeroed(InCallableCount);
+	ResetCallableVisitOrder(InCallableCount);
 #endif
 }
 
-void FRigVMProfilingInfo::SetupInstructionTracking(int32 InInstructionCount, bool bEnableProfiling)
+void FRigVMProfilingInfo::SetupInstructionTracking(int32 InInstructionCount, int32 InCallableCount, bool bEnableProfiling)
 {
 #if WITH_EDITOR
 	ResetInstructionCyclesDuringLastRun(InInstructionCount);
+	ResetCallableCyclesDuringLastRun(InCallableCount);
 	if (bEnableProfiling)
 	{
 		InitInstructionCyclesDuringLastRunValues(InInstructionCount, UINT64_MAX);
+		InitCallableCyclesDuringLastRunValues(InCallableCount, UINT64_MAX);
 	}
 #endif
 }
diff --git a/Engine/Plugins/Runtime/RigVM/Source/RigVM/Private/RigVMCore/RigVMPropertyPath.cpp b/Engine/Plugins/Runtime/RigVM/Source/RigVM/Private/RigVMCore/RigVMPropertyPath.cpp
--- a/Engine/Plugins/Runtime/RigVM/Source/RigVM/Private/RigVMCore/RigVMPropertyPath.cpp
+++ b/Engine/Plugins/Runtime/RigVM/Source/RigVM/Private/RigVMCore/RigVMPropertyPath.cpp
@@ -1,6 +1,7 @@
 // Copyright Epic Games, Inc. All Rights Reserved.
 
 #include "RigVMCore/RigVMPropertyPath.h"
+#include "RigVMModule.h"
 
 ////////////////////////////////////////////////////////////////////////////////////////////////////////////////////////
 
@@ -55,6 +56,9 @@ FRigVMPropertyPath::FRigVMPropertyPath(const FProperty* InProperty, const FStrin
 			else
 			{
 				Segments.Empty();
+#if UE_BUILD_DEBUG
+				UE_LOG(LogRigVM, Warning, TEXT("Unable to resolve property path '%s' for property '%s'."), *InSegmentPath, *InProperty->GetPathName())
+#endif
 				return;
 			}
 		}
@@ -79,6 +83,9 @@ FRigVMPropertyPath::FRigVMPropertyPath(const FProperty* InProperty, const FStrin
 		else if(CastField<FSetProperty>(Property))
 		{
 			Segments.Empty();
+#if UE_BUILD_DEBUG
+			UE_LOG(LogRigVM, Warning, TEXT("Unable to resolve property path '%s' for property '%s'. FSetProperty is not supported."), *InSegmentPath, *InProperty->GetPathName())
+#endif
 			return;
 		}
 		else
diff --git a/Engine/Plugins/Runtime/RigVM/Source/RigVM/Private/RigVMTrace.cpp b/Engine/Plugins/Runtime/RigVM/Source/RigVM/Private/RigVMTrace.cpp
--- a/Engine/Plugins/Runtime/RigVM/Source/RigVM/Private/RigVMTrace.cpp
+++ b/Engine/Plugins/Runtime/RigVM/Source/RigVM/Private/RigVMTrace.cpp
@@ -390,7 +390,7 @@ bool FRigVMTrace::RestoreRigVMExecuteData(FRigVMExtendedExecuteContext& InOutCon
 	// retore the instruction visited tracking
 	if (InOutContext.InstructionVisitInfo)
 	{
-		InOutContext.InstructionVisitInfo->SetupInstructionTracking(VM->GetByteCode().GetNumInstructions());
+		InOutContext.InstructionVisitInfo->SetupInstructionTracking(VM->GetByteCode().GetNumInstructions(), VM->GetByteCode().NumCallables());
 
 		if (Entries.Num() > 0)
 		{
@@ -410,7 +410,7 @@ bool FRigVMTrace::RestoreRigVMExecuteData(FRigVMExtendedExecuteContext& InOutCon
 	// restore the profiling info
 	if (InOutContext.ProfilingInfo)
 	{
-		InOutContext.ProfilingInfo->SetupInstructionTracking(VM->GetByteCode().GetNumInstructions(), true);
+		InOutContext.ProfilingInfo->SetupInstructionTracking(VM->GetByteCode().GetNumInstructions(), VM->GetByteCode().NumCallables(), true);
 		InOutContext.ProfilingInfo->OverallCycles = InExecuteData.OverallCycles;
 		InOutContext.ProfilingInfo->LastExecutionMicroSeconds = static_cast<double>(InOutContext.ProfilingInfo->OverallCycles) * FPlatformTime::GetSecondsPerCycle() * 1000.0 * 1000.0;
 		InOutContext.ProfilingInfo->InstructionCyclesDuringLastRun = InExecuteData.InstructionCycles;
diff --git a/Engine/Plugins/Runtime/RigVM/Source/RigVM/Public/RigVMCore/RigVM.h b/Engine/Plugins/Runtime/RigVM/Source/RigVM/Public/RigVMCore/RigVM.h
--- a/Engine/Plugins/Runtime/RigVM/Source/RigVM/Public/RigVMCore/RigVM.h
+++ b/Engine/Plugins/Runtime/RigVM/Source/RigVM/Public/RigVMCore/RigVM.h
@@ -243,7 +243,7 @@ public:
 	UFUNCTION(BlueprintCallable, Category = RigVM, meta = (DeprecatedFunction, DeprecationMessage = "This function has been deprecated and it is no longer supported."))
 	virtual bool Execute(FRigVMExtendedExecuteContext& Context, const FName& InEntryName = NAME_None) { return false; }
 
-	UE_DEPRECATED(5.3, "Please, use ExecuteLazyBranch with Context param")
+	UE_DEPRECATED(5.3, "Please, use ExecuteBranch with Context param")
 	ERigVMExecuteResult ExecuteLazyBranch(const FRigVMBranchInfo& InBranchToRun) { return ERigVMExecuteResult::Failed; }
 
 	// Executes a single branch on the VM. We assume that the memory is already set correctly at this point.
@@ -259,7 +259,12 @@ public:
 	UFUNCTION()
 	RIGVM_API virtual int32 AddRigVMFunction(UScriptStruct* InRigVMStruct, const FName& InMethodName);
 	RIGVM_API virtual int32 AddRigVMFunction(const FString& InFunctionName);
+	RIGVM_API virtual int32 AddRigVMFunction(const FName& InFunctionFName);
 
+private:
+	RIGVM_API virtual int32 AddRigVMFunctionImpl(const FName& InFunctionFName, const FString& InFunctionName);
+
+public:
 	// Returns the name of a function given its index
 	UFUNCTION()
 	RIGVM_API virtual FString GetRigVMFunctionName(int32 InFunctionIndex) const;
@@ -472,6 +477,16 @@ public:
 		return 0;
 	}
 
+	// Returns the number of times a callable has been hit
+	int32 GetCallableVisitedCount(const FRigVMExtendedExecuteContext& Context, int32 InIndex) const
+	{
+		if (const FRigVMInstructionVisitInfo* InstructionVisitInfo = Context.GetRigVMInstructionVisitInfo())
+		{
+			return InstructionVisitInfo->GetCallableVisitedCountDuringLastRun(InIndex);
+		}
+		return 0;
+	}
+
 	UE_DEPRECATED(5.3, "Please, use GetInstructionCycles with Context param")
 	uint64 GetInstructionCycles(int32 InIndex) const { return 0LL; }
 
@@ -487,6 +502,18 @@ public:
 		return UINT64_MAX;
 	}
 
+	// Returns accumulated cycles spent in a callable  during the last run
+	// This requires bEnabledProfiling to be turned on in the runtime settings.
+	// If there is no information available this function returns UINT64_MAX.
+	uint64 GetCallableCycles(const FRigVMExtendedExecuteContext& Context, int32 InIndex) const
+	{
+		if (const FRigVMProfilingInfo* RigVMProfilingInfo = Context.GetRigVMProfilingInfo())
+		{
+			return RigVMProfilingInfo->GetCallableCyclesDuringLastRun(InIndex);
+		}
+		return UINT64_MAX;
+	}
+
 	UE_DEPRECATED(5.3, "Please, use GetInstructionMicroSeconds with Context param")
 	double GetInstructionMicroSeconds(int32 InIndex) const { return 0.0; }
 
@@ -503,6 +530,19 @@ public:
 		return double(Cycles) * FPlatformTime::GetSecondsPerCycle() * 1000.0 * 1000.0;
 	}
 
+	// Returns accumulated duration of the callable in microseconds during the last run
+	// Note: this requires bEnabledProfiling to be turned on in the runtime settings.
+	// If there is no information available this function returns -1.0.
+	double GetCallableMicroSeconds(const FRigVMExtendedExecuteContext& Context, int32 InIndex) const
+	{
+		const uint64 Cycles = GetCallableCycles(Context, InIndex);
+		if(Cycles == UINT64_MAX)
+		{
+			return -1.0;
+		}
+		return double(Cycles) * FPlatformTime::GetSecondsPerCycle() * 1000.0 * 1000.0;
+	}
+
 	UE_DEPRECATED(5.3, "Please, use GetInstructionVisitOrder with Context param")
 	const TArray<int32> GetInstructionVisitOrder() const { return TArray<int32>(); }
 
@@ -1140,7 +1180,7 @@ protected:
 	RIGVM_API void ClearMemory_Internal();
 	RIGVM_API void InvalidateCachedMemory_Internal();
 
-	RIGVM_API void SetupInstructionTracking(FRigVMExtendedExecuteContext& Context, int32 InInstructionCount);
+	RIGVM_API void SetupInstructionTracking(FRigVMExtendedExecuteContext& Context, int32 InInstructionCount, int32 InCallableCount);
 	RIGVM_API void StartProfiling(FRigVMExtendedExecuteContext& Context);
 	RIGVM_API void StopProfiling(FRigVMExtendedExecuteContext& Context);
 
diff --git a/Engine/Plugins/Runtime/RigVM/Source/RigVM/Public/RigVMCore/RigVMByteCode.h b/Engine/Plugins/Runtime/RigVM/Source/RigVM/Public/RigVMCore/RigVMByteCode.h
--- a/Engine/Plugins/Runtime/RigVM/Source/RigVM/Public/RigVMCore/RigVMByteCode.h
+++ b/Engine/Plugins/Runtime/RigVM/Source/RigVM/Public/RigVMCore/RigVMByteCode.h
@@ -20,6 +20,7 @@
 #include "UObject/NameTypes.h"
 #include "UObject/ObjectMacros.h"
 #include "UObject/UnrealNames.h"
+#include "RigVMGraphFunctionIdentifier.h"
 
 #include "RigVMByteCode.generated.h"
 
@@ -116,6 +117,115 @@ struct FRigVMPredicateBranch
 	URigVM* VM;
 };
 
+// An argument for a callable in the bytecode.
+USTRUCT()
+struct FRigVMCallableArgument
+{
+	GENERATED_BODY()
+
+	FRigVMCallableArgument()
+		: Name(NAME_None)
+		, Direction(ERigVMPinDirection::Invalid)
+	{}
+
+	UPROPERTY()
+	FName Name;
+
+	UPROPERTY()
+	FString TypeString;
+
+	UPROPERTY()
+	FRigVMOperand InterfaceOperand;
+
+	UPROPERTY()
+	FRigVMOperand ForwardedOperand;
+
+	UPROPERTY()
+	ERigVMPinDirection Direction;
+
+	RIGVM_API void Serialize(FArchive& Ar);
+	friend FArchive& operator<<(FArchive& Ar, FRigVMCallableArgument& P)
+	{
+		P.Serialize(Ar);
+		return Ar;
+	}
+
+	bool IsInput() const
+	{
+		return Direction == ERigVMPinDirection::Input || Direction == ERigVMPinDirection::IO; 
+	}
+
+	bool IsOutput() const
+	{
+		return Direction == ERigVMPinDirection::Output || Direction == ERigVMPinDirection::IO; 
+	}
+
+	friend uint32 GetTypeHash(const FRigVMCallableArgument& InArgument)
+	{
+		return HashCombine(
+			GetTypeHash(InArgument.Name),
+			GetTypeHash(InArgument.TypeString),
+			GetTypeHash(InArgument.InterfaceOperand),
+			GetTypeHash(InArgument.ForwardedOperand),
+			GetTypeHash(InArgument.Direction));
+	}
+
+};
+
+// A callable section of the rigvm bytecode. This is used to represent functions
+// and events with arguments.
+USTRUCT()
+struct FRigVMCallableInfo
+{
+	GENERATED_BODY()
+
+	FRigVMCallableInfo()
+		: Index(INDEX_NONE)
+		, Name(NAME_None)
+		, Hash(0)
+		, FirstInstruction(INDEX_NONE)
+		, LastInstruction(INDEX_NONE)
+	{
+	}
+
+	UPROPERTY()
+	int32 Index;
+	
+	UPROPERTY()
+	FName Name;
+	
+	UPROPERTY()
+	uint32 Hash;
+
+	UPROPERTY()
+	TArray<FRigVMCallableArgument> Arguments;
+
+	UPROPERTY()
+	int32 FirstInstruction;
+
+	UPROPERTY()
+	int32 LastInstruction;
+
+	RIGVM_API int32 GetArgumentIndex(const FName& InName) const;
+	RIGVM_API const FRigVMCallableArgument* FindArgument(const FName& InName) const;
+	int32 GetNumInstructions() const
+	{
+		return FMath::Max(0, LastInstruction - FirstInstruction + 1);
+	}
+
+	RIGVM_API void Serialize(FArchive& Ar);
+	friend FArchive& operator<<(FArchive& Ar, FRigVMCallableInfo& P)
+	{
+		P.Serialize(Ar);
+		return Ar;
+	}
+
+	friend uint32 GetTypeHash(const FRigVMCallableInfo& InCallableInfo)
+	{
+		return InCallableInfo.Hash;
+	}
+};
+
 // The code for a single operation within the RigVM
 UENUM()
 enum class ERigVMOpCode : uint8
@@ -224,6 +334,7 @@ enum class ERigVMOpCode : uint8
 	Execute, // single execute op (formerly Execute_0_Operands to Execute_64_Operands)
 	RunInstructions, // runs a set of instructions lazily
 	SetupTraits, // sets up a list of traits on executecontext
+	InvokeCallable, // invoke a section of the bytecode as a callable
 	Invalid,
 	FirstArrayOpCode = ArrayReset,
 	LastArrayOpCode = ArrayReverse,
@@ -253,42 +364,42 @@ struct FRigVMBaseOp
 	}
 };
 
-
-// execute a function
+// invoke a callable
 USTRUCT()
-struct FRigVMExecuteOp : public FRigVMBaseOp
+struct FRigVMInvokeCallableOp : public FRigVMBaseOp
 {
 	GENERATED_USTRUCT_BODY()
 
-	FRigVMExecuteOp()
+	FRigVMInvokeCallableOp()
 	: FRigVMBaseOp()
-	, FunctionIndex(INDEX_NONE)
+	, CallableIndex(0)
 	, ArgumentCount(0)
-	, FirstPredicateIndex(INDEX_NONE)
-	, PredicateCount(0)
 	{
 	}
 
-	FRigVMExecuteOp(uint16 InFunctionIndex, uint16 InArgumentCount)
-	: FRigVMBaseOp(ERigVMOpCode::Execute)
-	, FunctionIndex(InFunctionIndex)
+	FRigVMInvokeCallableOp(uint16 InCallableIndex, uint16 InArgumentCount)
+	: FRigVMBaseOp(ERigVMOpCode::InvokeCallable)
+	, CallableIndex(InCallableIndex)
 	, ArgumentCount(InArgumentCount)
-	, FirstPredicateIndex(INDEX_NONE)
-	, PredicateCount(0)
 	{
 	}
 
-	uint16 FunctionIndex;
+	FRigVMInvokeCallableOp(ERigVMOpCode InOpCode, uint16 InCallableIndex, uint16 InArgumentCount)
+	: FRigVMBaseOp(InOpCode)
+	, CallableIndex(InCallableIndex)
+	, ArgumentCount(InArgumentCount)
+	{
+	}
+
+	uint16 CallableIndex;
 	uint16 ArgumentCount;
-	uint16 FirstPredicateIndex;
-	uint16 PredicateCount;
 
-	friend uint32 GetTypeHash(const FRigVMExecuteOp& Op)
+	friend uint32 GetTypeHash(const FRigVMInvokeCallableOp& Op)
 	{
 		return HashCombine(
 			GetTypeHash((const FRigVMBaseOp&)Op),
 			HashCombine(
-				GetTypeHash(Op.FunctionIndex),
+				GetTypeHash(Op.CallableIndex),
 				GetTypeHash(Op.ArgumentCount)
 			)
 		);
@@ -297,7 +408,7 @@ struct FRigVMExecuteOp : public FRigVMBaseOp
 	uint16 GetOperandCount() const { return ArgumentCount; }
 
 	UE_API void Serialize(FArchive& Ar);
-	friend FArchive& operator<<(FArchive& Ar, FRigVMExecuteOp& P)
+	friend FArchive& operator<<(FArchive& Ar, FRigVMInvokeCallableOp& P)
 	{
 		P.Serialize(Ar);
 		return Ar;
@@ -305,8 +416,39 @@ struct FRigVMExecuteOp : public FRigVMBaseOp
 	
 	static void ZeroPaddedMemoryIfNeeded(FRigVMBaseOp* InMemory)
 	{
-		FRigVMExecuteOp* This = reinterpret_cast<FRigVMExecuteOp*>(InMemory);
-		RigVM::ZeroPaddedMemory(&This->OpCode, &This->FunctionIndex);
+		FRigVMInvokeCallableOp* This = reinterpret_cast<FRigVMInvokeCallableOp*>(InMemory);
+		RigVM::ZeroPaddedMemory(&This->OpCode, &This->CallableIndex);
+	}
+};
+
+// execute a function
+USTRUCT()
+struct FRigVMExecuteOp : public FRigVMInvokeCallableOp
+{
+	GENERATED_USTRUCT_BODY()
+
+	FRigVMExecuteOp()
+	: FRigVMInvokeCallableOp()
+	, FirstPredicateIndex(INDEX_NONE)
+	, PredicateCount(0)
+	{
+	}
+
+	FRigVMExecuteOp(uint16 InFunctionIndex, uint16 InArgumentCount)
+	: FRigVMInvokeCallableOp(ERigVMOpCode::Execute, InFunctionIndex, InArgumentCount)
+	, FirstPredicateIndex(INDEX_NONE)
+	, PredicateCount(0)
+	{
+	}
+
+	uint16 FirstPredicateIndex;
+	uint16 PredicateCount;
+
+	UE_API void Serialize(FArchive& Ar);
+	friend FArchive& operator<<(FArchive& Ar, FRigVMExecuteOp& P)
+	{
+		P.Serialize(Ar);
+		return Ar;
 	}
 };
 
@@ -835,6 +977,30 @@ struct FRigVMJumpOp : public FRigVMBaseOp
 		FRigVMJumpOp* This = reinterpret_cast<FRigVMJumpOp*>(InMemory);
 		RigVM::ZeroPaddedMemory(&This->OpCode, &This->InstructionIndex);
 	}
+
+	static int32 GetTargetInstruction(ERigVMOpCode InOpCode, const int32 InCurrentInstruction, const int32 InOpJumpIndex)
+	{
+		check(InOpCode == ERigVMOpCode::JumpAbsolute ||
+			InOpCode == ERigVMOpCode::JumpForward ||
+			InOpCode == ERigVMOpCode::JumpBackward ||
+			InOpCode == ERigVMOpCode::JumpAbsoluteIf ||
+			InOpCode == ERigVMOpCode::JumpForwardIf ||
+			InOpCode == ERigVMOpCode::JumpBackwardIf);
+
+		if (InOpCode == ERigVMOpCode::JumpAbsolute ||
+			InOpCode == ERigVMOpCode::JumpAbsoluteIf)
+		{
+			return InOpJumpIndex;
+		}
+
+		if (InOpCode == ERigVMOpCode::JumpForward ||
+			InOpCode == ERigVMOpCode::JumpForwardIf)
+		{
+			return InCurrentInstruction + InOpJumpIndex;
+		}
+
+		return InCurrentInstruction - InOpJumpIndex;
+	}
 };
 
 // jump to a new instruction index based on a condition.
@@ -1153,6 +1319,9 @@ public:
 	// resets the container and removes all memory
 	UE_API void Empty();
 
+	// shrinks the container and reduces the memory
+	UE_API void Shrink();
+
 	// returns a unique hash identifying this bytecode
 	UE_API uint32 GetByteCodeHash() const;
 
@@ -1180,11 +1349,33 @@ public:
 	// returns the index of an entry given a name or INDEX_NONE
 	UE_API int32 FindEntryIndex(const FName& InEntryName) const;
 
+	// returns the number of callables
+	UE_API int32 NumCallables() const;
+
+	// returns true if the provided callable index is valid
+	UE_API bool IsValidCallableIndex(int32 InCallableIndex) const;
+	UE_API bool IsValidCallableIndex(uint16 InCallableIndex) const;
+
+	// returns the callable info with a given index
+	UE_API const FRigVMCallableInfo* GetCallable(int32 InCallableIndex) const;
+	UE_API FRigVMCallableInfo* GetCallable(int32 InCallableIndex);
+	UE_API const FRigVMCallableInfo* GetCallable(uint16 InCallableIndex) const;
+	UE_API FRigVMCallableInfo* GetCallable(uint16 InCallableIndex);
+
+	// returns the callable's index' with a given its hash
+	UE_API int32 GetCallableIndex(const uint32& InHash) const;
+
+	// returns a callable given its hash
+	UE_API const FRigVMCallableInfo* FindCallable(const uint32& InHash) const;
+
+	// returns a callable given its name
+	UE_API const FRigVMCallableInfo* FindCallable(const FName& InName) const;
+
 	// adds an execute operator given its function index operands
 	UE_API int32 AddExecuteOp(int32 InFunctionIndex, const FRigVMOperandArray& InOperands, const int32 InStartPredicateIndex, const int32 InPredicateCount);
 	
 	// adds an execute operator given its function index operands
-	UE_API int32 InlineFunction(const FRigVMByteCode* FunctionByteCode, const FRigVMOperandArray& InOperands);
+	UE_API int32 InlineFunction(const FRigVMByteCode* FunctionByteCode, int32 InFirstIndex, int32 InLastIndex);
 
 	// adds a zero operator to zero the memory of a given argument
 	UE_API int32 AddZeroOp(const FRigVMOperand& InArg);
@@ -1231,6 +1422,10 @@ public:
 	// adds an invoke entry operator
 	UE_API int32 AddInvokeEntryOp(const FName& InEntryName);
 
+	// adds an invoke callable operator
+	UE_API int32 AddInvokeCallableOp(const uint32& InCallableHash, const FRigVMOperandArray& InOperands);
+	UE_API int32 AddInvokeCallableOpByIndex(const uint16& InCallableIndex, const FRigVMOperandArray& InOperands);
+
 	// adds a jump to branch operator
 	UE_API int32 AddJumpToBranchOp(FRigVMOperand InBranchNameArg, int32 InFirstBranchInfoIndex);
 
@@ -1247,12 +1442,42 @@ public:
 	// adds information about a predicate branch for an instruction
 	UE_API int32 AddPredicateBranch(const FRigVMPredicateBranch& InPredicateBranch);
 
+	// returns the number of branches within this bytecode
+	int32 NumBranches() const { return BranchInfos.Num(); }
+
+	// returns the branch given the branch index
+	const FRigVMBranchInfo& GetBranch(int32 InIndex) const
+	{
+		return BranchInfos[InIndex];
+	}
+
+	// returns the number of predicate branches within this bytecode
+	int32 NumPredicateBranches() const { return PredicateBranches.Num(); }
+
+	// returns the predicate branch given the branch index
+	const FRigVMPredicateBranch& GetPredicateBranch(int32 InIndex) const
+	{
+		return PredicateBranches[InIndex];
+	}
+
 	// returns an instruction array for iterating over all operators
 	FRigVMInstructionArray GetInstructions() const
 	{
 		return FRigVMInstructionArray(*this, bByteCodeIsAligned);
 	}
 
+	// returns true if the provided opcode refers to a callable
+	static bool IsCallableOpCode(const ERigVMOpCode InOpCode)
+	{
+		return InOpCode == ERigVMOpCode::InvokeCallable || InOpCode == ERigVMOpCode::Execute;
+	}
+
+	// returns the size of a callable op
+	static int32 GetCallableOpSize(const ERigVMOpCode InOpCode)
+	{
+		return InOpCode == ERigVMOpCode::Execute ? sizeof(FRigVMExecuteOp) : sizeof(FRigVMInvokeCallableOp);
+	}
+
 	// returns the opcode at a given byte index
 	ERigVMOpCode GetOpCodeAt(int32 InByteCodeIndex) const
 	{
@@ -1306,16 +1531,27 @@ public:
 	}
 
 	// returns the operands for a given execute instruction
+	UE_DEPRECATED(5.8, "Please use GetOperandsForCallableOp instead.")
 	FRigVMOperandArray GetOperandsForExecuteOp(const FRigVMInstruction& InInstruction) const
+	{
+		return GetOperandsForCallableOp(InInstruction);
+	}
+
+	FRigVMOperandArray GetOperandsForCallableOp(const FRigVMInstruction& InInstruction) const
 	{
 		int32 ByteCodeIndex = InInstruction.ByteCodeIndex;
-		const FRigVMExecuteOp& ExecuteOp = GetOpAt<FRigVMExecuteOp>(ByteCodeIndex);
+		const FRigVMInvokeCallableOp& CallableOp = GetOpAt<FRigVMInvokeCallableOp>(ByteCodeIndex);
+
 		// if the bytecode is not aligned the OperandAlignment needs to be 0
 		check(bByteCodeIsAligned || InInstruction.OperandAlignment == 0);
-		ByteCodeIndex += sizeof(FRigVMExecuteOp) + InInstruction.OperandAlignment;
-		return GetOperandsAt(ByteCodeIndex, ExecuteOp.GetOperandCount());
+		
+		if (IsCallableOpCode(CallableOp.OpCode))
+		{
+			ByteCodeIndex += GetCallableOpSize(CallableOp.OpCode) + InInstruction.OperandAlignment;
+		}
+		return GetOperandsAt(ByteCodeIndex, CallableOp.GetOperandCount());
 	}
-
+	
 	// returns all of the operands for a given instruction
 	UE_API FRigVMOperandArray GetOperandsForOp(const FRigVMInstruction& InInstruction) const;
 
@@ -1546,6 +1782,10 @@ private:
 	UPROPERTY()
 	TArray<FRigVMPredicateBranch> PredicateBranches;
 
+	// a list of stored callables
+	UPROPERTY()
+	TArray<FRigVMCallableInfo> CallableInfos;
+
 	UPROPERTY()
 	FTopLevelAssetPath PublicContextAssetPath;
 
@@ -1560,6 +1800,7 @@ private:
 	static UE_API TArray<int32> EmptyInstructionIndices;
 
 	friend class URigVMCompiler;
+	friend struct FRigVMCompilerWorkData;
 	friend class FRigVMByteCodeTest;
 	friend class FRigVMInvokeEntryTest;
 	friend class URigVM;
diff --git a/Engine/Plugins/Runtime/RigVM/Source/RigVM/Public/RigVMCore/RigVMExecuteContext.h b/Engine/Plugins/Runtime/RigVM/Source/RigVM/Public/RigVMCore/RigVMExecuteContext.h
--- a/Engine/Plugins/Runtime/RigVM/Source/RigVM/Public/RigVMCore/RigVMExecuteContext.h
+++ b/Engine/Plugins/Runtime/RigVM/Source/RigVM/Public/RigVMCore/RigVMExecuteContext.h
@@ -42,6 +42,7 @@ struct FRigVMDispatchFactory;
 struct FRigVMExecuteContext;
 struct FRigVMExtendedExecuteContext;
 struct FRigVMLogSettings;
+struct FRigVMCallableInfo;
 
 extern UE_API TAutoConsoleVariable<bool> CVarRigVMReportAllMessages;
 
@@ -62,7 +63,7 @@ public:
 	, UpperBound(0)
 	, Index(INDEX_NONE)
 	, Offset(0)
-	, InstructionIndex(INDEX_NONE)
+	, InstructionOrCallableIndex(INDEX_NONE)
 	{
 		Reset();
 	}
@@ -140,14 +141,14 @@ public:
 		Offset = InSliceOffset;
 	}
 
-	int32 GetInstructionIndex() const
+	int32 GetInstructionOrCallableIndex() const
 	{
-		return InstructionIndex;
+		return InstructionOrCallableIndex;
 	}
 
-	void SetInstructionIndex(int32 InInstructionIndex)
+	void SetInstructionOrCallableIndex(int32 InInstructionOrCallableIndex)
 	{
-		InstructionIndex = InInstructionIndex;
+		InstructionOrCallableIndex = InInstructionOrCallableIndex;
 	}
 
 	int32 Num() const
@@ -228,7 +229,7 @@ private:
 	int32 UpperBound;
 	int32 Index;
 	int32 Offset;
-	int32 InstructionIndex;
+	int32 InstructionOrCallableIndex;
 };
 
 USTRUCT()
@@ -613,6 +614,7 @@ private:
 	friend struct FRigVMExtendedExecuteContext;
 	friend class URigVM;
 	friend class URigVMNativized;
+	friend struct FRigVMCallstackGuard;
 };
 
 template<>
@@ -731,62 +733,61 @@ public:
 		return Slices.Last();
 	}
 
+	FRigVMSlice& GetSlice()
+	{
+		return Slices.Last();
+	}
+
 	void BeginLoopSlice(int32 InCount, int32 InRelativeIndex = 0)
 	{
 		ensure(!IsSliceComplete());
 		Slices.Add(FRigVMSlice::MakeForLoop(InCount));
 		Slices.Last().SetRelativeIndex(InRelativeIndex);
-		Slices.Last().SetInstructionIndex(GetPublicData<>().InstructionIndex);
-		Slices.Last().SetOffset(SliceOffsets[GetPublicData<>().InstructionIndex]);
+		Slices.Last().SetInstructionOrCallableIndex(GetPublicData<>().InstructionIndex);
+		Slices.Last().SetOffset(SliceOffsetsPerInstruction[GetPublicData<>().InstructionIndex]);
 	}
 
 	void EndLoopSlice()
 	{
 		ensure(Slices.Num() > 1);
 		ensure(Slices.Last().IsLoop());
-		PopSlice();
+
+		// if this slice has reached its upperbound
+		// we want to make sure to increment the offset
+		// for the instruction the slice originated from,
+		// so that next time around the slice indices are not
+		// reused.
+		const FRigVMSlice PoppedSlice = Slices.Pop();
+		if(PoppedSlice.GetRelativeIndex() == PoppedSlice.Num() - 1)
+		{
+			if(SliceOffsetsPerInstruction.IsValidIndex(PoppedSlice.GetInstructionOrCallableIndex()))
+			{
+				SliceOffsetsPerInstruction[PoppedSlice.GetInstructionOrCallableIndex()] += static_cast<uint16>(PoppedSlice.Num());
+			}
+		}
 	}
 
-	void BeginCallableSlice()
+	void BeginCallableSlice(int32 InCallableIndex)
 	{
 		ensure(!IsSliceComplete());
 		Slices.Add(FRigVMSlice::MakeForCallable());
-		Slices.Last().SetInstructionIndex(GetPublicData<>().InstructionIndex);
-		Slices.Last().SetOffset(SliceOffsets[GetPublicData<>().InstructionIndex]);
+		Slices.Last().SetInstructionOrCallableIndex(InCallableIndex);
+		Slices.Last().SetOffset(SliceOffsetsPerCallable[InCallableIndex]);
 	}
 
 	void EndCallableSlice()
 	{
 		ensure(Slices.Num() > 0);
 		ensure(Slices.Last().IsCallable());
-		PopSlice();
-	}
 
-private:
-	void PopSlice()
-	{
-		// if this slice has reached its upperbound
-    	// we want to make sure to increment the offset
-    	// for the instruction the slice originated from,
-    	// so that next time around the slice indices are not
-    	// reused.
-    	const FRigVMSlice PoppedSlice = Slices.Pop();
-    	if(PoppedSlice.GetRelativeIndex() == PoppedSlice.Num() - 1)
-    	{
-    		if(SliceOffsets.IsValidIndex(PoppedSlice.GetInstructionIndex()))
-    		{
-    			SliceOffsets[PoppedSlice.GetInstructionIndex()] += static_cast<uint16>(PoppedSlice.Num());
-    		}
-    	}
+		const FRigVMSlice PoppedSlice = Slices.Pop();
+		if(SliceOffsetsPerCallable.IsValidIndex(PoppedSlice.GetInstructionOrCallableIndex()))
+		{
+			SliceOffsetsPerCallable[PoppedSlice.GetInstructionOrCallableIndex()]++;
+		}
 	}
 
 public:
-	void IncrementSlice()
-	{
-		FRigVMSlice& ActiveSlice = Slices.Last();
-		ActiveSlice++;
-	}
-
 	bool IsSliceComplete() const
 	{
 		return GetSlice().IsComplete();
@@ -802,6 +803,21 @@ public:
 		return Hash;
 	}
 
+	const TArray<TTuple<int32, const FRigVMCallableInfo*>>& GetCallstack() const
+	{
+		return Callstack;
+	}
+
+	uint32 GetCallstackHash() const
+	{
+		return CallstackHash;
+	}
+
+	uint32 GetHashForLazyBranch() const
+	{
+		return HashCombine(GetTypeHash(GetNumExecutions()), GetCallstackHash());
+	}
+
 	bool IsValidArrayIndex(int32& InOutIndex, const FScriptArrayHelper& InArrayHelper) const
 	{
 		return IsValidArrayIndex(InOutIndex, InArrayHelper.Num());
@@ -915,11 +931,24 @@ private:
 	URigVM* VM = nullptr;
 
 public:
+	// The list of active slices. Slices can introduced by loops or callables
 	TArray<FRigVMSlice> Slices;
-	TArray<uint16> SliceOffsets;
+
+	// The slice offset per instruction. Typically only loop / control flow nodes will
+	// introduce a new loop slice - thus their offset will increase.
+	TArray<uint16> SliceOffsetsPerInstruction;
+
+	// The slice offset per callable. When multiple call sites invoke the same callable,
+	// the offset for that callable index will be increased to make sure each call site
+	// gets its own memory.
+	TArray<uint16> SliceOffsetsPerCallable;
+
+	TArray<TTuple<int32, const FRigVMCallableInfo*>> Callstack;
+	uint32 CallstackHash = 0;
+
 	const FRigVMDispatchFactory* Factory = nullptr;
 	FRigVMNameCache NameCache;
-
+	
 	int32 ExecutingThreadId = INDEX_NONE;
 
 	TArray<int32> EntriesBeingExecuted;
diff --git a/Engine/Plugins/Runtime/RigVM/Source/RigVM/Public/RigVMCore/RigVMGraphFunctionDefinition.h b/Engine/Plugins/Runtime/RigVM/Source/RigVM/Public/RigVMCore/RigVMGraphFunctionDefinition.h
--- a/Engine/Plugins/Runtime/RigVM/Source/RigVM/Public/RigVMCore/RigVMGraphFunctionDefinition.h
+++ b/Engine/Plugins/Runtime/RigVM/Source/RigVM/Public/RigVMCore/RigVMGraphFunctionDefinition.h
@@ -12,6 +12,7 @@
 #include "RigVMVariant.h"
 #include "RigVMNodeLayout.h"
 #include "RigVMStringUtils.h"
+#include "RigVMGraphFunctionIdentifier.h"
 #include "RigVMGraphFunctionDefinition.generated.h"
 
 #define UE_API RIGVM_API
@@ -19,6 +20,8 @@
 class IRigVMGraphFunctionHost;
 struct FRigVMGraphFunctionData;
 
+extern RIGVM_API TAutoConsoleVariable<bool> CVarRigVMCompileFunctionsToCallables;
+
 USTRUCT()
 struct FRigVMFunctionCompilationPropertyDescription
 {
@@ -104,6 +107,11 @@ struct FRigVMFunctionCompilationData
 	, bEncounteredSurpressedErrors(false)
 	{}
 
+	FRigVMFunctionCompilationData(const FRigVMFunctionCompilationData& InOther)
+	{
+		*this = InOther;
+	}
+
 	UPROPERTY()
 	FRigVMByteCode ByteCode;
 
@@ -140,6 +148,9 @@ struct FRigVMFunctionCompilationData
 	UPROPERTY()
 	TMap<FString, FRigVMOperand> Operands;
 
+	UPROPERTY()
+	TMap<FName, FRigVMOperand> InterfaceOperands;
+
 	UPROPERTY()
 	uint32 Hash;
 
@@ -156,6 +167,22 @@ struct FRigVMFunctionCompilationData
 		return bEncounteredSurpressedErrors;
 	}
 
+	RIGVM_API void Reset();
+
+	RIGVM_API FRigVMFunctionCompilationData& operator =(const FRigVMFunctionCompilationData& InOther);
+
+	/**
+	* If this function can be turned into a callable.
+	* Exceptions include:
+	*    Array parameters passed in.
+	*    External variables passed in.
+	*    Any use of traits.
+	*/
+	RIGVM_API bool SupportsCallable() const;
+
+	RIGVM_API const TArray<FRigVMFunctionCompilationPropertyDescription>& GetPropertyDescriptions(ERigVMMemoryType InMemoryType) const;
+	RIGVM_API const TArray<FRigVMFunctionCompilationPropertyPath>& GetPropertyPathDescriptions(ERigVMMemoryType InMemoryType) const;
+	
 	friend uint32 GetTypeHash(const FRigVMFunctionCompilationData& Data) 
 	{
 		uint32 DataHash = Data.ByteCode.GetByteCodeHash();
@@ -212,9 +239,18 @@ struct FRigVMFunctionCompilationData
 			DataHash = HashCombine(DataHash, GetTypeHash(Pair.Value));
 		}
 
+		for (const TPair<FName, FRigVMOperand>& Pair : Data.InterfaceOperands)
+		{
+			DataHash = HashCombine(DataHash, GetTypeHash(Pair.Key));
+			DataHash = HashCombine(DataHash, GetTypeHash(Pair.Value));
+		}
+
 		return DataHash;
 	}
 
+	mutable TOptional<bool> bSupportsCallable;
+	mutable FCriticalSection TaskMutex;
+	
 	RIGVM_API friend FArchive& operator<<(FArchive& Ar, FRigVMFunctionCompilationData& Data);
 };
 
@@ -309,101 +345,6 @@ struct FRigVMGraphFunctionArgument
 	UE_API FRigVMExternalVariable ToExternalVariable() const;
 };
 
-USTRUCT(BlueprintType)
-struct FRigVMGraphFunctionIdentifier
-{
-	GENERATED_BODY()
-
-	UPROPERTY(meta=(DeprecatedProperty))
-	FSoftObjectPath LibraryNode_DEPRECATED;
-
-private:
-	UPROPERTY(VisibleAnywhere, Category=FunctionIdentifier)
-	mutable FString LibraryNodePath;
-
-public:
-
-	// A path to the IRigVMGraphFunctionHost that stores the function information, and compilation data (e.g. RigVMBlueprintGeneratedClass)
-	UPROPERTY(BlueprintReadOnly, VisibleAnywhere, Category=FunctionIdentifier)
-	FSoftObjectPath HostObject;
-
-	FRigVMGraphFunctionIdentifier()
-		: LibraryNodePath(FString()), HostObject(nullptr) {}
-	
-	FRigVMGraphFunctionIdentifier(FSoftObjectPath InHostObject, FString InLibraryNodePath)
-		: LibraryNodePath(InLibraryNodePath), HostObject(InHostObject) {}
-
-	friend uint32 GetTypeHash(const FRigVMGraphFunctionIdentifier& Pointer)
-	{
-		return HashCombine(GetTypeHash(Pointer.GetLibraryNodePath()), GetTypeHash(Pointer.HostObject.ToString()));
-	}
-
-	bool operator==(const FRigVMGraphFunctionIdentifier& Other) const
-	{
-		return HostObject == Other.HostObject && GetNodeSoftPath().GetSubPathString() == Other.GetNodeSoftPath().GetSubPathString();
-	}
-
-	bool IsValid() const
-	{
-		return !HostObject.IsNull() && (!GetLibraryNodePath().IsEmpty());
-	}
-
-	FString GetFunctionName() const
-	{
-		if(IsValid())
-		{
-			const FString Path = GetLibraryNodePath();
-			FString NodeName;
-			if(Path.Split(TEXT("."), nullptr, &NodeName, ESearchCase::CaseSensitive, ESearchDir::FromEnd))
-			{
-				return NodeName;
-			}
-		}
-		return FString();
-	}
-	
-	FName GetFunctionFName() const
-	{
-		if(!IsValid())
-		{
-			return NAME_None;
-		}
-		return *GetFunctionName();
-	}
-
-	FString& GetLibraryNodePath() const
-	{
-		if (LibraryNodePath.IsEmpty() && LibraryNode_DEPRECATED.IsValid())
-		{
-			LibraryNodePath = LibraryNode_DEPRECATED.ToString();
-		}
-		return LibraryNodePath;
-	}
-
-	void SetLibraryNodePath(const FString& InPath)
-	{
-		LibraryNodePath = InPath;
-	}
-
-	FSoftObjectPath GetNodeSoftPath() const
-	{
-		return GetLibraryNodePath();
-	}
-
-	RIGVM_API friend FArchive& operator<<(FArchive& Ar, FRigVMGraphFunctionIdentifier& Data);
-
-	UE_API bool IsVariant() const;
-	UE_API TArray<FRigVMVariantRef> GetVariants(bool bIncludeSelf = false) const;
-	UE_API TArray<FRigVMGraphFunctionIdentifier> GetVariantIdentifiers(bool bIncludeSelf = false) const;
-	UE_API bool IsVariantOf(const FRigVMGraphFunctionIdentifier& InOther) const;
-
-protected:
-	
-	static UE_API TFunction<TArray<FRigVMVariantRef>(const FGuid& InGuid)> GetVariantRefsByGuidFunc;
-
-	friend class URigVMBuildData;
-};
-
 USTRUCT(BlueprintType)
 struct FRigVMGraphFunctionHeader
 {
@@ -495,6 +436,8 @@ struct FRigVMGraphFunctionHeader
 
 	static UE_API FRigVMGraphFunctionHeader FindGraphFunctionHeader(const FRigVMGraphFunctionIdentifier& InIdentifier, bool* bOutIsPublic = nullptr, FString* OutErrorMessage = nullptr);
 
+	static UE_API FRigVMGraphFunctionHeader FindGraphFunctionHeaderFromHash(const FString& InFunctionHash, bool* bOutIsPublic = nullptr, FString* OutErrorMessage = nullptr);
+
 	RIGVM_API TArray<FRigVMExternalVariable> GetExternalVariables() const;
 
 protected:
@@ -547,7 +490,7 @@ struct FRigVMGraphFunctionData
 		return Header == Other.Header;
 	}
 
-	void ClearCompilationData() { CompilationData = FRigVMFunctionCompilationData(); }
+	void ClearCompilationData() { CompilationData.Reset(); }
 
 	RIGVM_API friend FArchive& operator<<(FArchive& Ar, FRigVMGraphFunctionData& Data);
 
diff --git a/Engine/Plugins/Runtime/RigVM/Source/RigVM/Public/RigVMCore/RigVMGraphFunctionIdentifier.h b/Engine/Plugins/Runtime/RigVM/Source/RigVM/Public/RigVMCore/RigVMGraphFunctionIdentifier.h
new file mode 100644
--- /dev/null
+++ b/Engine/Plugins/Runtime/RigVM/Source/RigVM/Public/RigVMCore/RigVMGraphFunctionIdentifier.h
@@ -0,0 +1,107 @@
+// Copyright Epic Games, Inc. All Rights Reserved.
+
+#pragma once
+
+#include "RigVMObjectVersion.h"
+#include "RigVMCore/RigVMObjectArchive.h"
+#include "RigVMVariant.h"
+#include "RigVMGraphFunctionIdentifier.generated.h"
+
+#define UE_API RIGVM_API
+
+USTRUCT(BlueprintType)
+struct FRigVMGraphFunctionIdentifier
+{
+	GENERATED_BODY()
+
+	UPROPERTY(meta=(DeprecatedProperty))
+	FSoftObjectPath LibraryNode_DEPRECATED;
+
+private:
+	UPROPERTY(VisibleAnywhere, Category=FunctionIdentifier)
+	mutable FString LibraryNodePath;
+
+public:
+
+	// A path to the IRigVMGraphFunctionHost that stores the function information, and compilation data (e.g. RigVMBlueprintGeneratedClass)
+	UPROPERTY(BlueprintReadOnly, VisibleAnywhere, Category=FunctionIdentifier)
+	FSoftObjectPath HostObject;
+
+	FRigVMGraphFunctionIdentifier()
+		: LibraryNodePath(FString()), HostObject(nullptr) {}
+	
+	FRigVMGraphFunctionIdentifier(FSoftObjectPath InHostObject, FString InLibraryNodePath)
+		: LibraryNodePath(InLibraryNodePath), HostObject(InHostObject) {}
+
+	friend uint32 GetTypeHash(const FRigVMGraphFunctionIdentifier& Pointer)
+	{
+		return HashCombine(GetTypeHash(Pointer.GetLibraryNodePath()), GetTypeHash(Pointer.HostObject.ToString()));
+	}
+
+	bool operator==(const FRigVMGraphFunctionIdentifier& Other) const
+	{
+		return HostObject == Other.HostObject && GetNodeSoftPath().GetSubPathString() == Other.GetNodeSoftPath().GetSubPathString();
+	}
+
+	bool IsValid() const
+	{
+		return !HostObject.IsNull() && (!GetLibraryNodePath().IsEmpty());
+	}
+
+	FString GetFunctionName() const
+	{
+		if(IsValid())
+		{
+			const FString Path = GetLibraryNodePath();
+			FString NodeName;
+			if(Path.Split(TEXT("."), nullptr, &NodeName, ESearchCase::CaseSensitive, ESearchDir::FromEnd))
+			{
+				return NodeName;
+			}
+		}
+		return FString();
+	}
+	
+	FName GetFunctionFName() const
+	{
+		if(!IsValid())
+		{
+			return NAME_None;
+		}
+		return *GetFunctionName();
+	}
+
+	FString& GetLibraryNodePath() const
+	{
+		if (LibraryNodePath.IsEmpty() && LibraryNode_DEPRECATED.IsValid())
+		{
+			LibraryNodePath = LibraryNode_DEPRECATED.ToString();
+		}
+		return LibraryNodePath;
+	}
+
+	void SetLibraryNodePath(const FString& InPath)
+	{
+		LibraryNodePath = InPath;
+	}
+
+	FSoftObjectPath GetNodeSoftPath() const
+	{
+		return GetLibraryNodePath();
+	}
+
+	RIGVM_API friend FArchive& operator<<(FArchive& Ar, FRigVMGraphFunctionIdentifier& Data);
+
+	UE_API bool IsVariant() const;
+	UE_API TArray<FRigVMVariantRef> GetVariants(bool bIncludeSelf = false) const;
+	UE_API TArray<FRigVMGraphFunctionIdentifier> GetVariantIdentifiers(bool bIncludeSelf = false) const;
+	UE_API bool IsVariantOf(const FRigVMGraphFunctionIdentifier& InOther) const;
+
+protected:
+	
+	static UE_API TFunction<TArray<FRigVMVariantRef>(const FGuid& InGuid)> GetVariantRefsByGuidFunc;
+
+	friend class URigVMBuildData;
+};
+
+#undef UE_API
diff --git a/Engine/Plugins/Runtime/RigVM/Source/RigVM/Public/RigVMCore/RigVMMemoryCommon.h b/Engine/Plugins/Runtime/RigVM/Source/RigVM/Public/RigVMCore/RigVMMemoryCommon.h
--- a/Engine/Plugins/Runtime/RigVM/Source/RigVM/Public/RigVMCore/RigVMMemoryCommon.h
+++ b/Engine/Plugins/Runtime/RigVM/Source/RigVM/Public/RigVMCore/RigVMMemoryCommon.h
@@ -73,8 +73,10 @@ struct FRigVMInstructionSetExecuteState
 	{
 	}
 
+	RIGVM_API bool RequiresExecute(const uint32& InSliceHash, const uint32& InHashForLazyBranch);
+
 	UPROPERTY(transient)
-	TMap<uint32, uint32> SliceHashToNumInstruction;
+	TMap<uint32, uint32> SliceHashToCallStackHash;
 };
 
 /**
@@ -165,6 +167,8 @@ public:
 		RigVM::ZeroPaddedMemory(&InOperand->MemoryType, &InOperand->RegisterIndex);
 	}
 
+	RIGVM_API FString ToString() const;
+
 private:
 
 	UPROPERTY()
@@ -181,6 +185,7 @@ private:
 	uint16 RegisterOffset;
 
 	friend class URigVMCompiler;
+	friend struct FRigVMCompilerWorkData;
 };
 
 typedef TArrayView<const FRigVMOperand> FRigVMOperandArray;
diff --git a/Engine/Plugins/Runtime/RigVM/Source/RigVM/Public/RigVMCore/RigVMMemoryStorage.h b/Engine/Plugins/Runtime/RigVM/Source/RigVM/Public/RigVMCore/RigVMMemoryStorage.h
--- a/Engine/Plugins/Runtime/RigVM/Source/RigVM/Public/RigVMCore/RigVMMemoryStorage.h
+++ b/Engine/Plugins/Runtime/RigVM/Source/RigVM/Public/RigVMCore/RigVMMemoryStorage.h
@@ -94,8 +94,7 @@ struct FRigVMBranchInfo
 
 	friend uint32 GetTypeHash(const FRigVMBranchInfo& InBranchInfo)
 	{
-		uint32 Hash = GetTypeHash(GetTypeHash(InBranchInfo.Index));
-		Hash = HashCombine(Hash, GetTypeHash(InBranchInfo.Label.ToString()));
+		uint32 Hash = GetTypeHash(InBranchInfo.Label.ToString());
 		Hash = HashCombine(Hash, GetTypeHash(InBranchInfo.InstructionIndex));
 		Hash = HashCombine(Hash, GetTypeHash(InBranchInfo.ArgumentIndex));
 		Hash = HashCombine(Hash, GetTypeHash(InBranchInfo.FirstInstruction));
@@ -225,6 +224,31 @@ protected:
 /// Memory handle
 //////////////////////////////////////////////////////////////////////////////
 
+/**
+ * The FRigVMForwardedMemoryHandle is used to forward memory from a call site
+ * to a callable. The compiler compiled callables arguments to this structure.
+ * FRigVMMemoryHandles pointing to this structure will use the pointer to
+ * perform indirection on the memory.
+ */
+USTRUCT()
+struct FRigVMForwardedMemoryHandle
+{
+	GENERATED_USTRUCT_BODY()
+
+	const FProperty* Property = nullptr;
+	uint8* Memory = nullptr;
+
+	void Reset()
+	{
+		Property = nullptr;
+		Memory = nullptr;
+	}
+};
+
+//////////////////////////////////////////////////////////////////////////////
+/// Memory handle
+//////////////////////////////////////////////////////////////////////////////
+
 /**
  * The FRigVMMemoryHandle is used to access the memory used within a URigMemoryStorage.
  * The Memory Handle caches the pointer of the head property, and can rely on a
@@ -247,6 +271,7 @@ public:
 		, Property(nullptr) 
 		, PropertyPath(nullptr)
 		, LazyBranch(nullptr)
+		, bIsForwarded(false)
 	{}
 
 	// Constructor from complete data
@@ -255,7 +280,12 @@ public:
 		, Property(InProperty)
 		, PropertyPath(InPropertyPath)
 		, LazyBranch(InLazyBranch)
+		, bIsForwarded(false)
 	{
+		if (const FStructProperty* StructProperty = CastField<FStructProperty>(Property))
+		{
+			bIsForwarded = StructProperty->Struct == FRigVMForwardedMemoryHandle::StaticStruct();
+		}
 	}
 
 	/**
@@ -341,6 +371,12 @@ public:
 	// Returns the head property of this handle
 	const FProperty* GetProperty() const
 	{
+		if (bIsForwarded)
+		{
+			const FRigVMForwardedMemoryHandle* ForwardedHandle = (const FRigVMForwardedMemoryHandle*)Ptr;
+			checkf(ForwardedHandle->Property, TEXT("If this is null the handle is being interacted with outside of a Invoke Callable instruction."));
+			return ForwardedHandle->Property;
+		}
 		return Property;
 	}
 
@@ -366,13 +402,13 @@ public:
 		if(bIsHiddenArgument)
 		{
 			check(PropertyPath == nullptr);
-			return GetArrayElementPropertyChecked(Property);
+			return GetArrayElementPropertyChecked(GetProperty());
 		}
 		else if(PropertyPath)
 		{
 			return PropertyPath->GetTailProperty();
 		}
-		return Property;
+		return GetProperty();
 	}
 
 	// Returns true if this memory handle maps to a given type of property
@@ -821,6 +857,14 @@ private:
 
 	uint8* GetData_Internal(bool bFollowPropertyPath, int32 InSliceIndex) const
 	{
+		uint8* Memory = Ptr;
+		if (bIsForwarded)
+		{
+			FRigVMForwardedMemoryHandle* ForwardedHandle = (FRigVMForwardedMemoryHandle*)Memory;
+			checkf(ForwardedHandle->Memory, TEXT("If this is null the handle is being interacted with outside of a Invoke Callable instruction."));
+			Memory = ForwardedHandle->Memory;
+		}
+		
 		if(InSliceIndex != INDEX_NONE)
 		{
 			// Sliced memory cannot be accessed using a property path.
@@ -828,8 +872,8 @@ private:
 			check(PropertyPath == nullptr);
 			check(!bFollowPropertyPath);
 
-			const FArrayProperty* ArrayProperty = CastFieldChecked<FArrayProperty>(Property);
-			FScriptArrayHelper ArrayHelper(ArrayProperty, Ptr);
+			const FArrayProperty* ArrayProperty = CastFieldChecked<FArrayProperty>(GetProperty());
+			FScriptArrayHelper ArrayHelper(ArrayProperty, Memory);
 
 			if (InSliceIndex >= ArrayHelper.Num() - 1)
 			{
@@ -896,9 +940,9 @@ private:
 		// and return its memory instead.
 		if(bFollowPropertyPath && PropertyPath != nullptr)
 		{
-			return PropertyPath->GetData<uint8>(Ptr, Property);
+			return PropertyPath->GetData<uint8>(Memory, GetProperty());
 		}
-		return Ptr;
+		return Memory;
 	}
 
 	TRigVMLazyValueBase GetDataLazily_Internal(bool bFollowPropertyPath, uint32 InSliceHash) const
@@ -996,6 +1040,9 @@ private:
 	// The [optional] lazy branch used by this handle
 	FRigVMLazyBranch* LazyBranch;
 
+	// If this is true the handle points to a FRigVMForwardedMemoryHandle
+	bool bIsForwarded;
+
 	friend class URigVM;
 	friend struct TRigVMLazyValueBase;
 	friend class URigVMMemoryStorage;
diff --git a/Engine/Plugins/Runtime/RigVM/Source/RigVM/Public/RigVMCore/RigVMNativized.h b/Engine/Plugins/Runtime/RigVM/Source/RigVM/Public/RigVMCore/RigVMNativized.h
--- a/Engine/Plugins/Runtime/RigVM/Source/RigVM/Public/RigVMCore/RigVMNativized.h
+++ b/Engine/Plugins/Runtime/RigVM/Source/RigVM/Public/RigVMCore/RigVMNativized.h
@@ -96,13 +96,14 @@ public:
 protected:
 
 	template<typename ExecuteContextType = FRigVMExecuteContext>
-	ExecuteContextType& UpdateContext(FRigVMExtendedExecuteContext& Context, int32 InNumberInstructions, const FName& InEntryName)
+	ExecuteContextType& UpdateContext(FRigVMExtendedExecuteContext& Context, int32 InNumberInstructions, int32 InNumberCallables, const FName& InEntryName)
 	{
 		UpdateExternalVariables(Context);
 	
 		Context.ResetExecutionState();
 		Context.SetVM(this);
-		Context.SliceOffsets.AddZeroed(InNumberInstructions);
+		Context.SliceOffsetsPerInstruction.SetNumZeroed(InNumberInstructions);
+		Context.SliceOffsetsPerCallable.SetNumZeroed(InNumberCallables);
 		Context.GetPublicData<ExecuteContextType>().EventName = InEntryName;
 		return Context.GetPublicData<ExecuteContextType>();
 	}
@@ -137,9 +138,9 @@ protected:
 		Context.EndLoopSlice();
 	}
 
-	void BeginCallableSlice(FRigVMExtendedExecuteContext& Context)
+	void BeginCallableSlice(FRigVMExtendedExecuteContext& Context, int32 InCallableIndex)
 	{
-		Context.BeginCallableSlice();
+		Context.BeginCallableSlice(InCallableIndex);
 	}
 
 	void EndCallableSlice(FRigVMExtendedExecuteContext& Context)
diff --git a/Engine/Plugins/Runtime/RigVM/Source/RigVM/Public/RigVMCore/RigVMProfilingInfo.h b/Engine/Plugins/Runtime/RigVM/Source/RigVM/Public/RigVMCore/RigVMProfilingInfo.h
--- a/Engine/Plugins/Runtime/RigVM/Source/RigVM/Public/RigVMCore/RigVMProfilingInfo.h
+++ b/Engine/Plugins/Runtime/RigVM/Source/RigVM/Public/RigVMCore/RigVMProfilingInfo.h
@@ -24,6 +24,8 @@ struct FRigVMInstructionVisitInfo
 	{
 		InstructionVisitedDuringLastRun.Reset();
 		InstructionVisitOrder.Reset();
+		CallableVisitedDuringLastRun.Reset();
+		CallableVisitOrder.Reset();
 		FirstEntryEventInQueue = NAME_None;
 	}
 
@@ -43,10 +45,26 @@ struct FRigVMInstructionVisitInfo
 	inline void AddInstructionIndexToVisitOrder(int32 InstructionIndex) { InstructionVisitOrder.Add(InstructionIndex); }
 	inline const TArray<int32>& GetInstructionVisitOrder() const { return InstructionVisitOrder; }
 
+	inline void ResetCallableVisitedDuringLastRun(int32 NewSize = 0) { CallableVisitedDuringLastRun.Reset(NewSize); }
+	inline void SetNumCallableVisitedDuringLastRunZeroed(int32 Num) { CallableVisitedDuringLastRun.SetNumZeroed(Num); }
+	inline void SetCallableVisitedDuringLastRun(int32 CallableIndex)
+	{
+		if (CallableVisitedDuringLastRun.IsValidIndex(CallableIndex))
+		{
+			CallableVisitedDuringLastRun[CallableIndex]++;
+		}
+	}
+	inline int32 GetCallableVisitedCountDuringLastRun(int32 CallableIndex) const { return CallableVisitedDuringLastRun.IsValidIndex(CallableIndex) ? CallableVisitedDuringLastRun[CallableIndex] : 0; }
+	inline const TArray<int32>& GetCallableVisitedCountDuringLastRun() const { return CallableVisitedDuringLastRun; }
+
+	inline void ResetCallableVisitOrder(int32 NewSize = 0) { CallableVisitOrder.Reset(NewSize); }
+	inline void AddCallableIndexToVisitOrder(int32 CallableIndex) { CallableVisitOrder.Add(CallableIndex); }
+	inline const TArray<int32>& GetCallableVisitOrder() const { return CallableVisitOrder; }
+	
 	inline const void SetFirstEntryEventInEventQueue(const FName& InFirstEventName) { FirstEntryEventInQueue = InFirstEventName; }
 	inline const FName& GetFirstEntryEventInEventQueue() const { return FirstEntryEventInQueue; }
 
-	UE_API void SetupInstructionTracking(int32 InInstructionCount);
+	UE_API void SetupInstructionTracking(int32 InInstructionCount, int32 InCallableCount);
 
 private:
 
@@ -54,6 +72,10 @@ private:
 	TArray<int32> InstructionVisitedDuringLastRun;
 	TArray<int32> InstructionVisitOrder;
 
+	// stores the number of times each callable was visited
+	TArray<int32> CallableVisitedDuringLastRun;
+	TArray<int32> CallableVisitOrder;
+
 	// A RigVMHost can run multiple events per evaluation, such as the Backward&Forward Solve Mode,
 	// store the first event such that we know when to reset data for a new round of rig evaluation
 	FName FirstEntryEventInQueue = NAME_None;
@@ -133,7 +155,35 @@ struct FRigVMProfilingInfo
 		}
 	}
 
-	UE_API void SetupInstructionTracking(int32 InInstructionCount, bool bEnableProfiling);
+	inline void ResetCallableCyclesDuringLastRun(int32 NewSize = 0) { CallableCyclesDuringLastRun.Reset(NewSize); }
+	inline uint64 GetCallableCyclesDuringLastRun(int32 CallableIndex) const
+	{
+		return CallableCyclesDuringLastRun.IsValidIndex(CallableIndex) ? CallableCyclesDuringLastRun[CallableIndex] : UINT64_MAX;
+	}
+	inline void SetCallableCyclesDuringLastRun(int32 CallableIndex, uint64 CyclesDuringLastRun)
+	{
+		if (CallableCyclesDuringLastRun.IsValidIndex(CallableIndex))
+		{
+			CallableCyclesDuringLastRun[CallableIndex] = CyclesDuringLastRun;
+		}
+	}
+	inline void AddCallableCyclesDuringLastRun(int32 CallableIndex, uint64 CyclesDuringLastRun)
+	{
+		if (CallableCyclesDuringLastRun.IsValidIndex(CallableIndex))
+		{
+			CallableCyclesDuringLastRun[CallableIndex] += CyclesDuringLastRun;
+		}
+	}
+	inline void InitCallableCyclesDuringLastRunValues(int32 NewSize, uint64 DefaultValue)
+	{
+		CallableCyclesDuringLastRun.SetNumUninitialized(NewSize);
+		for (uint64& Value : CallableCyclesDuringLastRun)
+		{
+			Value = DefaultValue;
+		}
+	}
+	
+	UE_API void SetupInstructionTracking(int32 InInstructionCount, int32 InCallableCount, bool bEnableProfiling);
 
 	double GetLastExecutionMicroSeconds() const { return LastExecutionMicroSeconds; }
 	void SetLastExecutionMicroSeconds(double InLastExecutionMicroSeconds) { LastExecutionMicroSeconds = InLastExecutionMicroSeconds; }
@@ -143,8 +193,9 @@ struct FRigVMProfilingInfo
 
 private:
 
-	// stores the number of times each instruction was visited
+	// stores the cycles for each instruction
 	TArray<uint64> InstructionCyclesDuringLastRun;
+	TArray<uint64> CallableCyclesDuringLastRun;
 
 	uint64 StartCycles = 0;
 	uint64 OverallCycles = 0;
diff --git a/Engine/Plugins/Runtime/RigVM/Source/RigVM/Public/RigVMObjectVersion.h b/Engine/Plugins/Runtime/RigVM/Source/RigVM/Public/RigVMObjectVersion.h
--- a/Engine/Plugins/Runtime/RigVM/Source/RigVM/Public/RigVMObjectVersion.h
+++ b/Engine/Plugins/Runtime/RigVM/Source/RigVM/Public/RigVMObjectVersion.h
@@ -84,6 +84,9 @@ struct FRigVMObjectVersion
 		// Debug operand mapping simplified and moved to context only
 		DebugOperandMappingSimplified,
 
+		// Introduction of callables to the rigvm bytecode
+		RigVMCallables,
+
 		// -----<new versions can be added above this line>-------------------------------------------------
 		VersionPlusOne,
 		LatestVersion = VersionPlusOne - 1
diff --git a/Engine/Plugins/Runtime/RigVM/Source/RigVMDeveloper/Private/RigVMCompiler/RigVMCodeGenerator.cpp b/Engine/Plugins/Runtime/RigVM/Source/RigVMDeveloper/Private/RigVMCompiler/RigVMCodeGenerator.cpp
--- a/Engine/Plugins/Runtime/RigVM/Source/RigVMDeveloper/Private/RigVMCompiler/RigVMCodeGenerator.cpp
+++ b/Engine/Plugins/Runtime/RigVM/Source/RigVMDeveloper/Private/RigVMCompiler/RigVMCodeGenerator.cpp
@@ -645,9 +645,9 @@ FString FRigVMCodeGenerator::DumpInstructions(const FRigVMExtendedExecuteContext
 			case ERigVMOpCode::Execute:
 			{
 				const FRigVMExecuteOp& Op = ByteCode.GetOpAt<FRigVMExecuteOp>(Instruction);
-				FRigVMOperandArray Operands = ByteCode.GetOperandsForExecuteOp(Instruction);
+				FRigVMOperandArray Operands = ByteCode.GetOperandsForCallableOp(Instruction);
 
-				const FRigVMFunction* Function = FRigVMRegistry::Get().FindFunction(*Functions[Op.FunctionIndex].ToString());
+				const FRigVMFunction* Function = FRigVMRegistry::Get().FindFunction(*Functions[Op.CallableIndex].ToString());
 				check(Function);
 					
 				FStringArray Arguments;
@@ -1218,7 +1218,7 @@ void FRigVMCodeGenerator::ParseRequiredUProperties(const FRigVMExtendedExecuteCo
 		if(Instruction.OpCode == ERigVMOpCode::Execute)
 		{
 			const FRigVMExecuteOp& Op = ByteCode.GetOpAt<FRigVMExecuteOp>(Instruction);
-			const FRigVMFunction* Function = FRigVMRegistry::Get().FindFunction(*Functions[Op.FunctionIndex].ToString());
+			const FRigVMFunction* Function = FRigVMRegistry::Get().FindFunction(*Functions[Op.CallableIndex].ToString());
 			if(const FRigVMDispatchFactory* Factory = Function->Factory)
 			{
 				const FString FunctionName = Function->GetName();
@@ -1502,7 +1502,7 @@ void FRigVMCodeGenerator::ParseInstructionGroups(const FRigVMExtendedExecuteCont
 				case ERigVMOpCode::Execute:
 				{
 					const FRigVMExecuteOp& Op = ByteCode.GetOpAt<FRigVMExecuteOp>(Instruction);
-					const FRigVMFunction* Function = FRigVMRegistry::Get().FindFunction(*Functions[Op.FunctionIndex].ToString());
+					const FRigVMFunction* Function = FRigVMRegistry::Get().FindFunction(*Functions[Op.CallableIndex].ToString());
 					check(Function);
 
 					// make sure to include the required header
diff --git a/Engine/Plugins/Runtime/RigVM/Source/RigVMDeveloper/Private/RigVMCompiler/RigVMCompiler.cpp b/Engine/Plugins/Runtime/RigVM/Source/RigVMDeveloper/Private/RigVMCompiler/RigVMCompiler.cpp
--- a/Engine/Plugins/Runtime/RigVM/Source/RigVMDeveloper/Private/RigVMCompiler/RigVMCompiler.cpp
+++ b/Engine/Plugins/Runtime/RigVM/Source/RigVMDeveloper/Private/RigVMCompiler/RigVMCompiler.cpp
@@ -16,7 +16,6 @@
 #include "ProfilingDebugging/ScopedTimers.h"
 #include "RigVMModel/RigVMClient.h"
 #include "RigVMFunctions/RigVMDispatch_Array.h"
-#include "Algo/Count.h"
 #include "Algo/Sort.h"
 #include "String/Join.h"
 #include "RigVMDeveloperTypeUtils.h"
@@ -130,7 +129,7 @@ FRigVMOperand FRigVMCompilerWorkData::AddProperty(
 	UObject* InCPPTypeObject,
 	const FString& InDefaultValue)
 {
-	check(bSetupMemory);
+	check(CompilerPhase == ERigVMCompilerPhase_SetupMemory);
 
 	const FRigVMRegistry& Registry = FRigVMRegistry::Get();
 	FRigVMTemplateArgumentType ArgumentType(*InCPPType, InCPPTypeObject);
@@ -144,6 +143,7 @@ FRigVMOperand FRigVMCompilerWorkData::AddProperty(
 			ArgumentType = FRigVMTemplateArgumentType(ExecuteContextStruct);
 		}
 	}
+
 	
 	FRigVMPropertyDescription Description(InName, ArgumentType.CPPType.ToString(), ArgumentType.CPPTypeObject, InDefaultValue);
 	ensure(FindProperty(InMemoryType, Description.Name).IsValid() == false); // Warning : this check can not be done before the SanitizeName inside FRigVMPropertyDescription constructor
@@ -170,6 +170,91 @@ FRigVMOperand FRigVMCompilerWorkData::FindProperty(ERigVMMemoryType InMemoryType
 	return FRigVMOperand();
 }
 
+FRigVMOperand FRigVMCompilerWorkData::FindOrAddProperty(const ERigVMMemoryType InMemoryType, 
+	const FRigVMFunctionCompilationPropertyDescription& InProperty, const FString& InNewName,
+	const bool bMakeUniqueProperty)
+{
+	// Sharing / reusing memory / operands happens as per following contract:
+	// 1. properties are only shared if their CPP type matches
+	// 2. Literal / constant memory is only shared if the constant values match (and it is not a trait list)
+	// 3. Work state is only shared if it is not internal work state private to the instruction referring to it
+	// 4. Work state of type FRigVMInstructionSetExecuteState (bMakeUniqueProperty) is never shared either since it is work state private to a lazy branch.
+
+	FString NewName = InNewName;
+	int32 Suffix = 0;
+
+	// always share literals
+	if(InMemoryType == ERigVMMemoryType::Literal && !bMakeUniqueProperty && PropertyDescriptions.Contains(ERigVMMemoryType::Literal))
+	{
+		const TArray<FRigVMPropertyDescription>& LiteralDescriptions = PropertyDescriptions[ERigVMMemoryType::Literal]; 
+		for(int32 Index = 0; Index < LiteralDescriptions.Num(); Index++)
+		{
+			const FRigVMPropertyDescription& ExistingProperty = LiteralDescriptions[Index];
+			if(ExistingProperty.CPPType.Equals(InProperty.CPPType, ESearchCase::IgnoreCase))
+			{
+				// if the value is the same and this is not a trait setup list
+				if(ExistingProperty.DefaultValue.Equals(InProperty.DefaultValue, ESearchCase::CaseSensitive))
+				{
+					return FRigVMOperand(ERigVMMemoryType::Literal, Index);
+				}
+			}
+		}
+	}
+
+	FRigVMPropertyDescription::SanitizeName(NewName);
+	FRigVMOperand Operand = FindProperty(InMemoryType, *NewName);
+
+	// look for a free / unused operand
+	while(Operand.IsValid())
+	{
+		const FRigVMPropertyDescription& ExistingProperty = PropertyDescriptions[InMemoryType][Operand.GetRegisterIndex()];
+		if(!bMakeUniqueProperty &&
+			ExistingProperty.CPPType.Equals(InProperty.CPPType, ESearchCase::CaseSensitive))
+		{
+			if(InMemoryType == ERigVMMemoryType::Literal)
+			{
+				// if the value is the same and this is not a trait setup list
+				if(ExistingProperty.DefaultValue.Equals(InProperty.DefaultValue))
+				{
+					return Operand;
+				}
+			}
+			else
+			{
+				if(InMemoryType == ERigVMMemoryType::Work &&
+					RigVMTypeUtils::IsArrayType(InProperty.CPPType))
+				{
+					// for now we don't share operands which have an array type
+					// since they may represent internal state. multiple occurrences of
+					// the same function cannot share internal work state.
+				}
+				else
+				{
+					return Operand;
+				}
+			}
+		}
+
+		NewName = FString::Printf(TEXT("%s_%02d"), *InNewName, ++Suffix);
+		FRigVMPropertyDescription::SanitizeName(NewName);
+		Operand = FindProperty(InMemoryType, *NewName);
+	}
+			
+	// create a new operand as needed
+	return AddProperty(InMemoryType, *NewName, InProperty.CPPType, InProperty.CPPTypeObject.Get(), InProperty.DefaultValue);
+}
+
+FRigVMOperand FRigVMCompilerWorkData::FindOrAddProperty(ERigVMMemoryType InMemoryType, const FName& InName, const FString& InCPPType,
+	UObject* InCPPTypeObject, const FString& InDefaultValue, bool bMakeUniqueProperty)
+{
+	FRigVMFunctionCompilationPropertyDescription Description;
+	Description.Name = InName;
+	Description.CPPType = InCPPType;
+	Description.CPPTypeObject = InCPPTypeObject;
+	Description.DefaultValue = InDefaultValue;
+	return FindOrAddProperty(InMemoryType, Description, InName.ToString(), bMakeUniqueProperty);
+}
+
 FRigVMPropertyDescription FRigVMCompilerWorkData::GetProperty(const FRigVMOperand& InOperand)
 {
 	TArray<FRigVMPropertyDescription>* PropertyArray = PropertyDescriptions.Find(InOperand.GetMemoryType());
@@ -212,7 +297,7 @@ const FProperty* FRigVMCompilerWorkData::GetPropertyForOperand(const FRigVMOpera
 		return nullptr;
 	}
 	
-	check(!bSetupMemory);
+	check(CompilerPhase == ERigVMCompilerPhase_BuildInstructions);
 
 	auto GetPropertyFromMemory = [this](FRigVMMemoryStorageStruct& InMemory, const FRigVMOperand& InOperand)
 	{
@@ -309,6 +394,380 @@ FName FRigVMCompilerWorkData::GetUniquePropertyName(ERigVMMemoryType InMemoryTyp
 	return *Name;
 }
 
+bool FRigVMCompilerWorkData::AddPropertiesForFunction(
+	bool InIsCallable,
+	const uint32& InHash,
+	const FRigVMGraphFunctionHeader& InHeader,
+	const FRigVMFunctionCompilationData* InCompilationData,
+	int32 InFirstInstruction, int32 InLastInstruction,
+	const TArray<FRigVMOperand>& InInterfaceOperands,
+	const TMap<FRigVMOperand, FRigVMOperand>& InForwardedHandleToInterfaceOperand)
+{
+	TMap<FRigVMOperand, FRigVMOperand>& OperandMap = OperandMapPerFunction.FindOrAdd(InHash, {});
+	OperandMap.Reset();
+
+	auto ProcessOperand = [this, InHash, InIsCallable, &InHeader, InCompilationData, &OperandMap, &InInterfaceOperands, &InForwardedHandleToInterfaceOperand]
+	(const ERigVMOpCode& InOpCode, const FRigVMOperand& InOldOperand)
+	{
+		FRigVMOperand OldOperandNoOffset(InOldOperand.GetMemoryType(), InOldOperand.GetRegisterIndex());
+		if (const FRigVMOperand* InterfaceOperandFromForwardedHandle = InForwardedHandleToInterfaceOperand.Find(OldOperandNoOffset))
+		{
+			OldOperandNoOffset = *InterfaceOperandFromForwardedHandle;
+		}
+		
+		ERigVMMemoryType MemoryType = OldOperandNoOffset.GetMemoryType();
+
+		const TArray<FRigVMFunctionCompilationPropertyDescription>* Properties = nullptr;
+		switch (MemoryType)
+		{
+			case ERigVMMemoryType::Work:
+			{
+				Properties = &InCompilationData->WorkPropertyDescriptions;
+				break;
+			}
+			case ERigVMMemoryType::Literal:
+			{
+				Properties = &InCompilationData->LiteralPropertyDescriptions;
+				break;
+			}
+			case ERigVMMemoryType::External:
+			{
+				Properties = &InCompilationData->ExternalPropertyDescriptions;
+				break;
+			}
+			case ERigVMMemoryType::Debug:
+			{
+				Properties = &InCompilationData->DebugPropertyDescriptions;
+				break;
+			}
+			default:
+			{
+				checkNoEntry();
+			}
+		}
+
+		int32 FirstValidIndex = 0;
+
+		// if we don't contain any callables - we should skip over the first couple of operands
+		if (!InIsCallable)
+		{
+			if (MemoryType == ERigVMMemoryType::Work)
+			{
+				for (const FRigVMGraphFunctionArgument& Argument : InHeader.Arguments)
+				{
+					if (Argument.IsCPPTypeObjectValid())
+					{
+						if (const UScriptStruct* ScriptStruct = Cast<UScriptStruct>(Argument.CPPTypeObject.Get()))
+						{
+							if (ScriptStruct->IsChildOf(FRigVMExecutePin::StaticStruct()))
+							{
+								continue;
+							}
+						}
+					}
+					if (!Properties->IsValidIndex(FirstValidIndex))
+					{
+						break;
+					}
+					const FString PropertyName = FRigVMPropertyDescription::SanitizeName(Argument.Name).ToString();
+					if (!Properties->operator[](FirstValidIndex).Name.ToString().EndsWith(PropertyName))
+					{
+						continue;
+					}
+					FirstValidIndex++;
+				}
+			}
+		}
+
+		int32 PropertyIndex = OldOperandNoOffset.GetRegisterIndex();
+		if ((Properties->IsValidIndex(PropertyIndex) || MemoryType == ERigVMMemoryType::External) && PropertyIndex >= FirstValidIndex)
+		{
+			static const FString FunctionLibraryPrefix = TEXT("FunctionLibrary");
+			static const FString RigVMInstructionSetExecuteStateName = FRigVMInstructionSetExecuteState::StaticStruct()->GetStructCPPName();
+
+			FRigVMOperand NewOperandNoOffset;
+			if (const FRigVMOperand* ExistingNewOperandNoOffset = OperandMap.Find(OldOperandNoOffset))
+			{
+				NewOperandNoOffset = *ExistingNewOperandNoOffset;
+			}
+			else
+			{
+				FRigVMFunctionCompilationPropertyDescription CustomDescription;
+				const FRigVMFunctionCompilationPropertyDescription* Description = nullptr;
+				
+				if (Properties && Properties->IsValidIndex(PropertyIndex))
+				{
+					Description = &Properties->operator[](PropertyIndex);
+				}
+				else
+				{
+					check(MemoryType == ERigVMMemoryType::External);
+					check(VM->GetExternalVariableDefs().IsValidIndex(PropertyIndex));
+					
+					const FRigVMExternalVariableDef& ExternalVariableDef = VM->GetExternalVariableDefs()[PropertyIndex];
+					CustomDescription.Name = ExternalVariableDef.GetName();
+					CustomDescription.CPPType = ExternalVariableDef.GetExtendedCPPType().ToString();
+					CustomDescription.CPPTypeObject = const_cast<UObject*>(ExternalVariableDef.GetCPPTypeObject());
+					CustomDescription.DefaultValue = FString();
+					Description = &CustomDescription;
+
+					// if this is an input argument - make sure to turn it into work state.
+					if (InInterfaceOperands.Contains(OldOperandNoOffset))
+					{
+						MemoryType = ERigVMMemoryType::Work;
+					}
+				}
+
+				// if this is a trait setup list - let's copy all traits over
+				if(InOpCode == ERigVMOpCode::SetupTraits && MemoryType == ERigVMMemoryType::Literal)
+				{
+					FString OriginalTraitIndicesString = Properties->operator[](PropertyIndex).DefaultValue;
+					if(!OriginalTraitIndicesString.IsEmpty() && OriginalTraitIndicesString != TEXT("()"))
+					{
+						OriginalTraitIndicesString = OriginalTraitIndicesString.TrimChar(TEXT('('));
+						OriginalTraitIndicesString = OriginalTraitIndicesString.TrimChar(TEXT(')'));
+						FString PinPathRemaining = OriginalTraitIndicesString;
+						FString Left, Right;
+						TArray<FString> IndexStrings;
+						while(PinPathRemaining.Split(TEXT(","), &Left, &Right))
+						{
+							IndexStrings.Add(Left.TrimStartAndEnd());
+							Left.Empty();
+							PinPathRemaining = Right;
+						}
+						if (!Right.IsEmpty())
+						{
+							IndexStrings.Add(Right.TrimStartAndEnd());
+						}
+
+						for(int32 PartIndex = 0; PartIndex < IndexStrings.Num(); PartIndex++)
+						{
+							const int32 OriginalTraitIndex = FCString::Atoi(*IndexStrings[PartIndex]);
+							const FRigVMFunctionCompilationPropertyDescription& TraitPropertyDescription = InCompilationData->WorkPropertyDescriptions[OriginalTraitIndex];
+							const FFunctionRegisterData Data = {InHash, ERigVMMemoryType::Work, OriginalTraitIndex};
+							const FRigVMOperand NewOperand = FindOrAddProperty(ERigVMMemoryType::Work, TraitPropertyDescription, TraitPropertyDescription.Name.ToString(), true /* unique */);
+							FunctionRegisterToOperand.Add(Data, NewOperand);
+							OperandMap.Add(FRigVMOperand(ERigVMMemoryType::Work, OriginalTraitIndex), NewOperand);
+
+							IndexStrings[PartIndex] = FString::FromInt(NewOperand.GetRegisterIndex());
+						}
+
+						IndexStrings.Remove(FString());
+						CustomDescription = *Description;
+						CustomDescription.DefaultValue = FString::Printf(TEXT("(%s)"), *FString::Join(IndexStrings, TEXT(",")));
+						FString NewName = FString::Printf(TEXT("Hash%lu_%s"), InHash, *CustomDescription.Name.ToString().RightChop(FunctionLibraryPrefix.Len()));
+						FRigVMPropertyDescription::SanitizeName(NewName);
+						CustomDescription.Name = *NewName;
+
+						Description = &CustomDescription;
+					}
+				}
+
+				if (MemoryType == ERigVMMemoryType::External)
+				{
+					// for external variables we don't change anything.
+					NewOperandNoOffset = OldOperandNoOffset;
+				}
+				else
+				{
+					FString NewName = Description->Name.ToString();
+
+					const UScriptStruct* ScriptStruct = Cast<UScriptStruct>(Description->CPPTypeObject.Get());
+
+					// instantiate function library specific work state as well as
+					// instruction set execute state - which is used for lazy blocks.
+					const bool bIsExecuteState = Description->CPPType.Equals(RigVMInstructionSetExecuteStateName);
+					const bool bIsTrait = ScriptStruct && ScriptStruct->IsChildOf(FRigVMTrait::StaticStruct());
+					if (NewName.StartsWith(FunctionLibraryPrefix) || bIsExecuteState || bIsTrait)
+					{
+						NewName = FString::Printf(TEXT("Hash%lu_%s"), InHash, *NewName.RightChop(FunctionLibraryPrefix.Len()));
+						FRigVMPropertyDescription::SanitizeName(NewName);
+					}
+
+					check(Description->CPPTypeObject.Get() != FRigVMForwardedMemoryHandle::StaticStruct());
+
+					NewOperandNoOffset = FindOrAddProperty(MemoryType, *Description, NewName, bIsExecuteState || bIsTrait);
+				}
+				
+				FFunctionRegisterData Data = { InHash, OldOperandNoOffset.GetMemoryType(), PropertyIndex};
+				FunctionRegisterToOperand.Add(Data, NewOperandNoOffset);
+               	OperandMap.Add(OldOperandNoOffset, NewOperandNoOffset);
+			}
+			
+			if (InOldOperand.GetRegisterOffset() != INDEX_NONE)
+			{
+				const FRigVMOperand OldOperandJustOffset(InOldOperand.GetMemoryType(), INDEX_NONE, InOldOperand.GetRegisterOffset());
+				if (!OperandMap.Contains(OldOperandJustOffset))
+				{
+					// copy over the register offset
+					const TArray<FRigVMFunctionCompilationPropertyPath>& PathDescriptions = 
+						InCompilationData->GetPropertyPathDescriptions(InOldOperand.GetMemoryType());
+
+					check(PathDescriptions.IsValidIndex(InOldOperand.GetRegisterOffset()));
+					const FRigVMFunctionCompilationPropertyPath& PropertyPath = PathDescriptions[InOldOperand.GetRegisterOffset()];
+					const int32 NewRegisterOffset = FindOrAddPropertyPath(NewOperandNoOffset, PropertyPath.HeadCPPType, PropertyPath.SegmentPath);
+					const FRigVMOperand NewOperandJustOffset(NewOperandNoOffset.GetMemoryType(), INDEX_NONE, NewRegisterOffset);
+					OperandMap.Add(OldOperandJustOffset, NewOperandJustOffset);
+				}
+			}				
+		}
+	};
+
+	// we should include the interface operands to make sure all memory is there.
+	for (const FRigVMOperand& InterfaceOperand : InInterfaceOperands)
+	{
+		ProcessOperand(ERigVMOpCode::Invalid, InterfaceOperand);
+	}
+
+	const FRigVMInstructionArray FunctionInstructions = InCompilationData->ByteCode.GetInstructions();
+	for (int32 Index = InFirstInstruction; Index <= FMath::Min(InLastInstruction, FunctionInstructions.Num() - 1); ++Index)
+	{
+		const ERigVMOpCode OpCode = FunctionInstructions[Index].OpCode;
+		const FRigVMOperandArray FunctionOperands = InCompilationData->ByteCode.GetOperandsForOp(FunctionInstructions[Index]);
+		for (const FRigVMOperand& FunctionOperand : FunctionOperands)
+		{
+			ProcessOperand(OpCode, FunctionOperand);
+		}
+	}
+
+	return true;
+}
+
+void FRigVMCompilerWorkData::SetupPropertyPathsForFoldedCopies(const FRigVMNodeExprAST* InExpr)
+{
+	// setup additional property paths needed for call externs accessing sub pin data directly
+	for(const TPair<int32, FString>& SegmentPathPair : InExpr->SegmentPathForChild)
+	{
+		const FRigVMExprAST* ChildExpr = InExpr->ChildAt(SegmentPathPair.Key);
+		check(ChildExpr);
+
+		if(ChildExpr->IsA(FRigVMExprAST::EType::CachedValue))
+		{
+			ChildExpr = ChildExpr->To<FRigVMCachedValueExprAST>()->GetVarExpr();
+		}
+
+		if(!ChildExpr->IsA(FRigVMExprAST::EType::Var))
+		{
+			continue;
+		}
+
+		const FRigVMVarExprAST* SourceVarExpr = Compiler->GetSourceVarExpr(ChildExpr);
+		const FRigVMOperand RootOperand = ExprToOperand.FindChecked(SourceVarExpr);
+		const FString HeadCPPType = SourceVarExpr->GetPin()->GetRootPin()->GetCPPType();
+
+		FindOrAddPropertyPath(RootOperand, HeadCPPType, SegmentPathPair.Value);
+	}
+}
+
+bool FRigVMCompilerWorkData::RemapInstructionAt(
+	const FRigVMInstruction& InOldInstruction,
+	const FString& InFunctionIdentifier,
+	const FRigVMFunctionCompilationData* InCompilationData,
+	FRigVMByteCode& InNewByteCode,
+	const FRigVMInstructionArray& InNewInstructions,
+	const URigVMFunctionReferenceNode* InFunctionReferenceNode,
+	const TFunction<int32(const int32&)>& InRemapInstructionIndex,
+	const TFunction<FRigVMOperand(const FRigVMOperand&)>& InRemapOperand,
+	const TFunction<int32(const int32&)>& InRemapBranchIndex,
+	const TFunction<int32(const int32&)>& InRemapCallableIndex)
+{
+	const FRigVMInstruction NewInstruction = InNewInstructions[InRemapInstructionIndex(InOldInstruction.Index)];
+	const FRigVMOperandArray NewOperandArray = InNewByteCode.GetOperandsForOp(NewInstruction);
+	const FRigVMOperandArray OldOperandArray = InCompilationData->ByteCode.GetOperandsForOp(InOldInstruction);
+	checkf(OldOperandArray.Num() == NewOperandArray.Num(), TEXT("Hit a function with a mismatch of function operands, function %s in package %s"), *InFunctionIdentifier, *Compiler->GetPackage()->GetPathName());
+
+	for (int32 j=0; j<NewOperandArray.Num(); ++j)
+	{
+		FRigVMOperand* NewOperand = const_cast<FRigVMOperand*>(&NewOperandArray[j]);
+		const FRigVMOperand OriginalOperand = *NewOperand;
+		checkf(OriginalOperand == OldOperandArray[j], TEXT("Hit a function with an invalid / non-matching function operand, function %s in package %s"), *InFunctionIdentifier, *Compiler->GetPackage()->GetPathName());
+		*NewOperand = InRemapOperand(OriginalOperand);
+	}
+
+	if (NewInstruction.OpCode == ERigVMOpCode::Execute)
+	{
+		FRigVMExecuteOp& Op = InNewByteCode.GetOpAt<FRigVMExecuteOp>(NewInstruction);
+		const FName& FunctionFName = InCompilationData->FunctionNames[Op.CallableIndex];
+		const int32 FunctionIndex = VM->AddRigVMFunction(FunctionFName);
+		const FRigVMFunction* Function = VM->GetFunctions()[FunctionIndex];
+		check(Function)
+		if (Function->Factory)
+		{
+			checkf(Function->Arguments.Num() <= NewOperandArray.Num(), TEXT("%s: invalid number of operands (%d) for dispatch '%s' - expected at least (%d)."), *Compiler->GetPackage()->GetPathName(), NewOperandArray.Num(), *Function->GetName(), Function->Arguments.Num());
+		}
+		else
+		{
+			checkf(Function->Arguments.Num() == NewOperandArray.Num(), TEXT("%s: invalid number of operands (%d) for function '%s' - expected (%d)."), *Compiler->GetPackage()->GetPathName(), NewOperandArray.Num(), *Function->GetName(), Function->Arguments.Num());
+		}
+		Op.CallableIndex = IntCastChecked<uint16>(FunctionIndex);
+	}
+	else if (NewInstruction.OpCode == ERigVMOpCode::InvokeCallable)
+	{
+		FRigVMInvokeCallableOp& Op = InNewByteCode.GetOpAt<FRigVMInvokeCallableOp>(NewInstruction);
+		const int32 CallableIndex = InRemapCallableIndex(IntCastChecked<int32>(Op.CallableIndex));
+		check(InNewByteCode.IsValidCallableIndex(CallableIndex));
+		const FRigVMCallableInfo* Callable = InNewByteCode.GetCallable(CallableIndex);
+		check(Callable)
+		check(Callable->Arguments.Num() == NewOperandArray.Num());
+		Op.CallableIndex = IntCastChecked<uint16>(CallableIndex);
+	}
+	else if (NewInstruction.OpCode == ERigVMOpCode::JumpToBranch)
+	{
+		FRigVMJumpToBranchOp& Op = InNewByteCode.GetOpAt<FRigVMJumpToBranchOp>(NewInstruction);
+		Op.FirstBranchInfoIndex = InRemapBranchIndex(Op.FirstBranchInfoIndex);
+		check(InNewByteCode.BranchInfos.IsValidIndex(Op.FirstBranchInfoIndex));
+		check(InNewByteCode.BranchInfos[Op.FirstBranchInfoIndex].InstructionIndex == NewInstruction.Index);
+	}
+	else if (NewInstruction.OpCode == ERigVMOpCode::RunInstructions)
+	{
+		FRigVMRunInstructionsOp& Op = InNewByteCode.GetOpAt<FRigVMRunInstructionsOp>(NewInstruction);
+		Op.StartInstruction = InRemapInstructionIndex(Op.StartInstruction);
+		Op.EndInstruction = InRemapInstructionIndex(Op.EndInstruction);
+	}
+
+	if (Settings.SetupNodeInstructionIndex)
+	{
+		if (const TArray<TWeakObjectPtr<UObject>>* Callstack = InCompilationData->ByteCode.GetCallstackForInstruction(InOldInstruction.Index))
+		{
+			if (Callstack->Num() > 1)
+			{
+				FRigVMCallstack InstructionCallstack;
+				InstructionCallstack.Stack = *Callstack;
+				if (InFunctionReferenceNode)
+				{
+					InstructionCallstack.Stack[0] = TWeakObjectPtr<UObject>(const_cast<URigVMFunctionReferenceNode*>(InFunctionReferenceNode));
+				}
+				InNewByteCode.SetSubject(NewInstruction.Index, InstructionCallstack.GetCallPath(), InstructionCallstack.GetStack());
+			}
+		}
+
+		// also store the instruction for the subject so that profiling can determine the cost of functions.
+		if (InFunctionReferenceNode)
+		{
+			InNewByteCode.AddInstructionForSubject(const_cast<URigVMFunctionReferenceNode*>(InFunctionReferenceNode), NewInstruction.Index);
+		}
+	}
+
+	// remap the input and output operands per instruction
+	TArray<FRigVMOperand> InputOperands, OutputOperands;
+	if (InCompilationData->ByteCode.InputOperandsPerInstruction.IsValidIndex(InOldInstruction.Index))
+	{
+		for (const FRigVMOperand& Operand : InCompilationData->ByteCode.InputOperandsPerInstruction[InOldInstruction.Index])
+		{
+			InputOperands.Add(InRemapOperand(Operand));
+		}
+	}
+	if (InCompilationData->ByteCode.OutputOperandsPerInstruction.IsValidIndex(InOldInstruction.Index))
+	{
+		for (const FRigVMOperand& Operand : InCompilationData->ByteCode.OutputOperandsPerInstruction[InOldInstruction.Index])
+		{
+			OutputOperands.Add(InRemapOperand(Operand));
+		}
+	}
+	InNewByteCode.SetOperandsForInstruction(NewInstruction, InputOperands, OutputOperands);
+	return true;
+}
+
 void FRigVMCompilerWorkData::ReportInfo(const FString& InMessage) const
 {
 	Settings.ReportInfo(InMessage);
@@ -379,7 +838,7 @@ void FRigVMCompilerWorkData::RemoveOverrideReportDelegate()
 }
 
 URigVMCompiler::URigVMCompiler()
-	: CurrentCompilationFunction(nullptr)
+	: CurrentlyCompilingFunctionNode(nullptr)
 {
 }
 
@@ -395,8 +854,10 @@ bool URigVMCompiler::Compile(const FRigVMCompileSettings& InSettings, TArray<URi
 	TSharedPtr<FRigVMParserAST> InAST, FRigVMFunctionCompilationData* OutFunctionCompilationData)
 {
 	FRigVMCompilerWorkData WorkData;
+	WorkData.Compiler = this;
 	WorkData.Settings = InSettings;
 	WorkData.Graphs = InGraphs;
+	WorkData.FunctionCompilationData = OutFunctionCompilationData;
 
 	FRigVMCompileSettings& Settings = WorkData.Settings;
 	
@@ -478,13 +939,13 @@ bool URigVMCompiler::Compile(const FRigVMCompileSettings& InSettings, TArray<URi
 	URigVMFunctionLibrary* FunctionLibrary = InGraphs[0]->GetDefaultFunctionLibrary();
 	bool bEncounteredGraphError = false;
 
-	TMap<FString, const FRigVMFunctionCompilationData*> CurrentCompiledFunctions;
+	TMap<FString, FCompiledFunctionData> CurrentCompiledFunctions;
 
 #if WITH_EDITOR
 	// During cooking, compile all public functions
 	if (IsRunningCookCommandlet())
 	{
-		if (!CurrentCompilationFunction)
+		if (!CurrentlyCompilingFunctionNode)
 		{
 			if (FunctionLibrary)
 			{
@@ -609,7 +1070,7 @@ bool URigVMCompiler::Compile(const FRigVMCompileSettings& InSettings, TArray<URi
 							if (bSuccessfullCompilation || CompilationData->IsValid())
 							{
 								FunctionData->PatchSharedArgumentOperandsIfRequired();
-								CurrentCompiledFunctions.Add(FunctionHash, CompilationData);
+								CurrentCompiledFunctions.Add(FunctionHash, {ReferenceNode, CompilationData});
 							}
 							else
 							{
@@ -781,7 +1242,7 @@ bool URigVMCompiler::Compile(const FRigVMCompileSettings& InSettings, TArray<URi
 					}
 
 					FString FunctionHash = FunctionReferenceNode->GetReferencedFunctionHeader().GetHash();
-					if (const FRigVMFunctionCompilationData** CompilationData = CompiledFunctions.Find(FunctionHash))
+					if (const FCompiledFunctionData* CompiledData = CompiledFunctions.Find(FunctionHash))
 					{
 						// todo: this should go into a helper function
 						TArray<FRigVMExternalVariable> AvailableExternalVariables;
@@ -798,9 +1259,9 @@ bool URigVMCompiler::Compile(const FRigVMCompileSettings& InSettings, TArray<URi
 						{
 							FRigVMExternalVariable::MergeExternalVariable(AvailableExternalVariables, InExternalVariable);
 						}
-						if (CurrentCompilationFunction)
+						if (CurrentlyCompilingFunctionNode)
 						{
-							if (const FRigVMGraphFunctionData* FunctionData = CurrentCompilationFunction->GetFunctionHeader().GetFunctionData())
+							if (const FRigVMGraphFunctionData* FunctionData = CurrentlyCompilingFunctionNode->GetFunctionHeader().GetFunctionData())
 							{
 								// private functions are allowed to access all external variables
 								if (!FunctionData->IsPublic())
@@ -819,7 +1280,7 @@ bool URigVMCompiler::Compile(const FRigVMCompileSettings& InSettings, TArray<URi
 								}
 							}
 						}
-						for (const TPair<int32, FName>& Pair : (*CompilationData)->ExternalRegisterIndexToVariable)
+						for (const TPair<int32, FName>& Pair : CompiledData->CompilationData->ExternalRegisterIndexToVariable)
 						{
 							FName OuterName = FunctionReferenceNode->GetOuterVariableName(Pair.Value);
 							if (OuterName.IsNone())
@@ -986,7 +1447,6 @@ bool URigVMCompiler::Compile(const FRigVMCompileSettings& InSettings, TArray<URi
 						for (TFieldIterator<FProperty> It(ScriptStruct, EFieldIterationFlags::None); It; ++It)
 						{
 							const FRigVMTemplateArgument ExpectedArgument = FRigVMTemplateArgument::Make(*It);
-							const TRigVMTypeIndex ExpectedTypeIndex = ExpectedArgument.GetSupportedTypeIndices()[0];
 							if (URigVMPin* Pin = UnitNode->FindPin(ExpectedArgument.Name.ToString()))
 							{
 								if (Pin->GetTypeIndex() != ExpectedArgument.GetTypeIndex(0))
@@ -1145,15 +1605,15 @@ bool URigVMCompiler::Compile(const FRigVMCompileSettings& InSettings, TArray<URi
 	// If compiling a function, check if all mutable paths lead to a return node
 	// - Return node path should go through the completed pins of control flow nodes
 	// - At least one path on sequence nodes should lead to a return node
-	if (CurrentCompilationFunction)
+	if (CurrentlyCompilingFunctionNode)
 	{
-		const URigVMFunctionEntryNode* EntryNode = CurrentCompilationFunction->GetEntryNode();
+		const URigVMFunctionEntryNode* EntryNode = CurrentlyCompilingFunctionNode->GetEntryNode();
 		if (EntryNode && EntryNode->IsMutable())
 		{
-			URigVMFunctionReturnNode* ReturnNode = CurrentCompilationFunction->GetReturnNode();
+			URigVMFunctionReturnNode* ReturnNode = CurrentlyCompilingFunctionNode->GetReturnNode();
 			if (!ReturnNode)
 			{
-				const FString FunctionCompilationErrorMessage = FString::Printf(TEXT("[%s] Mutable function graph %s does not contain a return node."), *EntryNode->GetPackage()->GetPathName(), *CurrentCompilationFunction->GetName());
+				const FString FunctionCompilationErrorMessage = FString::Printf(TEXT("[%s] Mutable function graph %s does not contain a return node."), *EntryNode->GetPackage()->GetPathName(), *CurrentlyCompilingFunctionNode->GetName());
 				Settings.ASTSettings.Report(EMessageSeverity::Error, nullptr, FunctionCompilationErrorMessage);
 				bEncounteredGraphError = true;
 			}
@@ -1209,7 +1669,7 @@ bool URigVMCompiler::Compile(const FRigVMCompileSettings& InSettings, TArray<URi
 
 			if (!bReturnNodeFound)
 			{
-				const FString FunctionCompilationErrorMessage = FString::Printf(TEXT("[%s] Not all paths of the mutable function graph %s lead to a return node."), *CurrentCompilationFunction->GetPackage()->GetPathName(), *CurrentCompilationFunction->GetName());
+				const FString FunctionCompilationErrorMessage = FString::Printf(TEXT("[%s] Not all paths of the mutable function graph %s lead to a return node."), *CurrentlyCompilingFunctionNode->GetPackage()->GetPathName(), *CurrentlyCompilingFunctionNode->GetName());
 				Settings.ASTSettings.Report(EMessageSeverity::Error, nullptr, FunctionCompilationErrorMessage);
 				bEncounteredGraphError = true;
 			}
@@ -1221,16 +1681,16 @@ bool URigVMCompiler::Compile(const FRigVMCompileSettings& InSettings, TArray<URi
 		return false;
 	}
 #endif
-	
+
 	OutVM->ClearExternalVariables(OutVMContext);
 	
 	for (const FRigVMExternalVariable& ExternalVariable : InExternalVariables)
 	{
-		if (CurrentCompilationFunction && CurrentCompilationFunction->ContainsLocalVariable(ExternalVariable.GetName()))
+		if (CurrentlyCompilingFunctionNode && CurrentlyCompilingFunctionNode->ContainsLocalVariable(ExternalVariable.GetName()))
 		{
 			continue;
 		}
-		FRigVMOperand Operand = OutVM->AddExternalVariable(OutVMContext, ExternalVariable, CurrentCompilationFunction != nullptr);
+		FRigVMOperand Operand = OutVM->AddExternalVariable(OutVMContext, ExternalVariable, CurrentlyCompilingFunctionNode != nullptr);
 		FString Hash = FString::Printf(TEXT("Variable::%s"), *ExternalVariable.GetName().ToString());
 		OutOperands->Add(Hash, Operand);
 	}
@@ -1263,7 +1723,7 @@ bool URigVMCompiler::Compile(const FRigVMCompileSettings& InSettings, TArray<URi
 	WorkData.Context = &OutVMContext;
 	WorkData.ExecuteContextStruct = Settings.GetExecuteContextStruct();
 	WorkData.PinPathToOperand = OutOperands;
-	WorkData.bSetupMemory = true;
+	WorkData.CompilerPhase = ERigVMCompilerPhase_SetupMemory;
 	WorkData.ProxySources = &WorkData.AST->SharedOperandPins;
 
 	// tbd: do we need this only when we have no pins?
@@ -1279,13 +1739,95 @@ bool URigVMCompiler::Compile(const FRigVMCompileSettings& InSettings, TArray<URi
 
 	UE_LOG_RIGVMMEMORY(TEXT("RigVMCompiler: Begin '%s'..."), *InGraph->GetPathName());
 
+	// convert the compiled functions into callables. this inlines the callables at the beginning of the bytecode.
+	// the order of callables should be in reverse order of dependency - so the leaves of the dependency tree are first.
+	for (const TPair<FString, FCompiledFunctionData>& Pair : CompiledFunctions)
+	{
+		const FRigVMFunctionCompilationData* FunctionCompilationData = Pair.Value.CompilationData;
+		int32 FirstInstructionIndex = 0;
+		
+		// first look at all dependencies of this function and inline them. 
+		for (int32 CallableIndex = 0; CallableIndex < Pair.Value.CompilationData->ByteCode.NumCallables(); CallableIndex++)
+		{
+			const FRigVMCallableInfo* CallableInfo = Pair.Value.CompilationData->ByteCode.GetCallable(CallableIndex);
+			check(CallableInfo);
+
+			const FRigVMGraphFunctionHeader Header = FRigVMGraphFunctionHeader::FindGraphFunctionHeaderFromHash(CallableInfo->Name.ToString());
+			if (!CompileCallable(WorkData, OutVM, OutVMContext, CallableInfo, Header, Pair.Value.CompilationData))
+			{
+				return false;
+			}
+
+			FirstInstructionIndex = CallableInfo->LastInstruction + 1;
+		}
+
+		if (!FunctionCompilationData->SupportsCallable())
+		{
+			continue;
+		}
+
+		if (CVarRigVMCompileFunctionsToCallables->GetBool())
+		{
+			const FRigVMGraphFunctionHeader& Header = Pair.Value.FunctionReferenceNode->GetReferencedFunctionHeader();
+
+			// convert this function to a callable
+			FRigVMCallableInfo Callable;
+			Callable.Hash = Pair.Value.CompilationData->Hash;
+			Callable.Name = *Header.GetHash(); 
+			Callable.FirstInstruction = FirstInstructionIndex;
+			Callable.LastInstruction = Pair.Value.CompilationData->ByteCode.GetNumInstructions() - 1;
+
+			// convert all arguments
+			bool bAllArgumentsValid = true;
+			for (const FRigVMGraphFunctionArgument& FunctionArgument : Header.Arguments)
+			{
+				if (FunctionArgument.IsExecuteContext())
+				{
+					continue;
+				}
+				
+				FRigVMCallableArgument CallableArgument;
+				CallableArgument.Name = FunctionArgument.Name;
+				CallableArgument.TypeString = FunctionArgument.CPPType.ToString();
+
+				// look for the argument in the function's operands
+				const FRigVMOperand* InterfaceOperand = FunctionCompilationData->InterfaceOperands.Find(FunctionArgument.Name);
+				if (!InterfaceOperand)
+				{
+					bAllArgumentsValid = false;
+					UE_LOG(LogRigVMDeveloper, Warning, TEXT("Cannot convert function to callable '%s', missing interface operand for argument '%s'."), *Pair.Key, *FunctionArgument.Name.ToString());
+					break;
+				}
+
+				CallableArgument.InterfaceOperand = *InterfaceOperand;
+				CallableArgument.ForwardedOperand = FRigVMOperand();
+				CallableArgument.Direction = FunctionArgument.Direction;
+				Callable.Arguments.Add(CallableArgument);
+			}
+
+			// compile the callable of the top level function
+			if (bAllArgumentsValid)
+			{
+				if (!Pair.Value.FunctionReferenceNode)
+				{
+					return false;
+				}
+
+				if (!CompileCallable(WorkData, OutVM, OutVMContext, &Callable, Header, Pair.Value.CompilationData))
+				{
+					return false;
+				}
+			}
+		}
+	}
+
 	// If we are compiling a function, we want the first registers to represent the interface pins (in the order of the pins)
 	// so they can be replaced when inlining the function
-	if (CurrentCompilationFunction)
+	if (CurrentlyCompilingFunctionNode)
 	{
-		URigVMFunctionEntryNode* EntryNode = CurrentCompilationFunction->GetEntryNode();
-		URigVMFunctionReturnNode* ReturnNode = CurrentCompilationFunction->GetReturnNode();
-		for (URigVMPin* Pin : CurrentCompilationFunction->GetPins())
+		URigVMFunctionEntryNode* EntryNode = CurrentlyCompilingFunctionNode->GetEntryNode();
+		URigVMFunctionReturnNode* ReturnNode = CurrentlyCompilingFunctionNode->GetReturnNode();
+		for (URigVMPin* Pin : CurrentlyCompilingFunctionNode->GetPins())
 		{
 			URigVMPin* InterfacePin = nullptr;
 			if (Pin->GetDirection() == ERigVMPinDirection::Input ||
@@ -1293,7 +1835,7 @@ bool URigVMCompiler::Compile(const FRigVMCompileSettings& InSettings, TArray<URi
 			{
 				if(EntryNode == nullptr)
 				{
-					WorkData.ReportErrorf(TEXT("[%s] Corrupt library node '%s' - Missing entry node."), *CurrentCompilationFunction->GetPackage()->GetPathName(), *CurrentCompilationFunction->GetPathName());
+					WorkData.ReportErrorf(TEXT("[%s] Corrupt library node '%s' - Missing entry node."), *CurrentlyCompilingFunctionNode->GetPackage()->GetPathName(), *CurrentlyCompilingFunctionNode->GetPathName());
 					return false;
 				}
 				InterfacePin = EntryNode->FindPin(Pin->GetName());
@@ -1302,7 +1844,7 @@ bool URigVMCompiler::Compile(const FRigVMCompileSettings& InSettings, TArray<URi
 			{
 				if(ReturnNode == nullptr)
 				{
-					WorkData.ReportErrorf(TEXT("[%s] Corrupt library node '%s' - Missing return node."), *CurrentCompilationFunction->GetPackage()->GetPathName(), *CurrentCompilationFunction->GetPathName());
+					WorkData.ReportErrorf(TEXT("[%s] Corrupt library node '%s' - Missing return node."), *CurrentlyCompilingFunctionNode->GetPackage()->GetPathName(), *CurrentlyCompilingFunctionNode->GetPathName());
 					return false;
 				}
 				InterfacePin = ReturnNode->FindPin(Pin->GetName());
@@ -1310,13 +1852,21 @@ bool URigVMCompiler::Compile(const FRigVMCompileSettings& InSettings, TArray<URi
 
 			if(InterfacePin == nullptr)
 			{
-				WorkData.ReportErrorf(TEXT("[%s} Corrupt library node '%s' - Pin '%s' is not part of the entry / return node."), *CurrentCompilationFunction->GetPackage()->GetPathName(), *CurrentCompilationFunction->GetPathName(), *Pin->GetPathName());
+				WorkData.ReportErrorf(TEXT("[%s} Corrupt library node '%s' - Pin '%s' is not part of the entry / return node."), *CurrentlyCompilingFunctionNode->GetPackage()->GetPathName(), *CurrentlyCompilingFunctionNode->GetPathName(), *Pin->GetPathName());
 				return false;
 			}
-			
+
 			FRigVMASTProxy PinProxy = FRigVMASTProxy::MakeFromUObject(InterfacePin);
 			FRigVMVarExprAST* TempVarExpr = WorkData.AST->MakeExpr<FRigVMVarExprAST>(FRigVMExprAST::EType::Var, PinProxy);
-			FindOrAddRegister(TempVarExpr, WorkData);
+			const FRigVMOperand InterfaceOperand = FindOrAddRegister(TempVarExpr, WorkData);
+			if (WorkData.FunctionCompilationData)
+			{
+				if (!InterfacePin->IsExecuteContext())
+				{
+					const FRigVMOperand RootOperand(InterfaceOperand.GetMemoryType(), InterfaceOperand.GetRegisterIndex(), INDEX_NONE);
+					WorkData.FunctionCompilationData->InterfaceOperands.FindOrAdd(InterfacePin->GetFName()) = RootOperand;
+				}
+			}
 		}
 	}
 
@@ -1384,7 +1934,7 @@ bool URigVMCompiler::Compile(const FRigVMCompileSettings& InSettings, TArray<URi
 	}
 
 	const TMap<const FRigVMExprAST*, bool> ExprSetupMemoryComplete = WorkData.ExprComplete;
-	WorkData.bSetupMemory = false;
+	WorkData.CompilerPhase = ERigVMCompilerPhase_BuildInstructions;
 	WorkData.ExprComplete.Reset();
 
 	// sort the expressions in each block by depth
@@ -1408,7 +1958,7 @@ bool URigVMCompiler::Compile(const FRigVMCompileSettings& InSettings, TArray<URi
 		}
 	}
 
-	if (!CurrentCompilationFunction)
+	if (!CurrentlyCompilingFunctionNode)
 	{
 		if (WorkData.VM->GetByteCode().GetInstructions().Num() == 0)
 		{
@@ -1466,9 +2016,10 @@ bool URigVMCompiler::Compile(const FRigVMCompileSettings& InSettings, TArray<URi
 		WorkData.VM->GetByteCode().GetOpAt<FRigVMJumpOp>(JumpToEndOfBlocksExternByte).InstructionIndex = InstructionsToJump;
 	}
 	
-	if (!CurrentCompilationFunction)
+	if (!CurrentlyCompilingFunctionNode)
 	{
 		WorkData.VM->GetByteCode().AlignByteCode();
+		WorkData.VM->GetByteCode().Shrink();
 	}
 
 	// now that we have determined the needed memory, let's
@@ -1496,11 +2047,14 @@ bool URigVMCompiler::Compile(const FRigVMCompileSettings& InSettings, TArray<URi
 	}
 
 	// Store function compile data
-	if (CurrentCompilationFunction && OutFunctionCompilationData)
+	if (CurrentlyCompilingFunctionNode && OutFunctionCompilationData)
 	{
 		OutFunctionCompilationData->ByteCode = WorkData.VM->ByteCodeStorage;
 		OutFunctionCompilationData->FunctionNames = WorkData.VM->FunctionNamesStorage;
-		OutFunctionCompilationData->Operands = *OutOperands;
+		if (&OutFunctionCompilationData->Operands != OutOperands)
+		{
+			OutFunctionCompilationData->Operands = *OutOperands;
+		}
 		OutFunctionCompilationData->bEncounteredSurpressedErrors = bSurpressedASTError;
 
 		for (uint8 MemoryTypeIndex=0; MemoryTypeIndex<(uint8)ERigVMMemoryType::Invalid; ++MemoryTypeIndex)
@@ -1594,7 +2148,7 @@ bool URigVMCompiler::Compile(const FRigVMCompileSettings& InSettings, TArray<URi
 					if (UsedExternalVariableRegisters.Contains(Pair.Value.GetRegisterIndex()))
 					{
 						FString VariableName = Pair.Key.RightChop(VariablePrefix.Len());
-						if (CurrentCompilationFunction && CurrentCompilationFunction->ContainsLocalVariable(*VariableName))
+						if (CurrentlyCompilingFunctionNode && CurrentlyCompilingFunctionNode->ContainsLocalVariable(*VariableName))
 						{
 							continue;
 						}
@@ -1608,7 +2162,7 @@ bool URigVMCompiler::Compile(const FRigVMCompileSettings& InSettings, TArray<URi
 	}
 
 	// Localize all of the registry information
-	if(CurrentCompilationFunction == nullptr)
+	if(CurrentlyCompilingFunctionNode == nullptr)
 	{
 		WorkData.VM->CreateLocalizedRegistryIfRequired();
 	}
@@ -1616,7 +2170,7 @@ bool URigVMCompiler::Compile(const FRigVMCompileSettings& InSettings, TArray<URi
 	// make sure all functions are known and resolved now.
 	WorkData.VM->ResolveFunctionsIfRequired();
 	
-	if (!CurrentCompilationFunction)
+	if (!CurrentlyCompilingFunctionNode)
 	{
 		CompileTimer.Stop();
 
@@ -1640,7 +2194,7 @@ bool URigVMCompiler::CompileFunction(const FRigVMCompileSettings& InSettings, co
 	}
 
 	FFunctionCompilationScope FunctionCompilationScope(this, InLibraryNode);
-	TGuardValue<const URigVMLibraryNode*> CompilationGuard(CurrentCompilationFunction, InLibraryNode);
+	TGuardValue<const URigVMLibraryNode*> CompilationGuard(CurrentlyCompilingFunctionNode, InLibraryNode);
 
 	URigVMController* LibraryController = InController->GetControllerForGraph(InLibraryNode->GetContainedGraph());
 	if(LibraryController == nullptr)
@@ -1724,6 +2278,207 @@ bool URigVMCompiler::CompileFunction(const FRigVMCompileSettings& InSettings, co
 	return bSuccess;
 }
 
+bool URigVMCompiler::CompileCallable(FRigVMCompilerWorkData& WorkData, URigVM* OutVM, FRigVMExtendedExecuteContext& OutVMContext,
+	const FRigVMCallableInfo* InCallable, const FRigVMGraphFunctionHeader& InHeader, const FRigVMFunctionCompilationData* InFunctionCompilationData)
+{
+	check(OutVM);
+	check(InCallable);
+	check(InFunctionCompilationData);
+
+	if (!InHeader.IsValid())
+	{
+		return false;
+	}
+	
+	// if we already compiled this callable - carry on
+	FRigVMByteCode& ByteCode = OutVM->GetByteCode();
+	if (ByteCode.FindCallable(InCallable->Hash))
+	{
+		return true;
+	}
+
+	const FRigVMByteCode& FunctionByteCode = InFunctionCompilationData->ByteCode;
+	const FRigVMInstructionArray FunctionInstructions = FunctionByteCode.GetInstructions();
+
+	// Add all required properties, property paths and set up WorkData.OperandMapPerFunction
+	TArray<FRigVMOperand> ArgumentOperands;
+	TMap<FRigVMOperand,FRigVMOperand> OldForwardedHandleToInterfaceOperand;
+	ArgumentOperands.Reserve(InCallable->Arguments.Num());
+	for (const FRigVMCallableArgument& Argument : InCallable->Arguments)
+	{
+		ArgumentOperands.Add(Argument.InterfaceOperand);
+		if (Argument.ForwardedOperand.IsValid())
+		{
+			OldForwardedHandleToInterfaceOperand.Add(Argument.ForwardedOperand, Argument.InterfaceOperand);
+		}
+	}
+	(void)WorkData.AddPropertiesForFunction(true, InCallable->Hash, InHeader, InFunctionCompilationData, InCallable->FirstInstruction, InCallable->LastInstruction, ArgumentOperands, OldForwardedHandleToInterfaceOperand);
+
+	FRigVMCallableInfo NewCallable;
+	NewCallable.Index = ByteCode.NumCallables();
+	NewCallable.Hash = InCallable->Hash;
+	NewCallable.Name = InCallable->Name;
+	NewCallable.FirstInstruction = ByteCode.GetNumInstructions();
+	NewCallable.LastInstruction = NewCallable.FirstInstruction + InCallable->LastInstruction - InCallable->FirstInstruction;
+
+	TMap<FRigVMOperand, FRigVMOperand> NewInterfaceOperandToForwardedHandle;
+	const TMap<FRigVMOperand, FRigVMOperand>& OperandMap = WorkData.OperandMapPerFunction.FindChecked(NewCallable.Hash);
+	auto RemapOperand = [&OperandMap, &NewInterfaceOperandToForwardedHandle, &OldForwardedHandleToInterfaceOperand, InHeader](const FRigVMOperand& InOldOperand, bool bRemapToForwardedHandle)
+	{
+		FRigVMOperand OperandNoOffset(InOldOperand.GetMemoryType(), InOldOperand.GetRegisterIndex(), INDEX_NONE);
+		if (const FRigVMOperand* InterfaceOperand = OldForwardedHandleToInterfaceOperand.Find(OperandNoOffset))
+		{
+			OperandNoOffset = *InterfaceOperand;
+		}
+		
+		const FRigVMOperand* ExistingOperandNoOffset = OperandMap.Find(OperandNoOffset);
+		if (!ExistingOperandNoOffset)
+		{
+			checkf(false, TEXT("Missing remapped operand for callable '%s'"), *InHeader.LibraryPointer.GetLibraryNodePath());
+			return FRigVMOperand();
+		}
+
+		FRigVMOperand NewOperand = *ExistingOperandNoOffset;
+		if (bRemapToForwardedHandle)
+		{
+			if (const FRigVMOperand* ForwardedOperand = NewInterfaceOperandToForwardedHandle.Find(NewOperand))
+			{
+				NewOperand = *ForwardedOperand;
+			}
+		}
+
+		if (InOldOperand.GetRegisterOffset() != INDEX_NONE)
+		{
+			const FRigVMOperand OperandJustOffset(InOldOperand.GetMemoryType(), INDEX_NONE, InOldOperand.GetRegisterOffset());
+			const FRigVMOperand* ExistingOperandJustOffset = OperandMap.Find(OperandJustOffset);
+			if (!ExistingOperandJustOffset)
+			{
+				checkf(false, TEXT("Missing remapped register offset for callable '%s'"), *InHeader.LibraryPointer.GetLibraryNodePath());
+				return FRigVMOperand();
+			}
+			NewOperand.RegisterOffset = ExistingOperandJustOffset->RegisterOffset;
+		}
+
+		return NewOperand;
+	};
+	
+	// now add the forwarded handles for all arguments.
+    // we are only doing this if we are compiling a non-function bytecode
+    NewCallable.Arguments = InCallable->Arguments;
+
+    // if this is using a forwarded memory handle - we should look up the original property from it
+    for (int32 ArgumentIndex = 0; ArgumentIndex < InCallable->Arguments.Num(); ArgumentIndex++)
+    {
+    	FRigVMCallableArgument& Argument = NewCallable.Arguments[ArgumentIndex]; 
+    	Argument.InterfaceOperand = RemapOperand(Argument.InterfaceOperand, false);
+    	check(Argument.InterfaceOperand.GetMemoryType() == ERigVMMemoryType::Work);
+    	check(Argument.InterfaceOperand.GetRegisterOffset() == INDEX_NONE);
+
+    	const FRigVMPropertyDescription& InterfaceDescription = WorkData.PropertyDescriptions[ERigVMMemoryType::Work][Argument.InterfaceOperand.GetRegisterIndex()];
+    	check(InterfaceDescription.CPPTypeObject != FRigVMForwardedMemoryHandle::StaticStruct());
+    	check(InterfaceDescription.CPPType == Argument.TypeString);
+
+    	static const FString ForwardedHandleSuffix = TEXT("_ForwardedHandle");
+    	const FString ForwardedNewName = InterfaceDescription.Name.ToString() + ForwardedHandleSuffix;
+
+    	FRigVMFunctionCompilationPropertyDescription ForwardedDescription;
+    	ForwardedDescription.Name = *ForwardedNewName;
+    	ForwardedDescription.CPPType = FRigVMForwardedMemoryHandle::StaticStruct()->GetStructCPPName();
+    	ForwardedDescription.CPPTypeObject = FRigVMForwardedMemoryHandle::StaticStruct();
+    	ForwardedDescription.DefaultValue = TEXT("()");
+    	Argument.ForwardedOperand = WorkData.FindOrAddProperty(ERigVMMemoryType::Work, ForwardedDescription, ForwardedNewName, true);
+    	NewInterfaceOperandToForwardedHandle.Add(Argument.InterfaceOperand, Argument.ForwardedOperand);
+    }
+
+	ByteCode.CallableInfos.Add(NewCallable);
+
+
+	if (NewCallable.GetNumInstructions() > 0)
+	{
+		check(FunctionInstructions.IsValidIndex(InCallable->FirstInstruction));
+		check(FunctionInstructions.IsValidIndex(InCallable->LastInstruction));
+		ByteCode.InlineFunction(&InFunctionCompilationData->ByteCode, InCallable->FirstInstruction, InCallable->LastInstruction);
+
+		auto RemapInstructionIndex = [InCallable, &NewCallable](const int32& InOldInstructionIndex)
+		{
+			check(FMath::IsWithinInclusive(InOldInstructionIndex, InCallable->FirstInstruction, InCallable->LastInstruction + 1))
+			const int32 RemappedIndex = InOldInstructionIndex - InCallable->FirstInstruction + NewCallable.FirstInstruction;
+			check(FMath::IsWithinInclusive(RemappedIndex, NewCallable.FirstInstruction, NewCallable.LastInstruction + 1))
+			return RemappedIndex;
+		};
+	
+		auto RemapOperandForByteCode = [&RemapOperand](const FRigVMOperand& InOldOperand)
+		{
+			return RemapOperand(InOldOperand, true);
+		};
+
+		// Import branch infos, with the proper instruction indices
+		TMap<int32,int32> BranchIndexMap;
+		for(const FRigVMBranchInfo& BranchInfo : FunctionByteCode.BranchInfos)
+		{
+			if (!FMath::IsWithinInclusive(BranchInfo.InstructionIndex, InCallable->FirstInstruction, InCallable->LastInstruction) ||
+				!FMath::IsWithinInclusive(BranchInfo.LastInstruction, InCallable->FirstInstruction, InCallable->LastInstruction))
+			{
+				continue;
+			}
+			if (!FMath::IsWithinInclusive(BranchInfo.FirstInstruction, InCallable->FirstInstruction, InCallable->LastInstruction))
+			{
+				if (BranchInfo.FirstInstruction > InCallable->LastInstruction + 1)
+				{
+					continue;
+				}
+				if (BranchInfo.Label != FRigVMStruct::ControlFlowCompletedName)
+				{
+					continue;
+				}
+			}
+
+			FRigVMBranchInfo NewBranchInfo = BranchInfo;
+			NewBranchInfo.InstructionIndex = RemapInstructionIndex(BranchInfo.InstructionIndex);
+			NewBranchInfo.FirstInstruction = RemapInstructionIndex(BranchInfo.FirstInstruction);
+			NewBranchInfo.LastInstruction = RemapInstructionIndex(BranchInfo.LastInstruction);
+			
+			const int32 NewBranchIndex = ByteCode.AddBranchInfo(NewBranchInfo);
+			BranchIndexMap.Add(BranchInfo.Index, NewBranchIndex);
+		}
+
+		auto RemapBranchInfoIndex = [&BranchIndexMap](const int32& InBranchIndex)
+		{
+			return BranchIndexMap.FindChecked(InBranchIndex);
+		};
+
+		auto RemapCallableIndex = [&ByteCode, &FunctionByteCode](const int32& InCallableIndex)
+		{
+			const FRigVMCallableInfo* OldCallable = FunctionByteCode.GetCallable(InCallableIndex);
+			check(OldCallable);
+			const FRigVMCallableInfo* NewCallable = ByteCode.FindCallable(OldCallable->Hash);
+			check(NewCallable);
+			return NewCallable->Index;
+		};
+		
+		// remap all operand, branches and callables
+		const FRigVMInstructionArray Instructions = ByteCode.GetInstructions();
+		for (int32 OldInstructionIndex = InCallable->FirstInstruction; OldInstructionIndex <= InCallable->LastInstruction; ++OldInstructionIndex)
+		{
+			if (!WorkData.RemapInstructionAt(
+				FunctionInstructions[OldInstructionIndex],
+				InHeader.LibraryPointer.GetLibraryNodePath(),
+				InFunctionCompilationData,
+				ByteCode,
+				Instructions,
+				nullptr,
+				RemapInstructionIndex,
+				RemapOperandForByteCode,
+				RemapBranchInfoIndex,
+				RemapCallableIndex))
+			{
+				return false;
+			}
+		}
+	}
+	return true;
+}
+
 bool URigVMCompiler::TraverseExpression(const FRigVMExprAST* InExpr, FRigVMCompilerWorkData& WorkData)
 {
 	if (WorkData.ExprToSkip.Contains(InExpr))
@@ -1743,7 +2498,7 @@ bool URigVMCompiler::TraverseExpression(const FRigVMExprAST* InExpr, FRigVMCompi
 			TSharedPtr<FRigVMCompilerWorkData::FLazyBlockInfo>& BlockInfo = WorkData.LazyBlocks.FindChecked(NextBlockHash.GetValue());
 
 			// if we are hitting the block for the first time - let's register the execution state operand
-			if(WorkData.bSetupMemory)
+			if(WorkData.CompilerPhase == ERigVMCompilerPhase_SetupMemory)
 			{
 				if(!BlockInfo->ExecuteStateOperand.IsValid())
 				{
@@ -1789,7 +2544,7 @@ bool URigVMCompiler::TraverseExpression(const FRigVMExprAST* InExpr, FRigVMCompi
 		}
 	}
 
-	if(!WorkData.bSetupMemory)
+	if(WorkData.CompilerPhase == ERigVMCompilerPhase_BuildInstructions)
 	{
 		// skip any expression which is not part of this block
 		if(WorkData.CurrentBlockHash.IsSet())
@@ -1914,7 +2669,7 @@ bool URigVMCompiler::TraverseBlock(const FRigVMBlockExprAST* InExpr, FRigVMCompi
 	// check if the block is under a lazy pin, in which case we need to set up a branch info
 	URigVMNode* CallExternNode = nullptr;
 	FRigVMBranchInfo BranchInfo;
-	if(!WorkData.bSetupMemory)
+	if(WorkData.CompilerPhase == ERigVMCompilerPhase_BuildInstructions)
 	{
 		if(const FRigVMExprAST* ParentExpr = InExpr->GetParent())
 		{
@@ -2011,7 +2766,7 @@ bool URigVMCompiler::TraverseBlock(const FRigVMBlockExprAST* InExpr, FRigVMCompi
 		}
 	}
 
-	if(!WorkData.bSetupMemory && !BranchInfo.Label.IsNone())
+	if(WorkData.CompilerPhase == ERigVMCompilerPhase_BuildInstructions && !BranchInfo.Label.IsNone())
 	{
 		// We need to make sure lazy blocks are properly populated, with all the operations that need to run
 		// during the evaluation of the lazy branch, even if some of the expressions
@@ -2048,7 +2803,7 @@ bool URigVMCompiler::TraverseEntry(const FRigVMEntryExprAST* InExpr, FRigVMCompi
 			return false;
 		}
 
-		if (WorkData.bSetupMemory)
+		if (WorkData.CompilerPhase == ERigVMCompilerPhase_SetupMemory)
 		{
 			TSharedPtr<FStructOnScope> DefaultStruct = UnitNode->ConstructStructInstance();
 			if (!TraverseChildren(InExpr, WorkData))
@@ -2077,6 +2832,18 @@ bool URigVMCompiler::TraverseEntry(const FRigVMEntryExprAST* InExpr, FRigVMCompi
 
 			// setup the instruction
 			int32 FunctionIndex = WorkData.VM->AddRigVMFunction(UnitNode->GetScriptStruct(), UnitNode->GetMethodName());
+#if WITH_EDITOR
+			const FRigVMFunction* Function = WorkData.VM->GetFunctions()[FunctionIndex];
+			check(Function);
+			if (Function->Factory)
+			{
+				checkf(Function->Arguments.Num() <= Operands.Num(), TEXT("%s: invalid number of operands (%d) for dispatch '%s' - expected at least (%d)."), *GetPackage()->GetPathName(), Operands.Num(), *Function->GetName(), Function->Arguments.Num());
+			}
+			else
+			{
+				checkf(Function->Arguments.Num() == Operands.Num(), TEXT("%s: invalid number of operands (%d) for function '%s' - expected (%d)."), *GetPackage()->GetPathName(), Operands.Num(), *Function->GetName(), Function->Arguments.Num());
+			}
+#endif
 			WorkData.VM->GetByteCode().AddExecuteOp(FunctionIndex, Operands, 0, 0);
 		
 			int32 EntryInstructionIndex = WorkData.VM->GetByteCode().GetNumInstructions() - 1;
@@ -2097,26 +2864,26 @@ bool URigVMCompiler::TraverseEntry(const FRigVMEntryExprAST* InExpr, FRigVMCompi
 			}
 		}
 	}
-	else if (CurrentCompilationFunction && InExpr->NumParents() == 0)
+	else if (CurrentlyCompilingFunctionNode && InExpr->NumParents() == 0)
 	{
 		// Initialize local variables
-		if (WorkData.bSetupMemory)
+		if (WorkData.CompilerPhase == ERigVMCompilerPhase_SetupMemory)
 		{
-			TArray<FRigVMGraphVariableDescription> LocalVariables = CurrentCompilationFunction->GetContainedGraph()->GetLocalVariables();
+			TArray<FRigVMGraphVariableDescription> LocalVariables = CurrentlyCompilingFunctionNode->GetContainedGraph()->GetLocalVariables();
 			for (const FRigVMGraphVariableDescription& Variable : LocalVariables)
 			{
-				FString Path = FString::Printf(TEXT("LocalVariableDefault::%s|%s::Const"), *CurrentCompilationFunction->GetContainedGraph()->GetGraphName(), *Variable.Name.ToString());
+				FString Path = FString::Printf(TEXT("LocalVariableDefault::%s|%s::Const"), *CurrentlyCompilingFunctionNode->GetContainedGraph()->GetGraphName(), *Variable.Name.ToString());
 				FRigVMOperand Operand = WorkData.AddProperty(ERigVMMemoryType::Literal, *Path, Variable.CPPType, Variable.CPPTypeObject, Variable.DefaultValue);
 				WorkData.PinPathToOperand->Add(Path, Operand);
 			}
 		}
 		else
 		{
-			TArray<FRigVMGraphVariableDescription> LocalVariables = CurrentCompilationFunction->GetContainedGraph()->GetLocalVariables();
+			TArray<FRigVMGraphVariableDescription> LocalVariables = CurrentlyCompilingFunctionNode->GetContainedGraph()->GetLocalVariables();
 			for (const FRigVMGraphVariableDescription& Variable : LocalVariables)
 			{
-				FString TargetPath = FString::Printf(TEXT("LocalVariable::%s|%s"), *CurrentCompilationFunction->GetContainedGraph()->GetGraphName(), *Variable.Name.ToString());
-				FString SourcePath = FString::Printf(TEXT("LocalVariableDefault::%s|%s::Const"), *CurrentCompilationFunction->GetContainedGraph()->GetGraphName(), *Variable.Name.ToString());
+				FString TargetPath = FString::Printf(TEXT("LocalVariable::%s|%s"), *CurrentlyCompilingFunctionNode->GetContainedGraph()->GetGraphName(), *Variable.Name.ToString());
+				FString SourcePath = FString::Printf(TEXT("LocalVariableDefault::%s|%s::Const"), *CurrentlyCompilingFunctionNode->GetContainedGraph()->GetGraphName(), *Variable.Name.ToString());
 				FRigVMOperand* TargetPtr = WorkData.PinPathToOperand->Find(TargetPath);
 				FRigVMOperand* SourcePtr = WorkData.PinPathToOperand->Find(SourcePath);
 				if (SourcePtr && TargetPtr)
@@ -2204,14 +2971,14 @@ bool URigVMCompiler::TraverseCallExtern(const FRigVMCallExternExprAST* InExpr, F
 	int32 CallExternInstructionIndex = INDEX_NONE;
 	const FRigVMCallstack Callstack = InExpr->GetProxy().GetCallstack();
 
-	if (WorkData.bSetupMemory)
+	if (WorkData.CompilerPhase == ERigVMCompilerPhase_SetupMemory)
 	{
 		if (!TraverseChildren(InExpr, WorkData))
 		{
 			return false;
 		}
 
-		SetupPropertyPathsForFoldedCopies(InExpr, WorkData);
+		WorkData.SetupPropertyPathsForFoldedCopies(InExpr);
 
 		if(WorkData.Settings.ASTSettings.bSetupTraits)
 		{
@@ -2558,6 +3325,16 @@ bool URigVMCompiler::TraverseCallExtern(const FRigVMCallExternExprAST* InExpr, F
 		// setup the instruction
 		const int32 FunctionIndex = WorkData.VM->AddRigVMFunction(Function->GetName());
 		check(FunctionIndex != INDEX_NONE);
+#if WITH_EDITOR
+		if (Function->Factory)
+		{
+			checkf(Function->Arguments.Num() <= Operands.Num(), TEXT("%s: invalid number of operands (%d) for dispatch '%s' - expected at least (%d)."), *GetPackage()->GetPathName(), Operands.Num(), *Function->GetName(), Function->Arguments.Num());
+		}
+		else
+		{
+			checkf(Function->Arguments.Num() == Operands.Num(), TEXT("%s: invalid number of operands (%d) for function '%s' - expected (%d)."), *GetPackage()->GetPathName(), Operands.Num(), *Function->GetName(), Function->Arguments.Num());
+		}
+#endif
 		const int32 CallExternByteIndex = WorkData.VM->GetByteCode().AddExecuteOp(FunctionIndex, Operands, StartPredicateIndex, PredicateCount);
 		CallExternInstructionIndex = WorkData.VM->GetByteCode().GetNumInstructions() - 1;
 		const FRigVMInstruction CallExternInstruction(ERigVMOpCode::Execute, CallExternInstructionIndex, CallExternByteIndex);
@@ -2574,7 +3351,7 @@ bool URigVMCompiler::TraverseCallExtern(const FRigVMCallExternExprAST* InExpr, F
 		}
 
 #if WITH_EDITORONLY_DATA
-		TArray<FRigVMOperand> InputsOperands, OutputOperands;
+		TArray<FRigVMOperand> InputOperands, OutputOperands;
 
 		for(const URigVMPin* InputPin : Node->GetPins())
 		{
@@ -2607,7 +3384,7 @@ bool URigVMCompiler::TraverseCallExtern(const FRigVMCallExternExprAST* InExpr, F
 						continue;
 					}
 
-					InputsOperands.Add(Operand);
+					InputOperands.Add(Operand);
 				}
 			}
 			else
@@ -2624,13 +3401,13 @@ bool URigVMCompiler::TraverseCallExtern(const FRigVMCallExternExprAST* InExpr, F
 					continue;
 				}
 
-				InputsOperands.Add(Operand);
+				InputOperands.Add(Operand);
 			}
 		}
 
 		WorkData.VM->GetByteCode().SetOperandsForInstruction(
 			CallExternInstruction,
-			FRigVMOperandArray(InputsOperands.GetData(), InputsOperands.Num()),
+			FRigVMOperandArray(InputOperands.GetData(), InputOperands.Num()),
 			FRigVMOperandArray(OutputOperands.GetData(), OutputOperands.Num()));
 
 #endif
@@ -2737,256 +3514,117 @@ bool URigVMCompiler::TraverseInlineFunction(const FRigVMInlineFunctionExprAST* I
 	{
 		return false;
 	}
+
+	const FRigVMGraphFunctionHeader& Header = FunctionReferenceNode->GetReferencedFunctionHeader(); 
 	
-	int32 InstructionIndexStart = INDEX_NONE;
-	int32 InstructionIndexEnd = INDEX_NONE;
-	int32 BranchIndexStart = INDEX_NONE;
-	
-	FString FunctionHash = FunctionReferenceNode->GetReferencedFunctionHeader().GetHash();
+	FString FunctionHash = Header.GetHash();
 	if (!CompiledFunctions.Contains(FunctionHash))
 	{
 		return true;
 	}
-	const FRigVMFunctionCompilationData* FunctionCompilationData = CompiledFunctions.FindChecked(FunctionHash);
+
+	// if we are considering callables, and we've seen this function before, let's take a shortcut.
+	if(CVarRigVMCompileFunctionsToCallables->GetBool())
+	{
+		const uint32 CallableHash = CompiledFunctions.FindChecked(FunctionHash).CompilationData->Hash;
+		if (TraverseCallableFunction(InExpr, CallableHash, WorkData))
+		{
+			return true;
+		}
+	}
+
+	const FRigVMFunctionCompilationData* FunctionCompilationData = CompiledFunctions.FindChecked(FunctionHash).CompilationData;
 	const FRigVMByteCode& FunctionByteCode = FunctionCompilationData->ByteCode;
+	const uint32 UniqueHashForReferenceNode = (uint32)((PTRINT)FunctionReferenceNode); 
 
 	// Bytecode to be inlined should never be aligned
 	checkf(!FunctionByteCode.bByteCodeIsAligned, TEXT("Trying to inline aligned function bytecode %s in package %s"), *FunctionReferenceNode->GetFunctionIdentifier().GetLibraryNodePath(), *GetPackage()->GetPathName());
+
+	FRigVMByteCode& ByteCode = WorkData.VM->GetByteCode();
 	
-	if (WorkData.bSetupMemory)
+	// let's make sure we already have all required callables
+	int32 FirstFunctionInstructionIndex = 0;
+	for (int32 CallableIndex = 0; CallableIndex < FunctionByteCode.NumCallables(); CallableIndex++)
 	{
-		if (!TraverseChildren(InExpr, WorkData))
+		const FRigVMCallableInfo* OldCallable = FunctionByteCode.GetCallable(CallableIndex);
+		check(OldCallable);
+		const FRigVMCallableInfo* NewCallable = ByteCode.FindCallable(OldCallable->Hash);
+		if (!ensure(NewCallable != nullptr))
 		{
+			UE_LOG(LogRigVMDeveloper, Error, TEXT("'%s': Unexpected missing callable '%s' when inlining function '%s'."),
+				*GetPackage()->GetPathName(), *OldCallable->Name.ToString(), *Header.LibraryPointer.GetLibraryNodePath());
 			return false;
 		}
+		FirstFunctionInstructionIndex = OldCallable->LastInstruction + 1;
+	}
 
-		// create a map of all of the trait setup lists (indices to trait properties)
-		TMap<int32, FRigVMOperand> LiteralValueToSetupTraitArg;
-		const FRigVMInstructionArray FunctionInstructions = FunctionByteCode.GetInstructions();
-		for(const FRigVMInstruction& Instruction : FunctionInstructions)
+	if (WorkData.CompilerPhase == ERigVMCompilerPhase_SetupMemory)
+	{
+		if (!TraverseChildren(InExpr, WorkData))
 		{
-			if(Instruction.OpCode == ERigVMOpCode::SetupTraits)
-			{
-				const FRigVMSetupTraitsOp& Op = FunctionByteCode.GetOpAt<FRigVMSetupTraitsOp>(Instruction);
-				checkf(Op.Arg.GetMemoryType() == ERigVMMemoryType::Literal, TEXT("Hit a non-literal arg for a trait, function %s in package %s"), *FunctionReferenceNode->GetFunctionIdentifier().GetLibraryNodePath(), *GetPackage()->GetPathName());
-				checkf(FunctionCompilationData->LiteralPropertyDescriptions.IsValidIndex(Op.Arg.GetRegisterIndex()), TEXT("Hit an invalid register for trait arg, function %s in package %s"), *FunctionReferenceNode->GetFunctionIdentifier().GetLibraryNodePath(), *GetPackage()->GetPathName());
-				const FString& DefaultValue = FunctionCompilationData->LiteralPropertyDescriptions[Op.Arg.GetRegisterIndex()].DefaultValue;
-				if(!DefaultValue.IsEmpty() && DefaultValue != TEXT("()"))
-				{
-					LiteralValueToSetupTraitArg.Add(Op.Arg.GetRegisterIndex(), FRigVMOperand());
-				}
-			}
+			return false;
 		}
 
-		// Add internal operands (not the ones represented by interface pins)
-		for (uint8 MemoryIndex=0; MemoryIndex< (uint8)ERigVMMemoryType::Invalid; ++MemoryIndex)
+		TArray<FRigVMOperand> InterfaceOperands;
+		for (const TPair<FName,FRigVMOperand>& Pair : FunctionCompilationData->InterfaceOperands)
 		{
-			ERigVMMemoryType MemoryType = (ERigVMMemoryType) MemoryIndex;
-			TArray<FRigVMFunctionCompilationPropertyDescription> Properties;
-			switch (MemoryType)
-			{
-			case ERigVMMemoryType::Work:
-				{
-					Properties = FunctionCompilationData->WorkPropertyDescriptions;
-					break;
-				}
-			case ERigVMMemoryType::Literal:
-				{
-					Properties = FunctionCompilationData->LiteralPropertyDescriptions;
-					break;
-				}
-			case ERigVMMemoryType::External:
-				{
-					Properties = FunctionCompilationData->ExternalPropertyDescriptions;
-					break;
-				}
-			case ERigVMMemoryType::Debug:
-				{
-					Properties = FunctionCompilationData->DebugPropertyDescriptions;
-					break;
-				}
-			}
-
-			// Find out how many properties we need to skip (function arguments and execute pins) 
-			int32 NumProperties = 0;
-			if (MemoryType == ERigVMMemoryType::Work)
-			{
-				for (const FRigVMGraphFunctionArgument& Argument : FunctionReferenceNode->GetReferencedFunctionHeader().Arguments)
-				{
-					if (Argument.IsCPPTypeObjectValid())
-					{
-						if (const UScriptStruct* ScriptStruct = Cast<UScriptStruct>(Argument.CPPTypeObject.Get()))
-						{
-							if (ScriptStruct->IsChildOf(FRigVMExecutePin::StaticStruct()))
-							{
-								continue;
-							}
-						}
-					}
-					if (!Properties.IsValidIndex(NumProperties))
-					{
-						continue;
-					}
-					const FString PropertyName = FRigVMPropertyDescription::SanitizeName(Argument.Name).ToString();
-					if (!Properties[NumProperties].Name.ToString().EndsWith(PropertyName))
-					{
-						continue;
-					}
-					NumProperties++;
-				}
-			}
-
-			auto FindOrAddProperty = [&WorkData, MemoryType, LiteralValueToSetupTraitArg]
-			(const FRigVMFunctionCompilationPropertyDescription& InProperty, const FString& InNewName, const bool bMakeUniqueProperty) -> FRigVMOperand
-			{
-				// Sharing / reusing memory / operands happens as per following contract:
-				// 1. properties are only shared if their CPP type matches
-				// 2. Literal / constant memory is only shared if the constant values match (and it is not a trait list)
-				// 3. Work state is only shared if it is not internal work state private to the instruction referring to it
-				// 4. Work state of type FRigVMInstructionSetExecuteState (bMakeUniqueProperty) is never shared either since it is work state private to a lazy branch.
-
-				FString NewName = InNewName;
-				int32 Suffix = 0;
-
-				// always share literals
-				if(MemoryType == ERigVMMemoryType::Literal && !bMakeUniqueProperty && WorkData.PropertyDescriptions.Contains(MemoryType))
-				{
-					const TArray<FRigVMPropertyDescription>& LiteralDescriptions = WorkData.PropertyDescriptions[MemoryType]; 
-					for(int32 Index = 0; Index < LiteralDescriptions.Num(); Index++)
-					{
-						if(!LiteralValueToSetupTraitArg.Contains(Index))
-						{
-							const FRigVMPropertyDescription& ExistingProperty = LiteralDescriptions[Index];
-							if(ExistingProperty.CPPType.Equals(InProperty.CPPType, ESearchCase::CaseSensitive))
-							{
-								// if the value is the same and this is not a trait setup list
-								if(ExistingProperty.DefaultValue.Equals(InProperty.DefaultValue, ESearchCase::CaseSensitive))
-								{
-									return FRigVMOperand(MemoryType, Index);
-								}
-							}
-						}
-					}
-				}
-
-				FRigVMOperand Operand = WorkData.FindProperty(MemoryType, *NewName);
-
-				// look for a free / unused operand
-				while(Operand.IsValid())
-				{
-					const FRigVMPropertyDescription& ExistingProperty = WorkData.PropertyDescriptions[MemoryType][Operand.GetRegisterIndex()];
-					if(!bMakeUniqueProperty &&
-						ExistingProperty.CPPType.Equals(InProperty.CPPType, ESearchCase::CaseSensitive))
-					{
-						if(MemoryType == ERigVMMemoryType::Literal)
-						{
-							// if the value is the same and this is not a trait setup list
-							if(ExistingProperty.DefaultValue.Equals(InProperty.DefaultValue) &&
-								!LiteralValueToSetupTraitArg.Contains(Operand.GetRegisterIndex()))
-							{
-								return Operand;
-							}
-						}
-						else
-						{
-							if(MemoryType == ERigVMMemoryType::Work &&
-								RigVMTypeUtils::IsArrayType(InProperty.CPPType))
-							{
-								// for now we don't share operands which have an array type
-								// since they may represent internal state. multiple occurrences of
-								// the same function cannot share internal work state.
-							}
-							else
-							{
-								return Operand;
-							}
-						}
-					}
-
-					NewName = FString::Printf(TEXT("%s_%02d"), *InNewName, ++Suffix);
-					Operand = WorkData.FindProperty(MemoryType, *NewName);
-				}
-						
-				// create a new operand as needed
-				return WorkData.AddProperty(MemoryType, *NewName, InProperty.CPPType, InProperty.CPPTypeObject.Get(), InProperty.DefaultValue);
-			};
-
-			// Add internal properties
-			int32 StartIndex = MemoryType == ERigVMMemoryType::Work ? NumProperties : 0;
-			for (int32 PropertyIndex = StartIndex; PropertyIndex < Properties.Num(); ++PropertyIndex)
-			{
-				static const FString RigVMInstructionSetExecuteStateName = FRigVMInstructionSetExecuteState::StaticStruct()->GetStructCPPName();
-				
-				const FRigVMFunctionCompilationPropertyDescription& Description = Properties[PropertyIndex];
-				FString NewName = Description.Name.ToString();
-				static const FString FunctionLibraryPrefix = TEXT("FunctionLibrary");
-
-				const UScriptStruct* ScriptStruct = Cast<UScriptStruct>(Description.CPPTypeObject.Get());
-
-				// instantiate function library specific work state as well as
-				// instruction set execute state - which is used for lazy blocks.
-				const bool bIsExecuteState = Description.CPPType.Equals(RigVMInstructionSetExecuteStateName);
-				const bool bIsTrait = ScriptStruct && ScriptStruct->IsChildOf(FRigVMTrait::StaticStruct());
-				if (NewName.StartsWith(FunctionLibraryPrefix) || bIsExecuteState || bIsTrait)
-				{
-					NewName = FString::Printf(TEXT("%s%s"), *FunctionReferenceNode->GetNodePath(), *NewName.RightChop(FunctionLibraryPrefix.Len()));
-					FRigVMPropertyDescription::SanitizeName(NewName);
-				}
-
-				const FRigVMOperand Operand = FindOrAddProperty(Description, NewName, bIsExecuteState || bIsTrait);
-				FRigVMCompilerWorkData::FFunctionRegisterData Data = {FunctionReferenceNode, MemoryType, PropertyIndex};
-				WorkData.FunctionRegisterToOperand.Add(Data, Operand);
-
-				if(MemoryType == ERigVMMemoryType::Literal && LiteralValueToSetupTraitArg.Contains(PropertyIndex))
-				{
-					checkf(LiteralValueToSetupTraitArg.Contains(PropertyIndex), TEXT("Trait property index out of range, function %s in package %s"), *FunctionReferenceNode->GetFunctionIdentifier().GetLibraryNodePath(), *GetPackage()->GetPathName());
-					LiteralValueToSetupTraitArg.FindChecked(PropertyIndex) = Operand;
-				}
-			}
+			InterfaceOperands.Add(Pair.Value);
 		}
-		// For trait setup lists we need to update the integer values (pointing to trait property indices)
-		for(const TPair<int32, FRigVMOperand>& Pair : LiteralValueToSetupTraitArg)
+		(void)WorkData.AddPropertiesForFunction(false, UniqueHashForReferenceNode, Header, FunctionCompilationData, FirstFunctionInstructionIndex, FunctionByteCode.NumInstructions - 1, InterfaceOperands, {});
+		WorkData.SetupPropertyPathsForFoldedCopies(InExpr);
+	}
+	else
+	{
+		if (!TraverseChildren(InExpr, WorkData))
 		{
-			const FRigVMOperand TraitIndicesOperand = Pair.Value;
-			const int32 LiteralPropertyIndex = TraitIndicesOperand.GetRegisterIndex();
-			FString OriginalTraitIndicesString = WorkData.PropertyDescriptions[ERigVMMemoryType::Literal][LiteralPropertyIndex].DefaultValue;
-			if(!OriginalTraitIndicesString.IsEmpty() && OriginalTraitIndicesString != TEXT("()"))
-			{
-				OriginalTraitIndicesString = OriginalTraitIndicesString.TrimChar(TEXT('('));
-				OriginalTraitIndicesString = OriginalTraitIndicesString.TrimChar(TEXT(')'));
-				FString PinPathRemaining = OriginalTraitIndicesString;
-				FString Left, Right;
-				TArray<FString> IndexStrings;
-				while(PinPathRemaining.Split(TEXT(","), &Left, &Right))
-				{
-					IndexStrings.Add(Left.TrimStartAndEnd());
-					Left.Empty();
-					PinPathRemaining = Right;
-				}
-				if (!Right.IsEmpty())
-				{
-					IndexStrings.Add(Right.TrimStartAndEnd());
-				}
+			return false;
+		}
+		
+		// Inline the bytecode from the function
+		const int32 InstructionIndexStart = ByteCode.GetNumInstructions();
+		FRigVMInstructionArray OldInstructions = ByteCode.GetInstructions();
+		ByteCode.InlineFunction(&FunctionByteCode, FirstFunctionInstructionIndex, FunctionByteCode.GetNumInstructions() - 1);
+		FRigVMCallstack FuncRefCallstack = InExpr->GetProxy().GetCallstack();
 
-				for(int32 PartIndex = 0; PartIndex < IndexStrings.Num(); PartIndex++)
-				{
-					const int32 OriginalIndex = FCString::Atoi(*IndexStrings[PartIndex]);
-					FRigVMCompilerWorkData::FFunctionRegisterData Data = {FunctionReferenceNode, ERigVMMemoryType::Work, OriginalIndex};
-					const FRigVMOperand& NewOperand = WorkData.FunctionRegisterToOperand.FindChecked(Data);
-					checkf(NewOperand.GetMemoryType() == ERigVMMemoryType::Work, TEXT("New work operand for trait is not work state, function %s in package %s"), *FunctionReferenceNode->GetFunctionIdentifier().GetLibraryNodePath(), *GetPackage()->GetPathName());
-					IndexStrings[PartIndex] = FString::FromInt(NewOperand.GetRegisterIndex());
-				}
+		auto RemapInstructionIndex = [InstructionIndexStart, FirstFunctionInstructionIndex](const int32& InOldInstructionIndex) -> int32
+		{
+			return InOldInstructionIndex - FirstFunctionInstructionIndex + InstructionIndexStart;				
+		};
 
-				IndexStrings.Remove(FString());
-				WorkData.PropertyDescriptions[ERigVMMemoryType::Literal][LiteralPropertyIndex].DefaultValue =
-					FString::Printf(TEXT("(%s)"), *FString::Join(IndexStrings, TEXT(",")));
+		// Import branch infos, with the proper instruction indices
+		TMap<int32,int32> BranchIndexMap;
+		for(const FRigVMBranchInfo& BranchInfo : FunctionByteCode.BranchInfos)
+		{
+			// skip the branch infos relating to callables
+			if (BranchInfo.InstructionIndex < FirstFunctionInstructionIndex ||
+				BranchInfo.FirstInstruction < FirstFunctionInstructionIndex ||
+				BranchInfo.LastInstruction < FirstFunctionInstructionIndex)
+			{
+				continue;
 			}
+
+			const int32 BranchInfoIndex = ByteCode.AddBranchInfo(BranchInfo);
+			ByteCode.BranchInfos[BranchInfoIndex].Index = BranchInfoIndex;
+			ByteCode.BranchInfos[BranchInfoIndex].InstructionIndex = RemapInstructionIndex(BranchInfo.InstructionIndex);
+			ByteCode.BranchInfos[BranchInfoIndex].FirstInstruction = RemapInstructionIndex(BranchInfo.FirstInstruction);
+			ByteCode.BranchInfos[BranchInfoIndex].LastInstruction = RemapInstructionIndex(BranchInfo.LastInstruction);
+			BranchIndexMap.Add(BranchInfo.Index, BranchInfoIndex);
 		}
 
-		SetupPropertyPathsForFoldedCopies(InExpr, WorkData);
-	}
-	else
-	{
+		auto RemapBranchIndex = [&BranchIndexMap](const int32& InOldBranchIndex) -> int32
+		{
+			return BranchIndexMap.FindChecked(InOldBranchIndex);
+		};
+		
+		auto RemapCallableIndex = [&FunctionByteCode, &ByteCode](const int32& InOldCallableIndex) -> int32
+		{
+			const FRigVMCallableInfo* OldCallable = FunctionByteCode.GetCallable(InOldCallableIndex);
+			check(OldCallable);
+			const FRigVMCallableInfo* NewCallable = ByteCode.FindCallable(OldCallable->Hash);
+			check(NewCallable);
+			return NewCallable->Index;
+		};
+
 		TArray<FRigVMOperand> InterfaceOperands;
 		TArray<FString> OperandsPinNames;
 		// Find operands related to the function's interface
@@ -3028,244 +3666,161 @@ bool URigVMCompiler::TraverseInlineFunction(const FRigVMInlineFunctionExprAST* I
 				InterfaceOperands.Last().RegisterOffset = static_cast<uint16>(RegisterOffsetIndex);
 			}
 		}
-
-		if (!TraverseChildren(InExpr, WorkData))
-		{
-			return false;
-		}
-
-		// Inline the bytecode from the function
-		FRigVMByteCode& ByteCode = WorkData.VM->GetByteCode();
-		InstructionIndexStart = ByteCode.GetNumInstructions();
-		FRigVMInstructionArray OldInstructions = ByteCode.GetInstructions();
-		ByteCode.InlineFunction(&FunctionByteCode, InterfaceOperands);
-		InstructionIndexEnd = ByteCode.GetNumInstructions() - 1;
-		FRigVMCallstack FuncRefCallstack = InExpr->GetProxy().GetCallstack();
-
-		// Import branch infos, with the proper instruction indices
-		for(const FRigVMBranchInfo& BranchInfo : FunctionByteCode.BranchInfos)
-		{
-			const int32 BranchInfoIndex = ByteCode.AddBranchInfo(BranchInfo);
-			if(BranchIndexStart == INDEX_NONE)
-			{
-				BranchIndexStart = BranchInfoIndex;
-			}
-			ByteCode.BranchInfos[BranchInfoIndex].InstructionIndex += InstructionIndexStart;
-			ByteCode.BranchInfos[BranchInfoIndex].FirstInstruction += InstructionIndexStart;
-			ByteCode.BranchInfos[BranchInfoIndex].LastInstruction += InstructionIndexStart;
-		}
-
-		TMap<FRigVMOperand, FRigVMOperand> FunctionOperandToNewOperand;
-
-		// For each instruction, substitute the operand for the one used in the current bytecode
-		const FRigVMInstructionArray FunctionInstructions = FunctionByteCode.GetInstructions();
-		FRigVMInstructionArray Instructions = ByteCode.GetInstructions();
-		for (int32 i=InstructionIndexStart; i<=InstructionIndexEnd; ++i)
+		
+		auto RemapOperand = [this, &WorkData, UniqueHashForReferenceNode, &FunctionReferenceNode, FunctionCompilationData, &InterfaceOperands, &OperandsPinNames](const FRigVMOperand& InOldOperand) -> FRigVMOperand
 		{
-			if (!Instructions.IsValidIndex(i))
-			{
-				return false;
-			}
+			FRigVMOperand InterfaceOperand;
+			FRigVMOperand NewOperand = InOldOperand;
 
-			const int32 FunctionInstructionIndex = i - InstructionIndexStart;
-			
-			const FRigVMInstruction& Instruction = Instructions[i];
-			const FRigVMOperandArray OperandArray = ByteCode.GetOperandsForOp(Instruction);
-			const FRigVMOperandArray FunctionOperandArray = FunctionByteCode.GetOperandsForOp(FunctionInstructions[FunctionInstructionIndex]);
-			checkf(FunctionOperandArray.Num() == OperandArray.Num(), TEXT("Hit a function with a mismatch of function operands, function %s in package %s"), *FunctionReferenceNode->GetFunctionIdentifier().GetLibraryNodePath(), *GetPackage()->GetPathName());
-			
-			for (int32 j=0; j<OperandArray.Num(); ++j)
+			// Remap the variable: find the operand index of the outer variable
+			if (NewOperand.GetMemoryType() == ERigVMMemoryType::External)
 			{
-				FRigVMOperand* Operand = const_cast<FRigVMOperand*>(&OperandArray[j]);
-
-				const FRigVMOperand OriginalOperand = *Operand;
-				checkf(OriginalOperand == FunctionOperandArray[j], TEXT("Hit a function with an invalid / non-matching function operand, function %s in package %s"), *FunctionReferenceNode->GetFunctionIdentifier().GetLibraryNodePath(), *GetPackage()->GetPathName());
-				
-				FRigVMOperand InterfaceOperand;
-
-				// Remap the variable: find the operand index of the outer variable
-				if (Operand->GetMemoryType() == ERigVMMemoryType::External)
+				const FName& InnerVariableName = FunctionCompilationData->ExternalRegisterIndexToVariable[NewOperand.GetRegisterIndex()];
+				FString OuterVariablePath = InnerVariableName.ToString();
+				const FString VariableRemapped = FunctionReferenceNode->GetOuterVariablePath(*OuterVariablePath);
+				if (!VariableRemapped.IsEmpty())
 				{
-					const FName& InnerVariableName = FunctionCompilationData->ExternalRegisterIndexToVariable[Operand->GetRegisterIndex()];
-					FString OuterVariablePath = InnerVariableName.ToString();
-					const FString VariableRemapped = FunctionReferenceNode->GetOuterVariablePath(InnerVariableName);
-					if (!VariableRemapped.IsEmpty())
-					{
-						OuterVariablePath = VariableRemapped;
-					}
-					else
-					{
-						FRigVMGraphFunctionIdentifier LibraryPointer = FunctionReferenceNode->GetFunctionIdentifier();
-						const FString& LibraryPackagePath = LibraryPointer.GetNodeSoftPath().GetLongPackageName();
-						ensureMsgf(!FunctionReferenceNode->RequiresVariableRemapping(),
-							TEXT("Could not find variable %s in function reference %s variable map, in package %s\n. Library package is %s\n"),
-							*InnerVariableName.ToString(),
-							*FunctionReferenceNode->GetNodePath(),
-							*GetPackage()->GetPathName(),
-							*LibraryPackagePath);
-					}
+					OuterVariablePath = VariableRemapped;
+				}
+				else
+				{
+					FRigVMGraphFunctionIdentifier LibraryPointer = FunctionReferenceNode->GetFunctionIdentifier();
+					const FString& LibraryPackagePath = LibraryPointer.GetNodeSoftPath().GetLongPackageName();
+					ensureMsgf(!FunctionReferenceNode->RequiresVariableRemapping(),
+						TEXT("Could not find variable %s in function reference %s variable map, in package %s\n. Library package is %s\n"),
+						*InnerVariableName.ToString(),
+						*FunctionReferenceNode->GetNodePath(),
+						*GetPackage()->GetPathName(),
+						*LibraryPackagePath);
+				}
 
-					// for now variable paths are not supported for mapping.
-					if (OuterVariablePath.Contains(TEXT(".")))
-					{
-						UE_LOG(
-							LogRigVMDeveloper,
-							Error,
-							TEXT("Variable '%s' on function '%s' is mapped to path '%s'. Property paths on variable mappings are not yet supported."),
-							*InnerVariableName.ToString(),
-							*FunctionReferenceNode->GetReferencedFunctionHeader().Name.ToString(),
-							*OuterVariablePath);
-					}
-					else
-					{
-						const FString OuterOperandKey = FString::Printf(TEXT("Variable::%s"), *OuterVariablePath);
-						checkf(WorkData.PinPathToOperand->Contains(OuterOperandKey), TEXT("hit a non-existing operand '%s', function %s in package %s"), *OuterOperandKey, *FunctionReferenceNode->GetFunctionIdentifier().GetLibraryNodePath(), *GetPackage()->GetPathName());
+				// for now variable paths are not supported for mapping.
+				if (OuterVariablePath.Contains(TEXT(".")))
+				{
+					UE_LOG(
+						LogRigVMDeveloper,
+						Error,
+						TEXT("Variable '%s' on function '%s' is mapped to path '%s'. Property paths on variable mappings are not yet supported."),
+						*InnerVariableName.ToString(),
+						*FunctionReferenceNode->GetReferencedFunctionHeader().Name.ToString(),
+						*OuterVariablePath);
+				}
+				else
+				{
+					const FString OuterOperandKey = FString::Printf(TEXT("Variable::%s"), *OuterVariablePath);
+					checkf(WorkData.PinPathToOperand->Contains(OuterOperandKey), TEXT("hit a non-existing operand '%s', function %s in package %s"), *OuterOperandKey, *FunctionReferenceNode->GetReferencedFunctionHeader().Name.ToString(), *GetPackage()->GetPathName());
 
-						const FRigVMOperand& OuterOperand = WorkData.PinPathToOperand->FindChecked(OuterOperandKey);
-						Operand->MemoryType = OuterOperand.MemoryType;
-						Operand->RegisterIndex = OuterOperand.RegisterIndex;
-					}
+					const FRigVMOperand& OuterOperand = WorkData.PinPathToOperand->FindChecked(OuterOperandKey);
+					NewOperand.MemoryType = OuterOperand.MemoryType;
+					NewOperand.RegisterIndex = OuterOperand.RegisterIndex;
 				}
+			}
+			else
+			{
 				// If Operand is an interface pin: replace the index and memory type
-				else if (const int32 FunctionInterfaceParameterIndex = GetOperandFunctionInterfaceParameterIndex(OperandsPinNames, FunctionCompilationData, *Operand); FunctionInterfaceParameterIndex != INDEX_NONE)
+				if (const int32 FunctionInterfaceParameterIndex = GetOperandFunctionInterfaceParameterIndex(OperandsPinNames, FunctionCompilationData, NewOperand); FunctionInterfaceParameterIndex != INDEX_NONE)
 				{
 					InterfaceOperand = InterfaceOperands[FunctionInterfaceParameterIndex];
-					Operand->MemoryType = InterfaceOperand.MemoryType;
-					Operand->RegisterIndex = InterfaceOperand.RegisterIndex;
+					NewOperand.MemoryType = InterfaceOperand.MemoryType;
+					NewOperand.RegisterIndex = InterfaceOperand.RegisterIndex;
 				}
 				else
 				{
 					// Operand is internal
 					// Replace with added Operand
-					FRigVMCompilerWorkData::FFunctionRegisterData Data = {FunctionReferenceNode, Operand->GetMemoryType(), Operand->GetRegisterIndex()};
+					const FRigVMCompilerWorkData::FFunctionRegisterData Data = {UniqueHashForReferenceNode, NewOperand.GetMemoryType(), NewOperand.GetRegisterIndex()};
 					checkf(WorkData.FunctionRegisterToOperand.Contains(Data), TEXT("hit a missing function operand, function %s in package %s"), *FunctionReferenceNode->GetFunctionIdentifier().GetLibraryNodePath(), *GetPackage()->GetPathName());
-					FRigVMOperand NewOperand = WorkData.FunctionRegisterToOperand.FindChecked(Data);
-					Operand->MemoryType = NewOperand.MemoryType;
-					Operand->RegisterIndex = NewOperand.RegisterIndex;
+					FRigVMOperand RemappedOperand = WorkData.FunctionRegisterToOperand.FindChecked(Data);
+					NewOperand.MemoryType = RemappedOperand.MemoryType;
+					NewOperand.RegisterIndex = RemappedOperand.RegisterIndex;
 				}
+			}
 
-				// For all operands, check to see if we need to add a property path
-				if (Operand->GetRegisterOffset() != INDEX_NONE || InterfaceOperand.GetRegisterOffset() != INDEX_NONE)
+			// For all operands, check to see if we need to add a property path
+			if (NewOperand.GetRegisterOffset() != INDEX_NONE || InterfaceOperand.GetRegisterOffset() != INDEX_NONE)
+			{
+				if (InterfaceOperand.GetRegisterOffset() != INDEX_NONE && NewOperand.GetRegisterOffset() == INDEX_NONE)
 				{
-					if (InterfaceOperand.GetRegisterOffset() != INDEX_NONE && Operand->GetRegisterOffset() == INDEX_NONE)
-					{
-						Operand->RegisterOffset = IntCastChecked<uint16>(InterfaceOperand.GetRegisterOffset());
-					}
-					else
+					NewOperand.RegisterOffset = IntCastChecked<uint16>(InterfaceOperand.GetRegisterOffset());
+				}
+				else
+				{
+					FRigVMFunctionCompilationPropertyPath DescriptionFromFunction;
+					switch (InOldOperand.MemoryType)
 					{
-						FRigVMFunctionCompilationPropertyPath DescriptionFromFunction;
-						switch (OriginalOperand.MemoryType)
+						case ERigVMMemoryType::Work:
 						{
-							case ERigVMMemoryType::Work:
-							{
-								DescriptionFromFunction = FunctionCompilationData->WorkPropertyPathDescriptions[Operand->GetRegisterOffset()];
-								break;
-							}
-							case ERigVMMemoryType::Literal:
-							{
-								DescriptionFromFunction = FunctionCompilationData->LiteralPropertyPathDescriptions[Operand->GetRegisterOffset()];
-								break;
-							}
-							case ERigVMMemoryType::External:
-							{
-								DescriptionFromFunction = FunctionCompilationData->ExternalPropertyPathDescriptions[Operand->GetRegisterOffset()];
-								break;
-							}
-							case ERigVMMemoryType::Debug:
-							{
-								DescriptionFromFunction = FunctionCompilationData->DebugPropertyPathDescriptions[Operand->GetRegisterOffset()];
-								break;
-							}
-							default:
-							{
-								checkNoEntry();
-							}
+							DescriptionFromFunction = FunctionCompilationData->WorkPropertyPathDescriptions[NewOperand.GetRegisterOffset()];
+							break;
 						}
-
-						if (InterfaceOperand.GetRegisterOffset() == INDEX_NONE)
+						case ERigVMMemoryType::Literal:
 						{
-							Operand->RegisterOffset = IntCastChecked<uint16>(WorkData.FindOrAddPropertyPath(*Operand, DescriptionFromFunction.HeadCPPType, DescriptionFromFunction.SegmentPath));
+							DescriptionFromFunction = FunctionCompilationData->LiteralPropertyPathDescriptions[NewOperand.GetRegisterOffset()];
+							break;
 						}
-						else
+						case ERigVMMemoryType::External:
 						{
-							// in this case both the external operand and the operand within the function have a register offset.
-							// they need to be combined.
-							const FRigVMPropertyPathDescription& Description = WorkData.PropertyPathDescriptions.FindOrAdd(InterfaceOperand.GetMemoryType())[InterfaceOperand.GetRegisterOffset()];
-							const FString CombinedSegmentPath = URigVMPin::JoinPinPath(Description.SegmentPath, DescriptionFromFunction.SegmentPath);
-							Operand->RegisterOffset = IntCastChecked<uint16>(WorkData.FindOrAddPropertyPath(*Operand, Description.HeadCPPType, CombinedSegmentPath));
+							DescriptionFromFunction = FunctionCompilationData->ExternalPropertyPathDescriptions[NewOperand.GetRegisterOffset()];
+							break;
+						}
+						case ERigVMMemoryType::Debug:
+						{
+							DescriptionFromFunction = FunctionCompilationData->DebugPropertyPathDescriptions[NewOperand.GetRegisterOffset()];
+							break;
+						}
+						default:
+						{
+							checkNoEntry();
 						}
 					}
-				}
-
-				// Store correspondence between Function operand and New operands
-				checkf(!FunctionOperandToNewOperand.Contains(OriginalOperand) || FunctionOperandToNewOperand.FindChecked(OriginalOperand) == *Operand, TEXT("hit a correspondence issue with a function operand, function %s in package %s"), *FunctionReferenceNode->GetFunctionIdentifier().GetLibraryNodePath(), *GetPackage()->GetPathName());
-				FunctionOperandToNewOperand.Add(OriginalOperand, *Operand);
-			}
-
-			if (Instruction.OpCode == ERigVMOpCode::Execute)
-			{
-				FRigVMExecuteOp& Op = ByteCode.GetOpAt<FRigVMExecuteOp>(Instruction);
-				const int32 FunctionIndex = WorkData.VM->AddRigVMFunction(FunctionCompilationData->FunctionNames[Op.FunctionIndex].ToString());
-				Op.FunctionIndex = IntCastChecked<uint16>(FunctionIndex);
-			}
-			else if (Instruction.OpCode == ERigVMOpCode::JumpToBranch)
-			{
-				FRigVMJumpToBranchOp& Op = ByteCode.GetOpAt<FRigVMJumpToBranchOp>(Instruction);
-				Op.FirstBranchInfoIndex = Op.FirstBranchInfoIndex + BranchIndexStart;
-			}
-			else if (Instruction.OpCode == ERigVMOpCode::RunInstructions)
-			{
-				FRigVMRunInstructionsOp& Op = ByteCode.GetOpAt<FRigVMRunInstructionsOp>(Instruction);
-				Op.StartInstruction = Op.StartInstruction + InstructionIndexStart;
-				Op.EndInstruction = Op.EndInstruction + InstructionIndexStart;
-			}
 
-			if (WorkData.Settings.SetupNodeInstructionIndex)
-			{
-				if (const TArray<TWeakObjectPtr<UObject>>* Callstack = FunctionByteCode.GetCallstackForInstruction(FunctionInstructionIndex))
-				{
-					if (Callstack->Num() > 1)
+					if (InterfaceOperand.GetRegisterOffset() == INDEX_NONE)
 					{
-						FRigVMCallstack InstructionCallstack = FuncRefCallstack;
-						InstructionCallstack.Stack.Append(&(*Callstack)[1], Callstack->Num()-1);
-						ByteCode.SetSubject(i, InstructionCallstack.GetCallPath(), InstructionCallstack.GetStack());
+						NewOperand.RegisterOffset = IntCastChecked<uint16>(WorkData.FindOrAddPropertyPath(NewOperand, DescriptionFromFunction.HeadCPPType, DescriptionFromFunction.SegmentPath));
+					}
+					else
+					{
+						// in this case both the external operand and the operand within the function have a register offset.
+						// they need to be combined.
+						const FRigVMPropertyPathDescription& Description = WorkData.PropertyPathDescriptions.FindOrAdd(InterfaceOperand.GetMemoryType())[InterfaceOperand.GetRegisterOffset()];
+						const FString CombinedSegmentPath = URigVMPin::JoinPinPath(Description.SegmentPath, DescriptionFromFunction.SegmentPath);
+						NewOperand.RegisterOffset = IntCastChecked<uint16>(WorkData.FindOrAddPropertyPath(NewOperand, Description.HeadCPPType, CombinedSegmentPath));
 					}
-				}
-
-				// also store the instruction for the subject so that profiling can determine the cost of functions.
-				ByteCode.AddInstructionForSubject(FunctionReferenceNode, i);
-			}
-
-			// remap the input and output operands per instruction
-			TArray<FRigVMOperand> InputOperands, OutputOperands;
-			if (FunctionByteCode.InputOperandsPerInstruction.IsValidIndex(FunctionInstructionIndex))
-			{
-				for (const FRigVMOperand& Operand : FunctionByteCode.InputOperandsPerInstruction[FunctionInstructionIndex])
-				{
-					checkf(FunctionOperandToNewOperand.Contains(Operand), TEXT("hit a missing function input operand during remapping, function %s in package %s"), *FunctionReferenceNode->GetFunctionIdentifier().GetLibraryNodePath(), *GetPackage()->GetPathName());
-					InputOperands.Add(FunctionOperandToNewOperand.FindChecked(Operand));
-				}
-			}
-			if (FunctionByteCode.OutputOperandsPerInstruction.IsValidIndex(FunctionInstructionIndex))
-			{
-				for (const FRigVMOperand& Operand : FunctionByteCode.OutputOperandsPerInstruction[FunctionInstructionIndex])
-				{
-					checkf(FunctionOperandToNewOperand.Contains(Operand), TEXT("hit a missing function output operand during remapping, function %s in package %s"), *FunctionReferenceNode->GetFunctionIdentifier().GetLibraryNodePath(), *GetPackage()->GetPathName());
-					OutputOperands.Add(FunctionOperandToNewOperand.FindChecked(Operand));
 				}
 			}
-			ByteCode.SetOperandsForInstruction(Instruction, InputOperands, OutputOperands);
-		}
+		
+			return NewOperand;
+		};
 
+		// For each instruction, substitute the operand for the one used in the current bytecode
+		const FRigVMInstructionArray FunctionInstructions = FunctionByteCode.GetInstructions();
+		FRigVMInstructionArray Instructions = ByteCode.GetInstructions();
+		for (int32 FunctionInstructionIndex = FirstFunctionInstructionIndex; FunctionInstructionIndex < FunctionInstructions.Num(); FunctionInstructionIndex++)
+		{
+			WorkData.RemapInstructionAt(
+				FunctionInstructions[FunctionInstructionIndex],
+				FunctionReferenceNode->GetFunctionIdentifier().GetLibraryNodePath(),
+				FunctionCompilationData,
+				ByteCode,
+				Instructions,
+				FunctionReferenceNode,
+				RemapInstructionIndex,
+				RemapOperand,
+				RemapBranchIndex,
+				RemapCallableIndex);
+		}
+
+		const TMap<FRigVMOperand, FRigVMOperand>& FunctionOperandToNewOperand = WorkData.OperandMapPerFunction.FindChecked(UniqueHashForReferenceNode);
+		
 		// Add all pin paths to operand from the function
 		TMap<FRigVMOperand, TArray<FString>> OperandToPinPath;
 		for (const TPair<FString, FRigVMOperand>& PinPathToOperand : FunctionCompilationData->Operands)
 		{
-			if (FRigVMOperand* NewOperand = FunctionOperandToNewOperand.Find(PinPathToOperand.Value))
+			if (const FRigVMOperand* NewOperand = FunctionOperandToNewOperand.Find(PinPathToOperand.Value))
 			{
 				// avoid copying the variable pin paths over to the outer VM since
 				// the variables from within the function asset no longer exist and
-				// will have been remapped to the outer VM's external variables.
+				// will have been remapped to the outer VM's external variables.1
 				static const FString VariablePrefix = TEXT("Variable::");
 				if(PinPathToOperand.Key.StartsWith(VariablePrefix, ESearchCase::CaseSensitive))
 				{
@@ -3282,6 +3837,130 @@ bool URigVMCompiler::TraverseInlineFunction(const FRigVMInlineFunctionExprAST* I
 	return true;
 }
 
+bool URigVMCompiler::TraverseCallableFunction(const FRigVMInlineFunctionExprAST* InExpr, const uint32 InCallableHash, FRigVMCompilerWorkData& WorkData)
+{
+	FRigVMByteCode& ByteCode = WorkData.VM->GetByteCode();
+	const FRigVMCallableInfo* Callable = ByteCode.FindCallable(InCallableHash);
+	if (!Callable)
+	{
+		return false;
+	}
+
+	URigVMNode* Node = InExpr->GetNode();
+	URigVMFunctionReferenceNode* FunctionReferenceNode = Cast<URigVMFunctionReferenceNode>(Node);
+	if(!ValidateNode(WorkData.Settings, FunctionReferenceNode, false))
+	{
+		return false;
+	}
+
+	const FRigVMGraphFunctionHeader& Header = FunctionReferenceNode->GetReferencedFunctionHeader(); 
+	const FString FunctionHash = Header.GetHash();
+	if (!CompiledFunctions.Contains(FunctionHash))
+	{
+		return true;
+	}
+
+	const FRigVMFunctionCompilationData* FunctionCompilationData = CompiledFunctions.FindChecked(FunctionHash).CompilationData;
+
+	// not all function support callable instructions
+	if (!FunctionCompilationData->SupportsCallable())
+	{
+		return false;
+	}
+
+	const FRigVMCallstack Callstack = InExpr->GetProxy().GetCallstack();
+	
+	if (WorkData.CompilerPhase == ERigVMCompilerPhase_SetupMemory)
+	{
+		if (!TraverseChildren(InExpr, WorkData))
+		{
+			return false;
+		}
+	}
+	else
+	{
+		if (!TraverseChildren(InExpr, WorkData))
+		{
+			return false;
+		}
+
+		TArray<FRigVMOperand> CallableOperands;
+		CallableOperands.SetNum(Callable->Arguments.Num());
+#if WITH_EDITOR
+		TArray<FRigVMOperand> InputOperands, OutputOperands;
+#endif
+
+		for (int32 ChildIndex = 0; ChildIndex < InExpr->NumChildren(); ChildIndex++)
+		{
+			const FRigVMExprAST* ChildExpr = InExpr->ChildAt(ChildIndex);
+			const FRigVMVarExprAST* SourceVarExpr = nullptr;
+			if (ChildExpr->IsVar())
+			{
+				SourceVarExpr = ChildExpr->To<FRigVMVarExprAST>();
+			}
+			else if (ChildExpr->GetType() == FRigVMExprAST::EType::CachedValue)
+			{
+				SourceVarExpr = ChildExpr->To<FRigVMCachedValueExprAST>()->GetVarExpr();
+			}
+			else
+			{
+				// the AST is corrupt. we only expect cached values or var expressions as children here.
+				checkNoEntry();
+			}
+
+			if (SourceVarExpr->IsExecuteContext())
+			{
+				continue;
+			}
+			
+			SourceVarExpr = GetSourceVarExpr(SourceVarExpr);
+			
+			check(WorkData.ExprToOperand.Contains(SourceVarExpr));
+			const FRigVMOperand SourceOperand = WorkData.ExprToOperand.FindChecked(SourceVarExpr);
+			
+			const FName PinName = InExpr->GetPinNameForChildIndex(ChildIndex);
+			check(!PinName.IsNone());
+
+			const int32 ArgumentIndex = Callable->GetArgumentIndex(PinName);
+			check(Callable->Arguments.IsValidIndex(ArgumentIndex));
+
+			const FRigVMCallableArgument& Argument = Callable->Arguments[ArgumentIndex];
+			check(Argument.InterfaceOperand.IsValid());
+			check(Argument.ForwardedOperand.IsValid());
+			check(SourceOperand.IsValid());
+
+			CallableOperands[ArgumentIndex] = SourceOperand;
+
+#if WITH_EDITOR
+			if (Argument.IsInput())
+			{
+				InputOperands.Add(SourceOperand);
+			}
+			if (Argument.IsOutput())
+			{
+				OutputOperands.Add(SourceOperand);
+			}
+#endif
+		}
+		
+		const int32 InvokeCallableByteIndex = ByteCode.AddInvokeCallableOp(Callable->Hash, CallableOperands);
+		const int32 InvokeCallableInstructionIndex = ByteCode.GetNumInstructions() - 1;
+		const FRigVMInstruction InvokeCallableInstruction(ERigVMOpCode::InvokeCallable, InvokeCallableInstructionIndex, InvokeCallableByteIndex);
+
+#if WITH_EDITOR
+		WorkData.VM->GetByteCode().SetOperandsForInstruction(
+			InvokeCallableInstruction,
+			FRigVMOperandArray(InputOperands.GetData(), InputOperands.Num()),
+			FRigVMOperandArray(OutputOperands.GetData(), OutputOperands.Num()));
+#endif
+		if (WorkData.Settings.SetupNodeInstructionIndex)
+		{
+			WorkData.VM->GetByteCode().SetSubject(InvokeCallableInstructionIndex, Callstack.GetCallPath(), Callstack.GetStack());
+		}
+	}
+	return true;
+}
+
 bool URigVMCompiler::TraverseNoOp(const FRigVMNoOpExprAST* InExpr, FRigVMCompilerWorkData& WorkData)
 {
 	if (const URigVMVariableNode* VariableNode = Cast<URigVMVariableNode>(InExpr->GetNode()))
@@ -3372,7 +4051,7 @@ bool URigVMCompiler::TraverseNoOp(const FRigVMNoOpExprAST* InExpr, FRigVMCompile
 						}
 					}
 
-					if (WorkData.bSetupMemory)
+					if (WorkData.CompilerPhase == ERigVMCompilerPhase_SetupMemory)
 					{
 						if (!RootDefaultValueOperand.IsValid() && VariableOperand.IsValid())
 						{
@@ -3417,7 +4096,7 @@ bool URigVMCompiler::TraverseVar(const FRigVMVarExprAST* InExpr, FRigVMCompilerW
 		return false;
 	}
 
-	if (WorkData.bSetupMemory)
+	if (WorkData.CompilerPhase == ERigVMCompilerPhase_SetupMemory)
 	{
 		FindOrAddRegister(InExpr, WorkData);
 	}
@@ -3506,7 +4185,7 @@ bool URigVMCompiler::TraverseAssign(const FRigVMAssignExprAST* InExpr, FRigVMCom
 
 	// We cannot wait for the target var expr to be created, we need it now to setup the register offset
 	const FRigVMVarExprAST* VarExpr = GetSourceVarExpr(TargetExpr);
-	if (WorkData.bSetupMemory)
+	if (WorkData.CompilerPhase == ERigVMCompilerPhase_SetupMemory)
 	{
 		if (!WorkData.ExprToOperand.Contains(VarExpr))
 		{
@@ -3590,7 +4269,7 @@ bool URigVMCompiler::TraverseAssign(const FRigVMAssignExprAST* InExpr, FRigVMCom
 		Local::SetupRegisterOffset(WorkData.VM, Link, InExpr->GetTargetPin(), Target, TargetExpr, false, WorkData);
 	}
 	
-	if (!WorkData.bSetupMemory)
+	if (WorkData.CompilerPhase == ERigVMCompilerPhase_BuildInstructions)
 	{
 		FRigVMCopyOp CopyOp = WorkData.VM->GetCopyOpForOperands(Source, Target);
 		if(CopyOp.IsValid())
@@ -3615,7 +4294,7 @@ bool URigVMCompiler::TraverseCachedValue(const FRigVMCachedValueExprAST* InExpr,
 bool URigVMCompiler::TraverseExit(const FRigVMExitExprAST* InExpr, FRigVMCompilerWorkData& WorkData)
 {
 	ensure(InExpr->NumChildren() == 0);
-	if (!WorkData.bSetupMemory)
+	if (WorkData.CompilerPhase == ERigVMCompilerPhase_BuildInstructions)
 	{
 		WorkData.VM->GetByteCode().AddExitOp();
 	}
@@ -3630,7 +4309,7 @@ bool URigVMCompiler::TraverseInvokeEntry(const FRigVMInvokeEntryExprAST* InExpr,
 		return false;
 	}
 
-	if (WorkData.bSetupMemory)
+	if (WorkData.CompilerPhase == ERigVMCompilerPhase_SetupMemory)
 	{
 		return true;
 	}
@@ -3911,7 +4590,7 @@ FString URigVMCompiler::GetPinHashImpl(const URigVMPin* InPin, const FRigVMVarEx
 						{
 							ParentProxy = ParentProxy.GetParent();
 
-							if(URigVMLibraryNode* LibraryNode = ParentProxy.GetSubject<URigVMLibraryNode>())
+							if(ParentProxy.GetSubject<URigVMLibraryNode>())
 							{
 								break;
 							}
@@ -3972,7 +4651,7 @@ FString URigVMCompiler::GetPinHashImpl(const URigVMPin* InPin, const FRigVMVarEx
 		if (InVarExpr)
 		{
 			const FRigVMASTProxy NodeProxy = InVarExpr->GetProxy().GetSibling(Node);
-			if (const FRigVMExprAST* NodeExpr = InVarExpr->GetParser()->GetExprForSubject(NodeProxy))
+			if (InVarExpr->GetParser()->GetExprForSubject(NodeProxy))
 			{
 				// rely on the proxy callstack to differentiate registers
 				const FString CallStackPath = NodeProxy.GetCallstack().GetCallPath(false /* include last */);
@@ -4163,7 +4842,7 @@ void URigVMCompiler::MarkDebugWatch(FRigVMExtendedExecuteContext& Context, bool
 			}
 		}
 
-		FString PinHash = GetPinHashImpl(Pin, VarExpression, CurrentCompilationFunction);
+		FString PinHash = GetPinHashImpl(Pin, VarExpression, CurrentlyCompilingFunctionNode);
 		if(const FRigVMOperand* Operand = OutOperands->Find(PinHash))
 		{
 			if (OutVM->GetDebugMemory(Context)->Num() == 0)
@@ -4303,11 +4982,10 @@ FRigVMOperand URigVMCompiler::FindOrAddRegister(const FRigVMVarExprAST* InVarExp
 	
 	FString CPPType = Pin->GetCPPType();
 	FString BaseCPPType = Pin->IsArray() ? Pin->GetArrayElementCppType() : CPPType;
-	FString Hash = GetPinHash(Pin, InVarExpr, CurrentCompilationFunction);
+	FString Hash = GetPinHash(Pin, InVarExpr, CurrentlyCompilingFunctionNode);
 	FRigVMOperand Operand;
 	FString RegisterKey = Hash;
 
-	bool bIsExecutePin = Pin->IsExecuteContext();
 	bool bIsLiteral = InVarExpr->GetType() == FRigVMExprAST::EType::Literal;
 	bool bIsVariable = Pin->IsRootPin() && (Pin->GetName() == URigVMVariableNode::ValueName) &&
 		InVarExpr->GetPin()->GetNode()->IsA<URigVMVariableNode>();
@@ -4349,7 +5027,7 @@ FRigVMOperand URigVMCompiler::FindOrAddRegister(const FRigVMVarExprAST* InVarExp
 			{
 				if (const URigVMPin* VirtualPin = Cast<URigVMPin>(Proxy.GetSubject()))
 				{
-					const FString VirtualPinHash = GetPinHash(VirtualPin, InVarExpr, CurrentCompilationFunction, Proxy);
+					const FString VirtualPinHash = GetPinHash(VirtualPin, InVarExpr, CurrentlyCompilingFunctionNode, Proxy);
 					HashesWithSharedOperand.Add(VirtualPinHash);
 					if (Pin != VirtualPin)
 					{
@@ -4487,7 +5165,8 @@ FRigVMOperand URigVMCompiler::FindOrAddRegister(const FRigVMVarExprAST* InVarExp
 			}
 		}
 
-		Operand = WorkData.AddProperty(MemoryType, RegisterName, CPPType, Pin->GetCPPTypeObject(), JoinedDefaultValue);
+		const bool bMakeUniqueProperty = MemoryType == ERigVMMemoryType::Work;
+		Operand = WorkData.FindOrAddProperty(MemoryType, RegisterName, CPPType, Pin->GetCPPTypeObject(), JoinedDefaultValue, bMakeUniqueProperty);
 	}
 	ensure(Operand.IsValid());
 
@@ -4527,11 +5206,11 @@ const FRigVMCompilerWorkData::FRigVMASTProxyArray& URigVMCompiler::FindProxiesWi
 	
 	FRigVMCompilerWorkData::FRigVMASTProxyArray PinProxies, PinProxiesToProcess;
 	const FRigVMCompilerWorkData::FRigVMASTProxySourceMap& ProxySources = *WorkData.ProxySources;
-	const FRigVMCompilerWorkData::FRigVMASTProxyTargetsMap& ProxyTargets = WorkData.ProxyTargets;
 
 	PinProxiesToProcess.Add(InProxy);
 
 	const FString CPPType = InProxy.GetSubjectChecked<URigVMPin>()->GetCPPType();
+	const bool bEnableCallables = CVarRigVMCompileFunctionsToCallables->GetBool(); 
 
 	for(int32 ProxyIndex = 0; ProxyIndex < PinProxiesToProcess.Num(); ProxyIndex++)
 	{
@@ -4546,6 +5225,7 @@ const FRigVMCompilerWorkData::FRigVMASTProxyArray& URigVMCompiler::FindProxiesWi
 						continue;
 					}
 				}
+
 				// due to LWC we may have two pins that don't
 				// actually share the same CPP type (float vs double)
 				if(Pin->GetCPPType() != CPPType)
@@ -4573,7 +5253,28 @@ const FRigVMCompilerWorkData::FRigVMASTProxyArray& URigVMCompiler::FindProxiesWi
 			{
 				if (!PinProxies.Contains(*SourceProxy) && !PinProxiesToProcess.Contains(*SourceProxy))
 				{
-					PinProxiesToProcess.Add(*SourceProxy);
+					bool bValidSourceProxy = true;
+					if (bEnableCallables && PinProxiesToProcess[ProxyIndex].IsValid())
+					{
+						if (URigVMPin* TargetPin = Cast<URigVMPin>(PinProxiesToProcess[ProxyIndex].GetSubject()))
+						{
+							if (TargetPin->GetNode()->IsA<URigVMFunctionReturnNode>())
+							{
+								if (URigVMPin* SourcePin = Cast<URigVMPin>(SourceProxy->GetSubject()))
+								{
+									if (SourcePin->GetNode()->IsA<URigVMFunctionEntryNode>())
+									{
+										bValidSourceProxy = false;
+									}
+								}
+							}
+						}
+					}
+
+					if (bValidSourceProxy)
+					{
+						PinProxiesToProcess.Add(*SourceProxy);
+					}
 				}
 			}
 		}
@@ -4586,7 +5287,28 @@ const FRigVMCompilerWorkData::FRigVMASTProxyArray& URigVMCompiler::FindProxiesWi
 				{
 					if (!PinProxies.Contains(TargetProxy) && !PinProxiesToProcess.Contains(TargetProxy))
 					{
-						PinProxiesToProcess.Add(TargetProxy);
+						bool bValidTargetProxy = true;
+						if (bEnableCallables && PinProxiesToProcess[ProxyIndex].IsValid())
+						{
+							if (URigVMPin* SourcePin = Cast<URigVMPin>(PinProxiesToProcess[ProxyIndex].GetSubject()))
+							{
+								if (SourcePin->GetNode()->IsA<URigVMFunctionEntryNode>())
+								{
+									if (URigVMPin* TargetPin = Cast<URigVMPin>(TargetProxy.GetSubject()))
+									{
+										if (TargetPin->GetNode()->IsA<URigVMFunctionReturnNode>())
+										{
+											bValidTargetProxy = false;
+										}
+									}
+								}
+							}
+						}
+
+						if (bValidTargetProxy)
+						{
+							PinProxiesToProcess.Add(TargetProxy);
+						}
 					}
 				}
 			}
@@ -4662,32 +5384,6 @@ bool URigVMCompiler::ValidateNode(const FRigVMCompileSettings& InSettings, URigV
 	return false;
 }
 
-void URigVMCompiler::SetupPropertyPathsForFoldedCopies(const FRigVMNodeExprAST* InExpr, FRigVMCompilerWorkData& WorkData)
-{
-	// setup additional property paths needed for call externs accessing sub pin data directly
-	for(const TPair<int32, FString>& SegmentPathPair : InExpr->SegmentPathForChild)
-	{
-		const FRigVMExprAST* ChildExpr = InExpr->ChildAt(SegmentPathPair.Key);
-		check(ChildExpr);
-
-		if(ChildExpr->IsA(FRigVMExprAST::EType::CachedValue))
-		{
-			ChildExpr = ChildExpr->To<FRigVMCachedValueExprAST>()->GetVarExpr();
-		}
-
-		if(!ChildExpr->IsA(FRigVMExprAST::EType::Var))
-		{
-			continue;
-		}
-
-		const FRigVMVarExprAST* SourceVarExpr = GetSourceVarExpr(ChildExpr);
-		const FRigVMOperand RootOperand = WorkData.ExprToOperand.FindChecked(SourceVarExpr);
-		const FString HeadCPPType = SourceVarExpr->GetPin()->GetRootPin()->GetCPPType();
-
-		WorkData.FindOrAddPropertyPath(RootOperand, HeadCPPType, SegmentPathPair.Value);
-	}
-}
-
 void URigVMCompiler::ReportInfo(const FRigVMCompileSettings& InSettings, const FString& InMessage)
 {
 	if (InSettings.SurpressInfoMessages)
diff --git a/Engine/Plugins/Runtime/RigVM/Source/RigVMDeveloper/Private/RigVMModel/Nodes/RigVMFunctionReferenceNode.cpp b/Engine/Plugins/Runtime/RigVM/Source/RigVMDeveloper/Private/RigVMModel/Nodes/RigVMFunctionReferenceNode.cpp
--- a/Engine/Plugins/Runtime/RigVM/Source/RigVMDeveloper/Private/RigVMModel/Nodes/RigVMFunctionReferenceNode.cpp
+++ b/Engine/Plugins/Runtime/RigVM/Source/RigVMDeveloper/Private/RigVMModel/Nodes/RigVMFunctionReferenceNode.cpp
@@ -305,6 +305,15 @@ const FRigVMGraphFunctionData* URigVMFunctionReferenceNode::GetReferencedFunctio
 	return nullptr;
 }
 
+bool URigVMFunctionReferenceNode::SupportsCallable() const
+{
+	if (const FRigVMGraphFunctionData* Data = GetReferencedFunctionData())
+	{
+		return Data->CompilationData.SupportsCallable();
+	}
+	return false;
+}
+
 TArray<FRigVMTag> URigVMFunctionReferenceNode::GetVariantTags() const
 {
 	if (const FRigVMGraphFunctionData* Data = GetReferencedFunctionData(false))
diff --git a/Engine/Plugins/Runtime/RigVM/Source/RigVMDeveloper/Public/RigVMCompiler/RigVMAST.h b/Engine/Plugins/Runtime/RigVM/Source/RigVMDeveloper/Public/RigVMCompiler/RigVMAST.h
--- a/Engine/Plugins/Runtime/RigVM/Source/RigVMDeveloper/Public/RigVMCompiler/RigVMAST.h
+++ b/Engine/Plugins/Runtime/RigVM/Source/RigVMDeveloper/Public/RigVMCompiler/RigVMAST.h
@@ -609,6 +609,7 @@ private:
 	FRigVMASTProxy Proxy;
 	friend class FRigVMParserAST;
 	friend class URigVMCompiler;
+	friend struct FRigVMCompilerWorkData;
 };
 
 /*
diff --git a/Engine/Plugins/Runtime/RigVM/Source/RigVMDeveloper/Public/RigVMCompiler/RigVMASTProxy.h b/Engine/Plugins/Runtime/RigVM/Source/RigVMDeveloper/Public/RigVMCompiler/RigVMASTProxy.h
--- a/Engine/Plugins/Runtime/RigVM/Source/RigVMDeveloper/Public/RigVMCompiler/RigVMASTProxy.h
+++ b/Engine/Plugins/Runtime/RigVM/Source/RigVMDeveloper/Public/RigVMCompiler/RigVMASTProxy.h
@@ -119,9 +119,9 @@ private:
 	TArray<TWeakObjectPtr<UObject>> Stack;
 
 	friend class FRigVMParserAST;
-	friend class FRigVMCallstackGuard;
 	friend class FRigVMASTProxy;
 	friend class URigVMCompiler;
+	friend struct FRigVMCompilerWorkData;
 };
 
 /*
diff --git a/Engine/Plugins/Runtime/RigVM/Source/RigVMDeveloper/Public/RigVMCompiler/RigVMCompiler.h b/Engine/Plugins/Runtime/RigVM/Source/RigVMDeveloper/Public/RigVMCompiler/RigVMCompiler.h
--- a/Engine/Plugins/Runtime/RigVM/Source/RigVMDeveloper/Public/RigVMCompiler/RigVMCompiler.h
+++ b/Engine/Plugins/Runtime/RigVM/Source/RigVMDeveloper/Public/RigVMCompiler/RigVMCompiler.h
@@ -20,6 +20,7 @@
 
 #define UE_API RIGVMDEVELOPER_API
 
+class URigVMCompiler;
 class URigVMFunctionReferenceNode;
 struct FRigVMGraphFunctionData;
 struct FRigVMGraphFunctionHeader;
@@ -137,13 +138,23 @@ private:
 	TGuardValue<bool> ASTErrorsAsNotifications;
 };
 
+enum ERigVMCompilerPhase : uint8
+{
+	ERigVMCompilerPhase_SetupMemory,
+	ERigVMCompilerPhase_BuildInstructions,
+};
+
 struct FRigVMCompilerWorkData
 {
 public:
+	URigVMCompiler* Compiler = nullptr;
 	FRigVMCompileSettings Settings;
-	bool bSetupMemory = false;
+	ERigVMCompilerPhase CompilerPhase = ERigVMCompilerPhase_SetupMemory;
 	URigVM* VM = nullptr;
 	TArray<URigVMGraph*> Graphs;
+	const URigVMLibraryNode* CurrentlyCompilingFunctionNode = nullptr;
+	FRigVMFunctionCompilationData* FunctionCompilationData = nullptr;
+	TMap<uint32,TMap<FRigVMOperand,FRigVMOperand>> OperandMapPerFunction;
 	UScriptStruct* ExecuteContextStruct = nullptr;
 	TMap<FString, FRigVMOperand>* PinPathToOperand = nullptr;
 	FRigVMExtendedExecuteContext* Context = nullptr;
@@ -166,19 +177,19 @@ public:
 
 	struct FFunctionRegisterData
 	{
-		TSoftObjectPtr<URigVMFunctionReferenceNode> ReferenceNode;
+		uint32 ObjectId;;
 		ERigVMMemoryType MemoryType = ERigVMMemoryType::Invalid;
 		int32 RegisterIndex = 0;
 
 		friend inline uint32 GetTypeHash(const FFunctionRegisterData& Data)
 		{
-			uint32 Result = HashCombine(GetTypeHash(Data.ReferenceNode.ToSoftObjectPath().ToString()), GetTypeHash(Data.MemoryType));
+			uint32 Result = HashCombine(GetTypeHash(Data.ObjectId), GetTypeHash(Data.MemoryType));
 			return HashCombine(Result, GetTypeHash(Data.RegisterIndex));
 		}
 
 		bool operator==(const FFunctionRegisterData& Other) const
 		{
-			return ReferenceNode == Other.ReferenceNode && MemoryType == Other.MemoryType && RegisterIndex == Other.RegisterIndex;
+			return ObjectId == Other.ObjectId && MemoryType == Other.MemoryType && RegisterIndex == Other.RegisterIndex;
 		}
 	};
 	TMap<FFunctionRegisterData, FRigVMOperand> FunctionRegisterToOperand; 
@@ -190,11 +201,34 @@ public:
 
 	UE_API FRigVMOperand AddProperty(ERigVMMemoryType InMemoryType, const FName& InName, const FString& InCPPType, UObject* InCPPTypeObject, const FString& InDefaultValue = FString());
 	UE_API FRigVMOperand FindProperty(ERigVMMemoryType InMemoryType, const FName& InName) const;
+	UE_API FRigVMOperand FindOrAddProperty(const ERigVMMemoryType InMemoryType, const FRigVMFunctionCompilationPropertyDescription& InProperty, const FString& InNewName, const bool bMakeUniqueProperty);
+	UE_API FRigVMOperand FindOrAddProperty(ERigVMMemoryType InMemoryType, const FName& InName, const FString& InCPPType, UObject* InCPPTypeObject, const FString& InDefaultValue, const bool bMakeUniqueProperty);
 	UE_API FRigVMPropertyDescription GetProperty(const FRigVMOperand& InOperand);
 	UE_API int32 FindOrAddPropertyPath(const FRigVMOperand& InOperand, const FString& InHeadCPPType, const FString& InSegmentPath);
 	UE_API const FProperty* GetPropertyForOperand(const FRigVMOperand& InOperand) const;
 	UE_API TRigVMTypeIndex GetTypeIndexForOperand(const FRigVMOperand& InOperand) const;
 	UE_API FName GetUniquePropertyName(ERigVMMemoryType InMemoryType, const FName& InDesiredName) const;
+	UE_API bool AddPropertiesForFunction(
+		bool InIsCallable,
+		const uint32& InHash,
+		const FRigVMGraphFunctionHeader& InHeader,
+		const FRigVMFunctionCompilationData* InCompilationData,
+		int32 InFirstInstruction, int32 InLastInstruction,
+		const TArray<FRigVMOperand>& InInterfaceOperands,
+		const TMap<FRigVMOperand, FRigVMOperand>& InForwardedHandleToInterfaceOperand); 
+	UE_API void SetupPropertyPathsForFoldedCopies(const FRigVMNodeExprAST* InExpr);
+
+	UE_API bool RemapInstructionAt(
+		const FRigVMInstruction& InOldInstruction,
+		const FString& InFunctionIdentifier,
+		const FRigVMFunctionCompilationData* InCompilationData,
+		FRigVMByteCode& InNewByteCode,
+		const FRigVMInstructionArray& InNewInstructions,
+		const URigVMFunctionReferenceNode* InFunctionReferenceNode,
+		const TFunction<int32(const int32&)>& InRemapInstructionIndex,
+		const TFunction<FRigVMOperand(const FRigVMOperand&)>& InRemapOperand,
+		const TFunction<int32(const int32&)>& InRemapBranchIndex,
+		const TFunction<int32(const int32&)>& InRemapCallableIndex);
 
 	TSharedPtr<FRigVMParserAST> AST;
 	
@@ -330,8 +364,17 @@ public:
 
 	UE_API bool CompileFunction(const FRigVMCompileSettings& InSettings, const URigVMLibraryNode* InLibraryNode, URigVMController* InController, const TArray<FRigVMExternalVariable>& InExternalVariables, FRigVMFunctionCompilationData* OutFunctionCompilationData, FRigVMExtendedExecuteContext& OutVMContext);
 
+	UE_API bool CompileCallable(FRigVMCompilerWorkData& WorkData, URigVM* OutVM, FRigVMExtendedExecuteContext& OutVMContext, const FRigVMCallableInfo* InCallable, const FRigVMGraphFunctionHeader& InHeader, const FRigVMFunctionCompilationData* InFunctionCompilationData);
+
 	FRigVMCompiler_GetFunctionCompilationData GetFunctionCompilationData;
-	TMap<FString, const FRigVMFunctionCompilationData*> CompiledFunctions;
+
+	struct FCompiledFunctionData
+	{
+		const URigVMFunctionReferenceNode* FunctionReferenceNode;
+		const FRigVMFunctionCompilationData* CompilationData;
+	};
+	
+	TMap<FString, FCompiledFunctionData> CompiledFunctions;
 
 	static UE_API UScriptStruct* GetScriptStructForCPPType(const FString& InCPPType);
 	static UE_API FString GetPinHash(const URigVMPin* InPin, const FRigVMVarExprAST* InVarExpr, const URigVMLibraryNode* FunctionCompiling = nullptr, const FRigVMASTProxy& InPinProxy = FRigVMASTProxy());
@@ -345,7 +388,7 @@ public:
 
 private:
 
-	const URigVMLibraryNode* CurrentCompilationFunction = nullptr;
+	const URigVMLibraryNode* CurrentlyCompilingFunctionNode = nullptr;
 	TSet<const URigVMLibraryNode*> FunctionCompilationStack;
 
 	UE_API TArray<URigVMPin*> GetLinkedPins(URigVMPin* InPin, bool bInputs = true, bool bOutputs = true, bool bRecursive = true);
@@ -359,6 +402,7 @@ private:
 	UE_API bool TraverseEntry(const FRigVMEntryExprAST* InExpr, FRigVMCompilerWorkData& WorkData);
 	UE_API bool TraverseCallExtern(const FRigVMCallExternExprAST* InExpr, FRigVMCompilerWorkData& WorkData);
 	UE_API bool TraverseInlineFunction(const FRigVMInlineFunctionExprAST* InExpr, FRigVMCompilerWorkData& WorkData);
+	UE_API bool TraverseCallableFunction(const FRigVMInlineFunctionExprAST* InExpr, const uint32 InCallableHash, FRigVMCompilerWorkData& WorkData);
 	UE_API bool TraverseNoOp(const FRigVMNoOpExprAST* InExpr, FRigVMCompilerWorkData& WorkData);
 	UE_API bool TraverseVar(const FRigVMVarExprAST* InExpr, FRigVMCompilerWorkData& WorkData);
 	UE_API bool TraverseLiteral(const FRigVMVarExprAST* InExpr, FRigVMCompilerWorkData& WorkData);
@@ -388,8 +432,6 @@ private:
 	static UE_API int32 GetOperandFunctionInterfaceParameterIndex(const TArray<FString>& OperandsPinNames, const FRigVMFunctionCompilationData* FunctionCompilationData, const FRigVMOperand& Operand);
 
 	UE_API bool ValidateNode(const FRigVMCompileSettings& InSettings, URigVMNode* InNode, bool bCheck = true);
-
-	void SetupPropertyPathsForFoldedCopies(const FRigVMNodeExprAST* InExpr, FRigVMCompilerWorkData& WorkData);
 	
 	UE_API void ReportInfo(const FRigVMCompileSettings& InSettings, const FString& InMessage);
 	UE_API void ReportWarning(const FRigVMCompileSettings& InSettings, const FString& InMessage);
diff --git a/Engine/Plugins/Runtime/RigVM/Source/RigVMDeveloper/Public/RigVMModel/Nodes/RigVMFunctionReferenceNode.h b/Engine/Plugins/Runtime/RigVM/Source/RigVMDeveloper/Public/RigVMModel/Nodes/RigVMFunctionReferenceNode.h
--- a/Engine/Plugins/Runtime/RigVM/Source/RigVMDeveloper/Public/RigVMModel/Nodes/RigVMFunctionReferenceNode.h
+++ b/Engine/Plugins/Runtime/RigVM/Source/RigVMDeveloper/Public/RigVMModel/Nodes/RigVMFunctionReferenceNode.h
@@ -65,6 +65,7 @@ public:
 	UE_API void UpdateFunctionHeaderFromHost();
 
 	UE_API const FRigVMGraphFunctionData* GetReferencedFunctionData(bool bLoadIfNecessary = true) const;
+	UE_API bool SupportsCallable() const;
 
 	UE_API TArray<FRigVMTag> GetVariantTags() const;
 	
diff --git a/Engine/Plugins/Runtime/RigVM/Source/RigVMEditor/Private/Widgets/SRigVMExecutionStackView.cpp b/Engine/Plugins/Runtime/RigVM/Source/RigVMEditor/Private/Widgets/SRigVMExecutionStackView.cpp
--- a/Engine/Plugins/Runtime/RigVM/Source/RigVMEditor/Private/Widgets/SRigVMExecutionStackView.cpp
+++ b/Engine/Plugins/Runtime/RigVM/Source/RigVMEditor/Private/Widgets/SRigVMExecutionStackView.cpp
@@ -26,10 +26,10 @@ TAutoConsoleVariable<bool> CVarRigVMExecutionStackDetailedLabels(TEXT("RigVM.Sta
 //////////////////////////////////////////////////////////////
 /// FRigStackEntry
 ///////////////////////////////////////////////////////////
-FRigStackEntry::FRigStackEntry(int32 InEntryIndex, ERigStackEntry::Type InEntryType, int32 InInstructionIndex, ERigVMOpCode InOpCode, const FString& InLabel, const FRigVMASTProxy& InProxy)
+FRigStackEntry::FRigStackEntry(int32 InEntryIndex, ERigStackEntry::Type InEntryType, int32 InSubjectIndex, ERigVMOpCode InOpCode, const FString& InLabel, const FRigVMASTProxy& InProxy)
 	: EntryIndex(InEntryIndex)
 	, EntryType(InEntryType)
-	, InstructionIndex(InInstructionIndex)
+	, SubjectIndex(InSubjectIndex)
 	, CallPath(InProxy.GetCallstack().GetCallPath())
 	, Callstack(InProxy.GetCallstack())
 	, OpCode(InOpCode)
@@ -50,6 +50,36 @@ TSharedRef<ITableRow> FRigStackEntry::MakeTreeRowWidget(const TSharedRef<STableV
 	return SNew(SRigStackItem, InOwnerTable, InEntry, InCommandList, InBlueprint);
 }
 
+void FRigStackEntry::FilterChildren(const FString& InSearchText)
+{
+	if (FilteredChildren.IsSet())
+	{
+		return;
+	}
+
+	FilteredChildren = Children;
+	FilteredChildren.GetValue().RemoveAll([InSearchText](const TSharedPtr<FRigStackEntry>& InEntry) -> bool
+	{
+		return !InEntry->MatchesFilter(InSearchText);
+	});
+}
+
+bool FRigStackEntry::MatchesFilter(const FString& InSearchText)
+{
+	if (Label.Contains(InSearchText))
+	{
+		return true;
+	}
+	if (Children.IsEmpty())
+	{
+		return false;
+	}
+	
+	FilterChildren(InSearchText);
+	check(FilteredChildren.IsSet());
+	return !FilteredChildren.GetValue().IsEmpty();
+}
+
 //////////////////////////////////////////////////////////////
 /// SRigStackItem
 ///////////////////////////////////////////////////////////
@@ -65,7 +95,7 @@ void SRigStackItem::Construct(const FArguments& InArgs, const TSharedRef<STableV
 	const FSlateBrush* Icon = nullptr;
 	switch (InStackEntry->EntryType)
 	{
-		case ERigStackEntry::Operator:
+		case ERigStackEntry::Instruction:
 		{
 			Icon = FSlateIcon(TEXT("RigVMEditor"), "RigVM.Unit").GetIcon();
 			break;
@@ -85,12 +115,28 @@ void SRigStackItem::Construct(const FArguments& InArgs, const TSharedRef<STableV
 			Icon = FAppStyle::GetBrush("Icons.Error");
 			break;
 		}
+		case ERigStackEntry::Callable:
+		{
+			Icon = FAppStyle::GetBrush("Kismet.AllClasses.FunctionIcon");
+			break;
+		}
 		default:
 		{
 			break;
 		}
 	}
 
+	const EVisibility InstructionVisibility =
+		(
+			InStackEntry->EntryType == ERigStackEntry::Instruction
+		) ? EVisibility::Visible : EVisibility::Collapsed;
+
+	const EVisibility InstrumentationVisibility =
+		(
+			InStackEntry->EntryType == ERigStackEntry::Instruction ||
+			InStackEntry->EntryType == ERigStackEntry::Callable
+		) ? EVisibility::Visible : EVisibility::Collapsed;
+
 	STableRow<TSharedPtr<FRigStackEntry>>::Construct(
 		STableRow<TSharedPtr<FRigStackEntry>>::FArguments()
 		.Content()
@@ -109,6 +155,7 @@ void SRigStackItem::Construct(const FArguments& InArgs, const TSharedRef<STableV
 				.HAlign(HAlign_Right)
 				[
 					SAssignNew(NumberWidget, STextBlock)
+					.Visibility(InstructionVisibility)
 					.Text(this, &SRigStackItem::GetIndexText)
 					.ColorAndOpacity(this, &SRigStackItem::GetLabelColorAndOpacity)
 					.Font(IDetailLayoutBuilder::GetDetailFont())
@@ -155,7 +202,8 @@ void SRigStackItem::Construct(const FArguments& InArgs, const TSharedRef<STableV
 	            .HAlign(HAlign_Left)
 	            [
 	                SNew(STextBlock)
-	                .Text(this, &SRigStackItem::GetVisitedCountText)
+	                .Visibility(InstrumentationVisibility)
+					.Text(this, &SRigStackItem::GetVisitedCountText)
 					.ColorAndOpacity(this, &SRigStackItem::GetLabelColorAndOpacity)
 	                .Font(IDetailLayoutBuilder::GetDetailFont())
 	            ]
@@ -167,6 +215,7 @@ void SRigStackItem::Construct(const FArguments& InArgs, const TSharedRef<STableV
 	            .HAlign(HAlign_Left)
 	            [
 	                SNew(STextBlock)
+					.Visibility(InstrumentationVisibility)
 	                .Text(this, &SRigStackItem::GetDurationText)
 					.ColorAndOpacity(this, &SRigStackItem::GetLabelColorAndOpacity)
 	                .Font(IDetailLayoutBuilder::GetDetailFont())
@@ -177,7 +226,7 @@ void SRigStackItem::Construct(const FArguments& InArgs, const TSharedRef<STableV
 
 FText SRigStackItem::GetIndexText() const
 {
-	const FString IndexStr = FString::FromInt(WeakStackEntry.Pin()->EntryIndex) + TEXT(".");
+	const FString IndexStr = FString::FromInt(WeakStackEntry.Pin()->SubjectIndex) + TEXT(".");
 	return FText::FromString(IndexStr);
 }
 
@@ -280,8 +329,8 @@ void SRigVMExecutionStackView::Construct( const FArguments& InArgs, TSharedRef<I
 			.Padding(2.0f)
 			.BorderImage(FAppStyle::GetBrush("SCSEditor.TreePanel"))
 			[
-				SAssignNew(TreeView, STreeView<TSharedPtr<FRigStackEntry>>)
-				.TreeItemsSource(&Operators)
+				SAssignNew(TreeView, SRigVMExecutionStackTreeView)
+				.TreeItemsSource(&FilteredRootEntries)
 				.SelectionMode(ESelectionMode::Multi)
 				.OnGenerateRow(this, &SRigVMExecutionStackView::MakeTableRowWidget, WeakRigVMBlueprint)
 				.OnGetChildren(this, &SRigVMExecutionStackView::HandleGetChildrenForTree)
@@ -292,7 +341,7 @@ void SRigVMExecutionStackView::Construct( const FArguments& InArgs, TSharedRef<I
 		]
 	];
 
-	RefreshTreeView(nullptr, nullptr);
+	RefreshTreeView(nullptr, nullptr, true);
 
 	if (TStrongObjectPtr<UObject> StrongPtr = WeakRigVMBlueprint.GetWeakObjectPtr().Pin())
 	{
@@ -399,7 +448,14 @@ TSharedRef<ITableRow> SRigVMExecutionStackView::MakeTableRowWidget(TSharedPtr<FR
 
 void SRigVMExecutionStackView::HandleGetChildrenForTree(TSharedPtr<FRigStackEntry> InItem, TArray<TSharedPtr<FRigStackEntry>>& OutChildren)
 {
-	OutChildren = InItem->Children;
+	if (InItem->FilteredChildren.IsSet())
+	{
+		OutChildren = InItem->FilteredChildren.GetValue();
+	}
+	else
+	{
+		OutChildren = InItem->Children;
+	}
 }
 
 void SRigVMExecutionStackView::PopulateStackView(URigVM* InVM, FRigVMExtendedExecuteContext* InVMContext)
@@ -409,6 +465,82 @@ void SRigVMExecutionStackView::PopulateStackView(URigVM* InVM, FRigVMExtendedExe
 		const FRigVMInstructionArray Instructions = InVM->GetInstructions();
 		const FRigVMByteCode& ByteCode = InVM->GetByteCode();
 
+		TArray<int32> CallableIndexPerInstruction;
+		CallableIndexPerInstruction.AddZeroed(Instructions.Num());
+		for (int32& CallableIndex : CallableIndexPerInstruction)
+		{
+			CallableIndex = INDEX_NONE;
+		}
+
+		TArray<int32> BlockIndexPerInstruction;
+		BlockIndexPerInstruction.AddZeroed(Instructions.Num());
+		for (int32 InstructionIndex = 0; InstructionIndex < Instructions.Num(); InstructionIndex++)
+		{
+			BlockIndexPerInstruction[InstructionIndex] = INDEX_NONE;
+		}
+		for (int32 InstructionIndex = 0; InstructionIndex < Instructions.Num(); InstructionIndex++)
+		{
+			if (Instructions[InstructionIndex].OpCode == ERigVMOpCode::RunInstructions)
+			{
+				const FRigVMRunInstructionsOp& Op = ByteCode.GetOpAt<FRigVMRunInstructionsOp>(Instructions[InstructionIndex]);
+				if (Op.EndInstruction < Op.StartInstruction)
+				{
+					continue;
+				}
+				const TTuple<int32,int32> Block = {Op.StartInstruction, FMath::Max(Op.StartInstruction, Op.EndInstruction)};
+				int32 BlockIndex = Blocks.AddUnique(Block);
+				for (int32 ChildInstructionIndex = Block.Get<0>(); ChildInstructionIndex <= Block.Get<1>(); ChildInstructionIndex++)
+				{
+					if (BlockIndexPerInstruction[ChildInstructionIndex] == INDEX_NONE)
+					{
+						BlockIndexPerInstruction[ChildInstructionIndex] = BlockIndex;
+					}
+					else
+					{
+						const TTuple<int32,int32> OldBlock = Blocks[BlockIndexPerInstruction[ChildInstructionIndex]];
+						const int32 NewLen = Block.Get<1>() - Block.Get<0>(); 
+						const int32 OldLen = OldBlock.Get<1>() - OldBlock.Get<0>();
+						if (NewLen < OldLen)
+						{
+							BlockIndexPerInstruction[ChildInstructionIndex] = BlockIndex;
+						}
+					}
+				}
+			}
+		}
+
+		for (int32 CallableIndex = 0; CallableIndex < ByteCode.NumCallables(); CallableIndex++)
+		{
+			const FRigVMCallableInfo* Callable = ByteCode.GetCallable(CallableIndex);
+
+			FRigVMASTProxy Proxy;
+			FString CallableLabel = Callable->Name.ToString();
+			int32 LastIndex = INDEX_NONE;
+			if (CallableLabel.FindLastChar(TEXT(':'), LastIndex))
+			{
+				CallableLabel = CallableLabel.Mid(LastIndex + 1);
+			}
+			TSharedPtr<FRigStackEntry> NewEntry = MakeEntry(ERigStackEntry::Callable, CallableIndex, ERigVMOpCode::Invalid, CallableLabel, Proxy);
+			RootEntries.Add(NewEntry);
+			CallableToEntry.Add(CallableIndex, NewEntry);
+
+			for (int32 InstructionIndex = Callable->FirstInstruction; InstructionIndex <= Callable->LastInstruction; InstructionIndex++)
+			{
+				check(CallableIndexPerInstruction.IsValidIndex(InstructionIndex));
+				CallableIndexPerInstruction[InstructionIndex] = CallableIndex;
+			}
+		}
+
+		for (int32 BlockIndex = 0; BlockIndex < Blocks.Num(); BlockIndex++)
+		{
+			const TTuple<int32,int32> Block = Blocks[BlockIndex];
+			const FString BlockLabel = FString::Printf(TEXT("Block %04d - %04d"), Block.Get<0>(), Block.Get<1>());
+
+			FRigVMASTProxy Proxy;
+			TSharedPtr<FRigStackEntry> NewEntry = MakeEntry(ERigStackEntry::Block, BlockIndex, ERigVMOpCode::Invalid, BlockLabel, Proxy);
+			BlockToEntry.Add(Block, NewEntry);
+		}
+
 		TArray<URigVMGraph*> RootGraphs;
 		RootGraphs.AddZeroed(Instructions.Num());
 
@@ -468,7 +600,22 @@ void SRigVMExecutionStackView::PopulateStackView(URigVM* InVM, FRigVMExtendedExe
 						}
 						else
 						{
-							Label = InVM->GetRigVMFunctionName(Op.FunctionIndex);
+							Label = InVM->GetRigVMFunctionName(Op.CallableIndex);
+						}
+						break;
+					}
+					case ERigVMOpCode::InvokeCallable:
+					{
+						const FRigVMInvokeCallableOp& Op = ByteCode.GetOpAt<FRigVMInvokeCallableOp>(Instructions[InstructionIndex]);
+						if (const FRigVMCallableInfo* CallableInfo = ByteCode.GetCallable(Op.CallableIndex))
+						{
+							FString CallableLabel = CallableInfo->Name.ToString();
+							int32 LastIndex = INDEX_NONE;
+							if (CallableLabel.FindLastChar(TEXT(':'), LastIndex))
+							{
+								CallableLabel = CallableLabel.Mid(LastIndex + 1);
+							}
+							Label = FString::Printf(TEXT("Invoke %s"), *CallableLabel);
 						}
 						break;
 					}
@@ -480,12 +627,6 @@ void SRigVMExecutionStackView::PopulateStackView(URigVM* InVM, FRigVMExtendedExe
 					case ERigVMOpCode::Decrement:
 					case ERigVMOpCode::Equals:
 					case ERigVMOpCode::NotEquals:
-					case ERigVMOpCode::JumpAbsolute:
-					case ERigVMOpCode::JumpForward:
-					case ERigVMOpCode::JumpBackward:
-					case ERigVMOpCode::JumpAbsoluteIf:
-					case ERigVMOpCode::JumpForwardIf:
-					case ERigVMOpCode::JumpBackwardIf:
 					case ERigVMOpCode::BeginBlock:
 					case ERigVMOpCode::EndBlock:
 					case ERigVMOpCode::ChangeType:
@@ -523,6 +664,22 @@ void SRigVMExecutionStackView::PopulateStackView(URigVM* InVM, FRigVMExtendedExe
 						Label = FString::Printf(TEXT("Run %s Event"), *Op.EntryName.ToString());
 						break;
 					}
+					case ERigVMOpCode::JumpAbsolute:
+					case ERigVMOpCode::JumpForward:
+					case ERigVMOpCode::JumpBackward:
+					{
+						const FRigVMJumpOp& Op = ByteCode.GetOpAt<FRigVMJumpOp>(Instructions[InstructionIndex]);
+						Label = FString::Printf(TEXT("Jump to %d"), FRigVMJumpOp::GetTargetInstruction(Op.OpCode, InstructionIndex, Op.InstructionIndex));
+						break;
+					}
+					case ERigVMOpCode::JumpAbsoluteIf:
+					case ERigVMOpCode::JumpForwardIf:
+					case ERigVMOpCode::JumpBackwardIf:
+					{
+						const FRigVMJumpIfOp& Op = ByteCode.GetOpAt<FRigVMJumpIfOp>(Instructions[InstructionIndex]);
+						Label = FString::Printf(TEXT("Maybe jump to %d"), FRigVMJumpOp::GetTargetInstruction(Op.OpCode, InstructionIndex, Op.InstructionIndex));
+						break;
+					}
 					case ERigVMOpCode::JumpToBranch:
 					{
 						const FRigVMJumpToBranchOp& Op = ByteCode.GetOpAt<FRigVMJumpToBranchOp>(Instructions[InstructionIndex]);
@@ -532,7 +689,7 @@ void SRigVMExecutionStackView::PopulateStackView(URigVM* InVM, FRigVMExtendedExe
 					case ERigVMOpCode::RunInstructions:
 					{
 						const FRigVMRunInstructionsOp& Op = ByteCode.GetOpAt<FRigVMRunInstructionsOp>(Instructions[InstructionIndex]);
-						Label = FString::Printf(TEXT("Run Instructions %d-%d"), Op.StartInstruction, Op.EndInstruction);
+						Label = FString::Printf(TEXT("Run Block %d-%d"), Op.StartInstruction, FMath::Max(Op.StartInstruction, Op.EndInstruction));
 						break;
 					}
 					case ERigVMOpCode::SetupTraits:
@@ -545,6 +702,7 @@ void SRigVMExecutionStackView::PopulateStackView(URigVM* InVM, FRigVMExtendedExe
 						Label = TEXT("Exit");
 						break;
 					}
+
 					default:
 					{
 						ensure(false);
@@ -554,126 +712,254 @@ void SRigVMExecutionStackView::PopulateStackView(URigVM* InVM, FRigVMExtendedExe
 				
 				// add the entry with the new label to the stack view
 				const FRigVMInstruction& Instruction = Instructions[InstructionIndex];
-				if (FilterText.IsEmpty() || Label.Contains(FilterText.ToString()))
+				FRigVMASTProxy Proxy;
+				if(const TArray<TWeakObjectPtr<UObject>>* Callstack = ByteCode.GetCallstackForInstruction(InstructionIndex))
 				{
-					FRigVMASTProxy Proxy;
-					if(const TArray<TWeakObjectPtr<UObject>>* Callstack = ByteCode.GetCallstackForInstruction(InstructionIndex))
-					{
-						Proxy = FRigVMASTProxy::MakeFromCallstack(Callstack);
-					}
-					TSharedPtr<FRigStackEntry> NewEntry = MakeShared<FRigStackEntry>(InstructionIndex, ERigStackEntry::Operator, InstructionIndex, Instruction.OpCode, Label, Proxy);
-					Operators.Add(NewEntry);
+					Proxy = FRigVMASTProxy::MakeFromCallstack(Callstack);
 				}
+				TSharedPtr<FRigStackEntry> NewEntry = MakeEntry(ERigStackEntry::Instruction, InstructionIndex, Instruction.OpCode, Label, Proxy);
+				InstructionToEntry.Add(InstructionIndex, NewEntry);
 			}
-			return;
 		}
-
-		// 1. cache information about instructions/nodes, which will be used later 
-		TMap<FString, FString> OperandFormatMap;
-		for (int32 InstructionIndex = 0; InstructionIndex < Instructions.Num(); InstructionIndex++)
+		else
 		{
-			const FRigVMASTProxy Proxy = FRigVMASTProxy::MakeFromCallPath(ByteCode.GetCallPathForInstruction(InstructionIndex), RootGraphs[InstructionIndex]);
-			if(URigVMNode* Node = Proxy.GetSubject<URigVMNode>())
+			// 1. cache information about instructions/nodes, which will be used later 
+			TMap<FString, FString> OperandFormatMap;
+			for (int32 InstructionIndex = 0; InstructionIndex < Instructions.Num(); InstructionIndex++)
 			{
-				FString DisplayName = Node->GetName();
-
-				// only unit nodes among all nodes has StaticExecute() that generates actual instructions
-				if (URigVMUnitNode* UnitNode = Cast<URigVMUnitNode>(Node))
+				const FRigVMASTProxy Proxy = FRigVMASTProxy::MakeFromCallPath(ByteCode.GetCallPathForInstruction(InstructionIndex), RootGraphs[InstructionIndex]);
+				if(URigVMNode* Node = Proxy.GetSubject<URigVMNode>())
 				{
-					DisplayName = UnitNode->GetNodeTitle();
-#if WITH_EDITOR
-					UScriptStruct* Struct = UnitNode->GetScriptStruct();
-					FString MenuDescSuffixMetadata;
-					if (Struct)
+					FString DisplayName = Node->GetName();
+
+					// only unit nodes among all nodes has StaticExecute() that generates actual instructions
+					if (URigVMUnitNode* UnitNode = Cast<URigVMUnitNode>(Node))
 					{
-						Struct->GetStringMetaDataHierarchical(FRigVMStruct::MenuDescSuffixMetaName, &MenuDescSuffixMetadata);
+						DisplayName = UnitNode->GetNodeTitle();
+#if WITH_EDITOR
+						UScriptStruct* Struct = UnitNode->GetScriptStruct();
+						FString MenuDescSuffixMetadata;
+						if (Struct)
+						{
+							Struct->GetStringMetaDataHierarchical(FRigVMStruct::MenuDescSuffixMetaName, &MenuDescSuffixMetadata);
+						}
+						if (!MenuDescSuffixMetadata.IsEmpty())
+						{
+							DisplayName = FString::Printf(TEXT("%s %s"), *UnitNode->GetNodeTitle(), *MenuDescSuffixMetadata);
+						}
+#endif
 					}
-					if (!MenuDescSuffixMetadata.IsEmpty())
+
+					// this is needed for name replacement later
+					OperandFormatMap.Add(Node->GetName(), DisplayName);
+				}
+			}
+
+			// 2. replace raw operand names with NodeTitle.PinName/PropertyName.OffsetName
+			TArray<FString> Labels = InVM->DumpByteCodeAsTextArray(*InVMContext, TArray<int32>(), false, [OperandFormatMap](const FString& RegisterName, const FString& RegisterOffsetName)
+			{
+				FString NewRegisterName = RegisterName;
+				FString NodeName;
+				FString PinName;
+				if (RegisterName.Split(TEXT("."), &NodeName, &PinName))
+				{
+					const FString* NodeTitle = OperandFormatMap.Find(NodeName);
+					NewRegisterName = FString::Printf(TEXT("%s.%s"), NodeTitle ? **NodeTitle : *NodeName, *PinName);
+				}
+				FString OperandLabel;
+				OperandLabel = NewRegisterName;
+				if (!RegisterOffsetName.IsEmpty())
+				{
+					OperandLabel = FString::Printf(TEXT("%s.%s"), *OperandLabel, *RegisterOffsetName);
+				}
+				return OperandLabel;
+			});
+
+			if (!ensure(Labels.Num() == Instructions.Num()))
+			{
+				return;
+			}
+		
+			// 3. replace instruction names with node titles
+			for (int32 InstructionIndex = 0; InstructionIndex < Labels.Num(); InstructionIndex++)
+			{
+				FString Label = Labels[InstructionIndex];
+				const FRigVMASTProxy Proxy = FRigVMASTProxy::MakeFromCallPath(ByteCode.GetCallPathForInstruction(InstructionIndex), RootGraphs[InstructionIndex]);
+
+				if(URigVMNode* Node = Proxy.GetSubject<URigVMNode>())
+				{
+					FString Suffix;
+					switch(Instructions[InstructionIndex].OpCode)
 					{
-						DisplayName = FString::Printf(TEXT("%s %s"), *UnitNode->GetNodeTitle(), *MenuDescSuffixMetadata);
+						case ERigVMOpCode::Copy:
+						case ERigVMOpCode::Zero:
+						case ERigVMOpCode::BoolFalse:
+						case ERigVMOpCode::BoolTrue:
+						case ERigVMOpCode::Increment:
+						case ERigVMOpCode::Decrement:
+						case ERigVMOpCode::Equals:
+						case ERigVMOpCode::NotEquals:
+						case ERigVMOpCode::JumpAbsolute:
+						case ERigVMOpCode::JumpForward:
+						case ERigVMOpCode::JumpBackward:
+						case ERigVMOpCode::JumpAbsoluteIf:
+						case ERigVMOpCode::JumpForwardIf:
+						case ERigVMOpCode::JumpBackwardIf:
+						case ERigVMOpCode::BeginBlock:
+						case ERigVMOpCode::EndBlock:
+						{
+							const FText OpCodeText = StaticEnum<ERigVMOpCode>()->GetDisplayNameTextByValue((int32)Instructions[InstructionIndex].OpCode);
+							Suffix = FString::Printf(TEXT(" - %s"), *OpCodeText.ToString());
+							break;
+						}
+						default:
+						{
+							break;
+						}
 					}
-#endif
+					
+					Label =  OperandFormatMap[Node->GetName()] + Suffix;
 				}
 
-				// this is needed for name replacement later
-				OperandFormatMap.Add(Node->GetName(), DisplayName);
+				// add the entry with the new label to the stack view
+				const FRigVMInstruction& Instruction = Instructions[InstructionIndex];
+				TSharedPtr<FRigStackEntry> NewEntry = MakeEntry(ERigStackEntry::Instruction, InstructionIndex, Instruction.OpCode, Label, Proxy);
+				InstructionToEntry.Add(InstructionIndex, NewEntry);
 			}
 		}
 
-		// 2. replace raw operand names with NodeTitle.PinName/PropertyName.OffsetName
-		TArray<FString> Labels = InVM->DumpByteCodeAsTextArray(*InVMContext, TArray<int32>(), false, [OperandFormatMap](const FString& RegisterName, const FString& RegisterOffsetName)
+		for (int32 InstructionIndex = 0; InstructionIndex < Instructions.Num(); InstructionIndex++)
+		{
+			TSharedPtr<FRigStackEntry>& NewEntry = InstructionToEntry.FindChecked(InstructionIndex); 
+			if (BlockIndexPerInstruction[InstructionIndex] != INDEX_NONE)
+			{
+				BlockToEntry.FindChecked(Blocks[BlockIndexPerInstruction[InstructionIndex]])->Children.Add(NewEntry);
+			}
+			else if (CallableIndexPerInstruction[InstructionIndex] != INDEX_NONE)
+			{
+				CallableToEntry.FindChecked(CallableIndexPerInstruction[InstructionIndex])->Children.Add(NewEntry);
+			}
+			else
+			{
+				RootEntries.Add(NewEntry);
+			}
+		}
+
+		TArray<TTuple<int32, int32>> SortedBlocks = Blocks;
+		Algo::Sort(SortedBlocks, [](const TTuple<int32,int32>& A, const TTuple<int32,int32>& B) -> bool
 		{
-			FString NewRegisterName = RegisterName;
-			FString NodeName;
-			FString PinName;
-			if (RegisterName.Split(TEXT("."), &NodeName, &PinName))
+			const int32 LengthA = FMath::Max(A.Get<0>(), A.Get<1>()) - A.Get<0>();
+			const int32 LengthB = FMath::Max(B.Get<0>(), B.Get<1>()) - B.Get<0>();
+			if (LengthA < LengthB)
 			{
-				const FString* NodeTitle = OperandFormatMap.Find(NodeName);
-				NewRegisterName = FString::Printf(TEXT("%s.%s"), NodeTitle ? **NodeTitle : *NodeName, *PinName);
+				return true;
 			}
-			FString OperandLabel;
-			OperandLabel = NewRegisterName;
-			if (!RegisterOffsetName.IsEmpty())
+			if (LengthA > LengthB)
 			{
-				OperandLabel = FString::Printf(TEXT("%s.%s"), *OperandLabel, *RegisterOffsetName);
+				return false;
 			}
-			return OperandLabel;
+			if (A.Get<0>() == B.Get<0>())
+			{
+				return A.Get<1>() < B.Get<1>();
+			}
+			return A.Get<0>() < B.Get<0>();
 		});
 
-		if (!ensure(Labels.Num() == Instructions.Num()))
+		for (int32 BlockIndexA = 0; BlockIndexA < SortedBlocks.Num(); BlockIndexA++)
 		{
-			return;
+			const TTuple<int32, int32>& BlockA = SortedBlocks[BlockIndexA];
+			TTuple<int32, int32> ParentBlock = {INDEX_NONE, INDEX_NONE};
+
+			for (int32 BlockIndexB = BlockIndexA + 1; BlockIndexB < SortedBlocks.Num(); BlockIndexB++)
+			{
+				const TTuple<int32, int32>& BlockB = SortedBlocks[BlockIndexB];
+				if (BlockB.Get<0>() <= BlockA.Get<0>() &&
+					FMath::Max(BlockB.Get<0>(), BlockB.Get<1>()) >= FMath::Max(BlockA.Get<0>(), BlockA.Get<1>()))
+				{
+					BlockToEntry.FindChecked(BlockB)->Children.Add(BlockToEntry.FindChecked(BlockA));
+					break;
+				}
+			}
+
+			TSharedPtr<FRigStackEntry> NewEntry = BlockToEntry.FindChecked(BlockA);
+			if (ParentBlock.Get<0>() != INDEX_NONE)
+			{
+				BlockToEntry.FindChecked(ParentBlock)->Children.Add(NewEntry);
+			}
+			else if (CallableIndexPerInstruction[BlockA.Get<0>()] != INDEX_NONE)
+			{
+				CallableToEntry.FindChecked(CallableIndexPerInstruction[BlockA.Get<0>()])->Children.Add(NewEntry);
+			}
+			else
+			{
+				RootEntries.Add(NewEntry);
+			}
 		}
-		
-		// 3. replace instruction names with node titles
-		for (int32 InstructionIndex = 0; InstructionIndex < Labels.Num(); InstructionIndex++)
+
+		// sort all children again.
+		for (TSharedPtr<FRigStackEntry>& Entry : AllEntries)
 		{
-			FString Label = Labels[InstructionIndex];
-			const FRigVMASTProxy Proxy = FRigVMASTProxy::MakeFromCallPath(ByteCode.GetCallPathForInstruction(InstructionIndex), RootGraphs[InstructionIndex]);
+			if (Entry->Children.IsEmpty())
+			{
+				continue;
+			}
 
-			if(URigVMNode* Node = Proxy.GetSubject<URigVMNode>())
+			Algo::Sort(Entry->Children, [this](const TSharedPtr<FRigStackEntry>& A, const TSharedPtr<FRigStackEntry>& B) -> bool
 			{
-				FString Suffix;
-				switch(Instructions[InstructionIndex].OpCode)
+				const ERigStackEntry::Type TypeA = A->EntryType;
+				const ERigStackEntry::Type TypeB = B->EntryType;
+
+				if (TypeA == ERigStackEntry::Instruction && TypeB == ERigStackEntry::Block)
 				{
-					case ERigVMOpCode::Copy:
-					case ERigVMOpCode::Zero:
-	                case ERigVMOpCode::BoolFalse:
-	                case ERigVMOpCode::BoolTrue:
-	                case ERigVMOpCode::Increment:
-	                case ERigVMOpCode::Decrement:
-					case ERigVMOpCode::Equals:
-                	case ERigVMOpCode::NotEquals:
-					case ERigVMOpCode::JumpAbsolute:
-                	case ERigVMOpCode::JumpForward:
-                	case ERigVMOpCode::JumpBackward:
-					case ERigVMOpCode::JumpAbsoluteIf:
-                	case ERigVMOpCode::JumpForwardIf:
-                	case ERigVMOpCode::JumpBackwardIf:
-					case ERigVMOpCode::BeginBlock:
-					case ERigVMOpCode::EndBlock:
+					const int32 InstructionIndexA = A->SubjectIndex;
+					const TTuple<int32,int32>& BlockB = Blocks[B->SubjectIndex];
+					if (InstructionIndexA < BlockB.Get<0>())
 					{
-						const FText OpCodeText = StaticEnum<ERigVMOpCode>()->GetDisplayNameTextByValue((int32)Instructions[InstructionIndex].OpCode);
-						Suffix = FString::Printf(TEXT(" - %s"), *OpCodeText.ToString());
-						break;
+						return true;
 					}
-					default:
+					if (InstructionIndexA > FMath::Max(BlockB.Get<0>(), BlockB.Get<1>()))
 					{
-						break;
+						return false;
 					}
+					// we should never get here since then the instruction should be inside the block
+					checkNoEntry();
 				}
-				
-				Label =  OperandFormatMap[Node->GetName()] + Suffix;
-				// Label = Proxy.GetCallstack().GetCallPath() + Suffix;
-			}
 
-			// add the entry with the new label to the stack view
-			const FRigVMInstruction& Instruction = Instructions[InstructionIndex];
-			if (FilterText.IsEmpty() || Label.Contains(FilterText.ToString()))
-			{
-				TSharedPtr<FRigStackEntry> NewEntry = MakeShared<FRigStackEntry>(InstructionIndex, ERigStackEntry::Operator, InstructionIndex, Instruction.OpCode, Label, Proxy);
-				Operators.Add(NewEntry);
-			}
+				if (TypeA == ERigStackEntry::Block && TypeB == ERigStackEntry::Instruction)
+				{
+					const TTuple<int32,int32>& BlockA = Blocks[A->SubjectIndex];
+					const int32 InstructionIndexB = B->SubjectIndex;
+					if (InstructionIndexB < BlockA.Get<0>())
+					{
+						return false;
+					}
+					if (InstructionIndexB > FMath::Max(BlockA.Get<0>(), BlockA.Get<1>()))
+					{
+						return true;
+					}
+					// we should never get here since then the instruction should be inside the block
+					checkNoEntry();
+				}
+
+				if (TypeA != TypeB)
+				{
+					return TypeA < TypeB;
+				}
+
+				if (TypeA == ERigStackEntry::Block)
+				{
+					const TTuple<int32,int32>& BlockA = Blocks[A->SubjectIndex];
+					const TTuple<int32,int32>& BlockB = Blocks[B->SubjectIndex];
+					return BlockA.Get<0>() < BlockB.Get<0>();
+				}
+					
+				if (TypeA == ERigStackEntry::Error || TypeA == ERigStackEntry::Warning || TypeA == ERigStackEntry::Info)
+				{
+					return FCString::Strcmp(*(A->Label), *(B->Label)) < 0;
+				}
+
+				return A->SubjectIndex < B->SubjectIndex;
+			});
 		}
 	}
 }
@@ -740,50 +1026,73 @@ void SRigVMExecutionStackView::UnbindAssetDelegates()
 	}
 }
 
-void SRigVMExecutionStackView::RefreshTreeView(URigVM* InVM, FRigVMExtendedExecuteContext* InVMContext)
+TSharedPtr<FRigStackEntry> SRigVMExecutionStackView::MakeEntry(ERigStackEntry::Type InEntryType,
+	int32 InSubjectIndex, ERigVMOpCode InOpCode, const FString& InLabel, const FRigVMASTProxy& InProxy)
 {
-	Operators.Reset();
+	TSharedPtr<FRigStackEntry> Entry = MakeShared<FRigStackEntry>(INDEX_NONE, InEntryType, InSubjectIndex, InOpCode, InLabel, InProxy);
+	Entry->EntryIndex = AllEntries.Add(Entry);
+	return Entry;
+}
 
-	// populate the stack with node names/instruction names
-	PopulateStackView(InVM, InVMContext);
-	
-	if (InVM && InVMContext)
+void SRigVMExecutionStackView::RefreshTreeView(URigVM* InVM, FRigVMExtendedExecuteContext* InVMContext, bool bRebuildEntries)
+{
+	TreeView->SaveAndClearSparseItemInfos();
+
+	if (bRebuildEntries)
 	{
-		// fill the children from the log
-		if (TSharedPtr<IRigVMEditor> Editor = WeakEditor.Pin())
+		RootEntries.Reset();
+		FilteredRootEntries.Reset();
+		AllEntries.Reset();
+		InstructionToEntry.Reset();
+		CallableToEntry.Reset();
+		Blocks.Reset();
+		BlockToEntry.Reset();
+
+		// populate the stack with node names/instruction names
+		PopulateStackView(InVM, InVMContext);
+		
+		if (InVM && InVMContext)
 		{
-			URigVMHost* Host = PreviouslyDebuggedHost.Get();
-			if(Host && Host->GetLog())
+			// fill the children from the log
+			if (TSharedPtr<IRigVMEditor> Editor = WeakEditor.Pin())
 			{
-				const TArray<FRigVMLog::FLogEntry>& LogEntries = Host->GetLog()->Entries;
-				for (const FRigVMLog::FLogEntry& LogEntry : LogEntries)
+				URigVMHost* Host = PreviouslyDebuggedHost.Get();
+				if(Host && Host->GetLog())
 				{
-					if (!Operators.IsValidIndex(LogEntry.InstructionIndex))
+					const TArray<FRigVMLog::FLogEntry>& LogEntries = Host->GetLog()->Entries;
+					for (const FRigVMLog::FLogEntry& LogEntry : LogEntries)
 					{
-						continue;
-					}
-					int32 ChildIndex = Operators[LogEntry.InstructionIndex]->Children.Num();
-					switch (LogEntry.Severity)
-					{
-						case EMessageSeverity::Info:
-						{
-							Operators[LogEntry.InstructionIndex]->Children.Add(MakeShared<FRigStackEntry>(ChildIndex, ERigStackEntry::Info, LogEntry.InstructionIndex, ERigVMOpCode::Invalid, LogEntry.Message, FRigVMASTProxy()));
-							break;
-						}
-						case EMessageSeverity::Warning:
-						case EMessageSeverity::PerformanceWarning:
+						TSharedPtr<FRigStackEntry>* StackEntryPtr = InstructionToEntry.Find(LogEntry.InstructionIndex);
+						if (!StackEntryPtr)
 						{
-							Operators[LogEntry.InstructionIndex]->Children.Add(MakeShared<FRigStackEntry>(ChildIndex, ERigStackEntry::Warning, LogEntry.InstructionIndex, ERigVMOpCode::Invalid, LogEntry.Message, FRigVMASTProxy()));
-							break;
+							continue;
 						}
-						case EMessageSeverity::Error:
+						TSharedPtr<FRigStackEntry>& StackEntry = *StackEntryPtr;
+						check(StackEntry.IsValid());
+						
+						int32 ChildIndex = StackEntry->Children.Num();
+						switch (LogEntry.Severity)
 						{
-							Operators[LogEntry.InstructionIndex]->Children.Add(MakeShared<FRigStackEntry>(ChildIndex, ERigStackEntry::Error, LogEntry.InstructionIndex, ERigVMOpCode::Invalid, LogEntry.Message, FRigVMASTProxy()));
-							break;
-						}
-						default:
-						{
-							break;
+							case EMessageSeverity::Info:
+							{
+								StackEntry->Children.Add(MakeEntry(ERigStackEntry::Info, LogEntry.InstructionIndex, ERigVMOpCode::Invalid, LogEntry.Message, FRigVMASTProxy()));
+								break;
+							}
+							case EMessageSeverity::Warning:
+							case EMessageSeverity::PerformanceWarning:
+							{
+								StackEntry->Children.Add(MakeEntry(ERigStackEntry::Warning, LogEntry.InstructionIndex, ERigVMOpCode::Invalid, LogEntry.Message, FRigVMASTProxy()));
+								break;
+							}
+							case EMessageSeverity::Error:
+							{
+								StackEntry->Children.Add(MakeEntry(ERigStackEntry::Error, LogEntry.InstructionIndex, ERigVMOpCode::Invalid, LogEntry.Message, FRigVMASTProxy()));
+								break;
+							}
+							default:
+							{
+								break;
+							}
 						}
 					}
 				}
@@ -791,6 +1100,26 @@ void SRigVMExecutionStackView::RefreshTreeView(URigVM* InVM, FRigVMExtendedExecu
 		}
 	}
 
+	FilteredRootEntries = RootEntries;
+	for (TSharedPtr<FRigStackEntry>& Entry : AllEntries)
+	{
+		Entry->FilteredChildren.Reset();
+	}
+	
+	if (!FilterText.IsEmpty())
+	{
+		const FString FilterTextString = FilterText.ToString();
+		FilteredRootEntries.RemoveAll([FilterTextString](TSharedPtr<FRigStackEntry>& Entry) -> bool
+		{
+			return !Entry->MatchesFilter(FilterTextString);
+		});
+	}
+
+	for (TSharedPtr<FRigStackEntry>& Entry : AllEntries)
+	{
+		TreeView->RestoreSparseItemInfos(Entry);
+	}
+	
 	TreeView->RequestTreeRefresh();
 }
 
@@ -850,7 +1179,7 @@ void SRigVMExecutionStackView::HandleFocusOnSelectedGraphNode()
 				}
 
 				const FRigVMByteCode& ByteCode = Host->GetVM()->GetByteCode();
-				UObject* Subject = ByteCode.GetSubjectForInstruction(SelectedItems[0]->InstructionIndex);
+				UObject* Subject = ByteCode.GetSubjectForInstruction(SelectedItems[0]->SubjectIndex);
 				if (URigVMNode* SelectedNode = Cast<URigVMNode>(Subject))
 				{
 					URigVMGraph* GraphToFocus = SelectedNode->GetGraph();
@@ -881,7 +1210,10 @@ void SRigVMExecutionStackView::HandleGoToInstruction()
 	TArray<TSharedPtr<FRigStackEntry>> SelectedItems = TreeView->GetSelectedItems();
 	if (SelectedItems.Num() > 0)
 	{
-		Index = SelectedItems[0]->InstructionIndex;
+		if (SelectedItems[0]->EntryType == ERigStackEntry::Instruction)
+		{
+			Index = SelectedItems[0]->SubjectIndex;
+		}
 	}
 	
 	const FVector2D MouseCursorLocation = FSlateApplication::Get().GetCursorPos();
@@ -908,7 +1240,7 @@ void SRigVMExecutionStackView::HandleGoToInstruction()
 			.Value(Index)
 			.MinDesiredValueWidth(250)
 			.MinValue(0)
-			.MaxValue(Operators.Num() - 1)
+			.MaxValue(RootEntries.Num() - 1)
 			.OnValueChanged_Lambda([&Index](const int32& InIndex)
 			{
 				Index = InIndex;
@@ -923,8 +1255,11 @@ void SRigVMExecutionStackView::HandleGoToInstruction()
 
 	if(OptionsDialog->ShowModal() == 0)
 	{
-		TreeView->SetSelection(Operators[Index]);
-		TreeView->SetScrollOffset(static_cast<float>(FMath::Max(Index-5, 0)));
+		if (const TSharedPtr<FRigStackEntry>* EntryPtr = InstructionToEntry.Find(Index))
+		{
+			TreeView->SetSelection(*EntryPtr);
+			TreeView->SetScrollOffset(static_cast<float>(FMath::Max(Index-5, 0)));
+		}
 	}
 }
 
@@ -955,60 +1290,101 @@ TArray<TSharedPtr<FRigStackEntry>> SRigVMExecutionStackView::GetTargetItems(cons
 	const FRigVMByteCode& ByteCode = Host->GetVM()->GetByteCode();
 	const FRigVMInstructionArray Instructions = ByteCode.GetInstructions();
 
-	TArray<int32> TargetInstructionIndices;
+	TArray<TSharedPtr<FRigStackEntry>> TempTargetItems;
 	const TArray<TSharedPtr<FRigStackEntry>> SelectedItems = TreeView->GetSelectedItems();
 	for(const TSharedPtr<FRigStackEntry>& SelectedItem : SelectedItems)
 	{
-		if(!Instructions.IsValidIndex(SelectedItem->InstructionIndex))
+		if(!Instructions.IsValidIndex(SelectedItem->SubjectIndex))
 		{
 			continue;
 		}
 		
-		const FRigVMInstruction& Instruction = Instructions[SelectedItem->InstructionIndex];
+		const FRigVMInstruction& Instruction = Instructions[SelectedItem->SubjectIndex];
 		switch(SelectedItem->OpCode)
 		{
 			case ERigVMOpCode::JumpAbsolute:
 			{
 				const FRigVMJumpOp& Op = ByteCode.GetOpAt<FRigVMJumpOp>(Instruction);
-				TargetInstructionIndices.AddUnique(Op.InstructionIndex);
+				if(const TSharedPtr<FRigStackEntry>* EntryPtr = InstructionToEntry.Find(Op.InstructionIndex))
+				{
+					TargetItems.Add(*EntryPtr);
+				}
 				break;
 			}
 			case ERigVMOpCode::JumpAbsoluteIf:
 			{
 				const FRigVMJumpIfOp& Op = ByteCode.GetOpAt<FRigVMJumpIfOp>(Instruction);
-				TargetInstructionIndices.AddUnique(Op.InstructionIndex);
+				if(const TSharedPtr<FRigStackEntry>* EntryPtr = InstructionToEntry.Find(Op.InstructionIndex))
+				{
+					TargetItems.Add(*EntryPtr);
+				}
 				break;
 			}
 			case ERigVMOpCode::JumpForward:
 			{
 				const FRigVMJumpOp& Op = ByteCode.GetOpAt<FRigVMJumpOp>(Instruction);
-				TargetInstructionIndices.AddUnique(SelectedItem->InstructionIndex + Op.InstructionIndex);
+				if(const TSharedPtr<FRigStackEntry>* EntryPtr = InstructionToEntry.Find(SelectedItem->SubjectIndex + Op.InstructionIndex))
+				{
+					TargetItems.Add(*EntryPtr);
+				}
 				break;
 			}
 			case ERigVMOpCode::JumpForwardIf:
 			{
 				const FRigVMJumpIfOp& Op = ByteCode.GetOpAt<FRigVMJumpIfOp>(Instruction);
-				TargetInstructionIndices.AddUnique(SelectedItem->InstructionIndex + Op.InstructionIndex);
+				if(const TSharedPtr<FRigStackEntry>* EntryPtr = InstructionToEntry.Find(SelectedItem->SubjectIndex + Op.InstructionIndex))
+				{
+					TargetItems.Add(*EntryPtr);
+				}
 				break;
 			}
 			case ERigVMOpCode::JumpBackward:
 			{
 				const FRigVMJumpOp& Op = ByteCode.GetOpAt<FRigVMJumpOp>(Instruction);
-				TargetInstructionIndices.AddUnique(SelectedItem->InstructionIndex - Op.InstructionIndex);
+				if(const TSharedPtr<FRigStackEntry>* EntryPtr = InstructionToEntry.Find(SelectedItem->SubjectIndex - Op.InstructionIndex))
+				{
+					TargetItems.Add(*EntryPtr);
+				}
 				break;
 			}
 			case ERigVMOpCode::JumpBackwardIf:
 			{
 				const FRigVMJumpIfOp& Op = ByteCode.GetOpAt<FRigVMJumpIfOp>(Instruction);
-				TargetInstructionIndices.AddUnique(SelectedItem->InstructionIndex - Op.InstructionIndex);
+				if(const TSharedPtr<FRigStackEntry>* EntryPtr = InstructionToEntry.Find(SelectedItem->SubjectIndex - Op.InstructionIndex))
+				{
+					TargetItems.Add(*EntryPtr);
+				}
 				break;
 			}
 			case ERigVMOpCode::RunInstructions:
 			{
 				const FRigVMRunInstructionsOp& Op = ByteCode.GetOpAt<FRigVMRunInstructionsOp>(Instruction);
-				for(int32 Index = Op.StartInstruction; Index <= Op.EndInstruction; Index++)
+				for(int32 Index = Op.StartInstruction; Index <= FMath::Max(Op.StartInstruction, Op.EndInstruction); Index++)
 				{
-					TargetInstructionIndices.AddUnique(Index);
+					if(const TSharedPtr<FRigStackEntry>* EntryPtr = InstructionToEntry.Find(Index))
+					{
+						TargetItems.Add(*EntryPtr);
+					}
+				}
+				for (int32 BlockIndex = 0; BlockIndex < Blocks.Num(); BlockIndex++)
+				{
+					if (Blocks[BlockIndex].Get<0>() == Op.StartInstruction && Blocks[BlockIndex].Get<1>() == Op.EndInstruction)
+					{
+						if(const TSharedPtr<FRigStackEntry>* EntryPtr = BlockToEntry.Find(Blocks[BlockIndex]))
+						{
+							TargetItems.Add(*EntryPtr);
+						}
+					}
+				}
+				break;
+			}
+			case ERigVMOpCode::InvokeCallable:
+			{
+				const FRigVMInvokeCallableOp& Op = ByteCode.GetOpAt<FRigVMInvokeCallableOp>(Instruction);
+				check(ByteCode.IsValidCallableIndex(Op.CallableIndex));
+				if(const TSharedPtr<FRigStackEntry>* EntryPtr = CallableToEntry.Find(Op.CallableIndex))
+				{
+					TargetItems.Add(*EntryPtr);
 				}
 				break;
 			}
@@ -1019,20 +1395,6 @@ TArray<TSharedPtr<FRigStackEntry>> SRigVMExecutionStackView::GetTargetItems(cons
 		}
 	}
 
-	TMap<int32, TSharedPtr<FRigStackEntry>> InstructionToEntry;
-	for(const TSharedPtr<FRigStackEntry>& Entry : Operators)
-	{
-		InstructionToEntry.Add(Entry->InstructionIndex, Entry);
-	}
-
-	for(const int32 TargetInstructionIndex : TargetInstructionIndices)
-	{
-		if(const TSharedPtr<FRigStackEntry>* EntryPtr = InstructionToEntry.Find(TargetInstructionIndex))
-		{
-			TargetItems.Add(*EntryPtr);
-		}
-	}
-
 	return TargetItems;
 }
 
@@ -1053,7 +1415,7 @@ void SRigVMExecutionStackView::UpdateTargetItemHighlighting()
 void SRigVMExecutionStackView::OnVMCompiled(UObject* InCompiledObject, URigVM* InCompiledVM, FRigVMExtendedExecuteContext& InVMContext)
 {
 	BindHostDelegates();
-	RefreshTreeView(InCompiledVM, &InVMContext);
+	RefreshTreeView(InCompiledVM, &InVMContext, true);
 }
 
 void SRigVMExecutionStackView::OnObjectBeingDebuggedChanged(UObject* InObjectBeingDebugged)
@@ -1061,16 +1423,17 @@ void SRigVMExecutionStackView::OnObjectBeingDebuggedChanged(UObject* InObjectBei
 	BindHostDelegates(InObjectBeingDebugged);
 	if (URigVMHost* DebuggedHost = Cast<URigVMHost>(InObjectBeingDebugged))
 	{
-		RefreshTreeView(DebuggedHost->GetVM(), &DebuggedHost->GetRigVMExtendedExecuteContext());
+		RefreshTreeView(DebuggedHost->GetVM(), &DebuggedHost->GetRigVMExtendedExecuteContext(), true);
 	}
 }
 
 void SRigVMExecutionStackView::OnFilterTextChanged(const FText& SearchText)
 {
 	FilterText = SearchText;
+
 	if (URigVMHost* RigVMHost = PreviouslyDebuggedHost.Get())
 	{
-		RefreshTreeView(RigVMHost->GetVM(), &RigVMHost->GetRigVMExtendedExecuteContext());
+		RefreshTreeView(RigVMHost->GetVM(), &RigVMHost->GetRigVMExtendedExecuteContext(), false);
 	}
 }
 
@@ -1106,13 +1469,13 @@ void SRigVMExecutionStackView::HandleModifiedEvent(ERigVMGraphNotifType InNotifT
 			}		
 
 			TArray<TSharedPtr<FRigStackEntry>> SelectedItems;
-			for (TSharedPtr<FRigStackEntry>& Operator : Operators)
+			for (TSharedPtr<FRigStackEntry>& Entry : AllEntries)
 			{
 				for (const URigVMNode* Node : SelectedNodes)
 				{
-					if(Operator->Callstack.Contains(Node))
+					if(Entry->Callstack.Contains(Node))
 					{
-						SelectedItems.Add(Operator);
+						SelectedItems.Add(Entry);
 						break;
 					}
 				}
@@ -1140,16 +1503,16 @@ void SRigVMExecutionStackView::HandleHostInitializedEvent(URigVMHost* InHost, co
 	{
 		TGuardValue<bool> SuspendControllerSelection(bSuspendControllerSelection, true);
 
-		RefreshTreeView(InHost->GetVM(), &InHost->GetRigVMExtendedExecuteContext());
+		RefreshTreeView(InHost->GetVM(), &InHost->GetRigVMExtendedExecuteContext(), true);
 		OnSelectionChanged(TSharedPtr<FRigStackEntry>(), ESelectInfo::Direct);
 
-		for (TSharedPtr<FRigStackEntry>& Operator : Operators)
+		for (TSharedPtr<FRigStackEntry>& Entry : AllEntries)
 		{
-			for (TSharedPtr<FRigStackEntry>& Child : Operator->Children)
+			for (TSharedPtr<FRigStackEntry>& Child : Entry->Children)
 			{
-				if (Child->EntryType == ERigStackEntry::Warning || Child->EntryType == ERigStackEntry::Error)
+				if (Child->EntryType == ERigStackEntry::Warning || Child->EntryType == ERigStackEntry::Error || Child->EntryType == ERigStackEntry::Info)
 				{
-					TreeView->SetItemExpansion(Operator, true);
+					TreeView->SetItemExpansion(Entry, true);
 					break;
 				}
 			}
@@ -1191,7 +1554,7 @@ void SRigVMExecutionStackView::HandleHostExecutedEvent(URigVMHost* InHost, const
 {
 	auto HandleHostExecutedEventLambda = [this, InHost, InEventName]()
 	{
-		for(const TSharedPtr<FRigStackEntry>& Entry : Operators)
+		for (TSharedPtr<FRigStackEntry>& Entry : AllEntries)
 		{
 			Entry->ShowAsFadedOut.Reset();
 			Entry->VisitedCountText.Reset();
@@ -1207,11 +1570,23 @@ void SRigVMExecutionStackView::HandleHostExecutedEvent(URigVMHost* InHost, const
 			{
 				if(URigVM* VM = InHost->GetVM())
 				{
-					for(TSharedPtr<FRigStackEntry>& Entry : Operators)
+					for (TSharedPtr<FRigStackEntry>& Entry : AllEntries)
 					{
-						if(Entry->EntryType == ERigStackEntry::Operator)
+						if(Entry->EntryType == ERigStackEntry::Instruction)
+						{
+							const int32 Count = VM->GetInstructionVisitedCount(InHost->GetRigVMExtendedExecuteContext(), Entry->SubjectIndex);
+							if(Count > 0)
+							{
+								Entry->VisitedCountText = FText::FromString(FString::FromInt(Count));
+							}
+							else
+							{
+								Entry->ShowAsFadedOut = true;
+							}
+						}
+						else if(Entry->EntryType == ERigStackEntry::Callable)
 						{
-							const int32 Count = VM->GetInstructionVisitedCount(InHost->GetRigVMExtendedExecuteContext(), Entry->InstructionIndex);
+							const int32 Count = VM->GetCallableVisitedCount(InHost->GetRigVMExtendedExecuteContext(), Entry->SubjectIndex);
 							if(Count > 0)
 							{
 								Entry->VisitedCountText = FText::FromString(FString::FromInt(Count));
@@ -1229,11 +1604,24 @@ void SRigVMExecutionStackView::HandleHostExecutedEvent(URigVMHost* InHost, const
 			{
 				if(URigVM* VM = InHost->GetVM())
 				{
-					for(TSharedPtr<FRigStackEntry>& Entry : Operators)
+					for (TSharedPtr<FRigStackEntry>& Entry : AllEntries)
 					{
-						if(Entry->EntryType == ERigStackEntry::Operator)
+						if(Entry->EntryType == ERigStackEntry::Instruction)
 						{
-							const double CurrentMicroSeconds = VM->GetInstructionMicroSeconds(InHost->GetRigVMExtendedExecuteContext(), Entry->InstructionIndex);
+							const double CurrentMicroSeconds = VM->GetInstructionMicroSeconds(InHost->GetRigVMExtendedExecuteContext(), Entry->SubjectIndex);
+							Entry->MicroSeconds = Blueprint->GetRigGraphDisplaySettings().AggregateAverage(Entry->MicroSecondsFrames, Entry->MicroSeconds, CurrentMicroSeconds);
+							if(Entry->MicroSeconds > 0.0)
+							{
+								Entry->DurationText = FText::FromString(FString::Printf(TEXT("%.02f µs"), (float)Entry->MicroSeconds));
+							}
+							else
+							{
+								Entry->ShowAsFadedOut = true;
+							}
+						}
+						else if(Entry->EntryType == ERigStackEntry::Callable)
+						{
+							const double CurrentMicroSeconds = VM->GetCallableMicroSeconds(InHost->GetRigVMExtendedExecuteContext(), Entry->SubjectIndex);
 							Entry->MicroSeconds = Blueprint->GetRigGraphDisplaySettings().AggregateAverage(Entry->MicroSecondsFrames, Entry->MicroSeconds, CurrentMicroSeconds);
 							if(Entry->MicroSeconds > 0.0)
 							{
@@ -1265,7 +1653,7 @@ void SRigVMExecutionStackView::HandlePreviewHostUpdated(IRigVMEditor* InEditor)
 {
 	if (URigVMHost* RigVMHost = PreviouslyDebuggedHost.Get())
 	{
-		RefreshTreeView(RigVMHost->GetVM(), &RigVMHost->GetRigVMExtendedExecuteContext());
+		RefreshTreeView(RigVMHost->GetVM(), &RigVMHost->GetRigVMExtendedExecuteContext(), true);
 	}
 }
 
@@ -1283,7 +1671,36 @@ void SRigVMExecutionStackView::HandleItemMouseDoubleClick(TSharedPtr<FRigStackEn
 			}
 
 			const FRigVMByteCode& ByteCode = Host->GetVM()->GetByteCode();
-			if (URigVMNode* Subject = Cast<URigVMNode>(ByteCode.GetSubjectForInstruction(InItem->InstructionIndex)))
+			URigVMNode* Subject = Cast<URigVMNode>(ByteCode.GetSubjectForInstruction(InItem->SubjectIndex));
+			if (!Subject)
+			{
+				if (InItem->Callstack.Num() > 0)
+				{
+					Subject = Cast<URigVMNode>(const_cast<UObject*>(InItem->Callstack.Last()));
+				}
+			}
+
+			if (!Subject)
+			{
+				if (InItem->EntryType == ERigStackEntry::Callable)
+				{
+					if (ByteCode.IsValidCallableIndex(InItem->SubjectIndex))
+					{
+						const FRigVMCallableInfo* Callable = ByteCode.GetCallable(InItem->SubjectIndex);
+						check(Callable);
+						const FRigVMGraphFunctionHeader Header = FRigVMGraphFunctionHeader::FindGraphFunctionHeaderFromHash(Callable->Name.ToString());
+						if (Header.IsValid())
+						{
+							if (URigVMLibraryNode* LibraryNode = Cast<URigVMLibraryNode>(Header.LibraryPointer.GetNodeSoftPath().TryLoad()))
+							{
+								Subject = LibraryNode->GetEntryNode();
+							}
+						}
+					}
+				}
+			}
+			
+			if (Subject)
 			{
 				URigVMGraph* GraphToFocus = Subject->GetGraph();
 				if (GraphToFocus && GraphToFocus->GetTypedOuter<URigVMAggregateNode>())
diff --git a/Engine/Plugins/Runtime/RigVM/Source/RigVMEditor/Public/Widgets/SRigVMExecutionStackView.h b/Engine/Plugins/Runtime/RigVM/Source/RigVMEditor/Public/Widgets/SRigVMExecutionStackView.h
--- a/Engine/Plugins/Runtime/RigVM/Source/RigVMEditor/Public/Widgets/SRigVMExecutionStackView.h
+++ b/Engine/Plugins/Runtime/RigVM/Source/RigVMEditor/Public/Widgets/SRigVMExecutionStackView.h
@@ -21,10 +21,12 @@ namespace ERigStackEntry
 {
 	enum Type
 	{
-		Operator,
-		Info,
-		Warning,
 		Error,
+		Warning,
+		Info,
+		Callable,
+		Instruction,
+		Block,
 	};
 }
 
@@ -32,21 +34,35 @@ namespace ERigStackEntry
 class FRigStackEntry : public TSharedFromThis<FRigStackEntry>
 {
 public:
-	UE_API FRigStackEntry(int32 InEntryIndex, ERigStackEntry::Type InEntryType, int32 InInstructionIndex, ERigVMOpCode InOpCode, const FString& InLabel, const FRigVMASTProxy& InProxy);
+	UE_API FRigStackEntry(int32 InEntryIndex, ERigStackEntry::Type InEntryType, int32 InSubjectIndex, ERigVMOpCode InOpCode, const FString& InLabel, const FRigVMASTProxy& InProxy);
 
 	UE_DEPRECATED(5.7, "Use TSharedRef<ITableRow> MakeTreeRowWidget(const TSharedRef<STableViewBase>& InOwnerTable, TSharedRef<FRigStackEntry> InEntry, TSharedRef<FUICommandList> InCommandList, TSharedPtr<SRigVMExecutionStackView> InStackView, TWeakInterfacePtr<IRigVMEditorAssetInterface> InBlueprint); instead")
 	UE_API TSharedRef<ITableRow> MakeTreeRowWidget(const TSharedRef<STableViewBase>& InOwnerTable, TSharedRef<FRigStackEntry> InEntry, TSharedRef<FUICommandList> InCommandList, TSharedPtr<SRigVMExecutionStackView> InStackView, TWeakObjectPtr<UObject> InBlueprint);
 	UE_API TSharedRef<ITableRow> MakeTreeRowWidget(const TSharedRef<STableViewBase>& InOwnerTable, TSharedRef<FRigStackEntry> InEntry, TSharedRef<FUICommandList>
 	                                               InCommandList, TSharedPtr<SRigVMExecutionStackView> InStackView, TWeakInterfacePtr<IRigVMEditorAssetInterface> InBlueprint);
 
+	friend uint32 GetTypeHash(const FRigStackEntry& InEntry)
+	{
+		return HashCombine(
+			GetTypeHash(InEntry.EntryType),
+			GetTypeHash(InEntry.SubjectIndex),
+			GetTypeHash(InEntry.OpCode),
+			GetTypeHash(InEntry.Label)
+		);
+	}
+
+	void FilterChildren(const FString& InSearchText);
+	bool MatchesFilter(const FString& InSearchText);
+
 	int32 EntryIndex;
 	ERigStackEntry::Type EntryType;
-	int32 InstructionIndex;
+	int32 SubjectIndex;
 	FString CallPath;
 	FRigVMCallstack Callstack;
 	ERigVMOpCode OpCode;
 	FString Label;
 	TArray<TSharedPtr<FRigStackEntry>> Children;
+	TOptional<TArray<TSharedPtr<FRigStackEntry>>> FilteredChildren;
 	mutable TOptional<bool> ShowAsFadedOut;
 	mutable TOptional<FText> VisitedCountText;
 	mutable TOptional<FText> DurationText;
@@ -74,6 +90,45 @@ private:
 	FText GetDurationText() const;
 };
 
+class SRigVMExecutionStackTreeView : public STreeView<TSharedPtr<FRigStackEntry>>
+{
+	typedef STreeView<TSharedPtr<FRigStackEntry>> BaseView;
+	
+public:
+	
+	/** Save a snapshot of the internal map that tracks item expansion before tree reconstruction */
+	void SaveAndClearSparseItemInfos()
+	{
+		// Only save the info if there is something to save (do not overwrite info with an empty map)
+		if (!BaseView::SparseItemInfos.IsEmpty())
+		{
+			OldSparseItemInfos = BaseView::SparseItemInfos;
+		}
+		ClearExpandedItems();
+	}
+
+	/** Restore the expansion infos map from the saved snapshot after tree reconstruction */
+	void RestoreSparseItemInfos(TSharedPtr<FRigStackEntry> EntryPtr)
+	{
+		const uint32 NewHash = GetTypeHash(*EntryPtr.Get());
+		for (const auto& Pair : OldSparseItemInfos)
+		{
+			const uint32 OldHash = GetTypeHash(*Pair.Key.Get());
+			if (OldHash == NewHash)
+			{
+				// the SparseItemInfos now reference the new element, but keep the same expansion state
+				BaseView::SparseItemInfos.Add(EntryPtr, Pair.Value);
+				return;
+			}
+		}
+		// set default state as expanded if not found
+		BaseView::SparseItemInfos.Add(EntryPtr, FSparseItemInfo(true, false));
+	}
+
+private:
+	TSparseItemMap OldSparseItemInfos;
+};
+
 class SRigVMExecutionStackView : public SCompoundWidget
 {
 public:
@@ -92,7 +147,7 @@ public:
 protected:
 
 	/** Rebuild the tree view */
-	UE_API void RefreshTreeView(URigVM* InVM, FRigVMExtendedExecuteContext* InVMContext);
+	UE_API void RefreshTreeView(URigVM* InVM, FRigVMExtendedExecuteContext* InVMContext, bool bRebuildEntries);
 
 private:
 
@@ -143,7 +198,9 @@ private:
 	void UnbindHostDelegates();
 	void UnbindAssetDelegates();
 
-	TSharedPtr<STreeView<TSharedPtr<FRigStackEntry>>> TreeView;
+	TSharedPtr<FRigStackEntry> MakeEntry(ERigStackEntry::Type InEntryType, int32 InSubjectIndex, ERigVMOpCode InOpCode, const FString& InLabel, const FRigVMASTProxy& InProxy);
+
+	TSharedPtr<SRigVMExecutionStackTreeView> TreeView;
 
 	/** Command list we bind to */
 	TSharedPtr<FUICommandList> CommandList;
@@ -151,7 +208,13 @@ private:
 	TWeakPtr<IRigVMEditor> WeakEditor;
 	TWeakInterfacePtr<IRigVMEditorAssetInterface> WeakRigVMBlueprint;
 
-	TArray<TSharedPtr<FRigStackEntry>> Operators;
+	TArray<TSharedPtr<FRigStackEntry>> RootEntries;
+	TArray<TSharedPtr<FRigStackEntry>> FilteredRootEntries;
+	TArray<TSharedPtr<FRigStackEntry>> AllEntries;
+	TMap<int32, TSharedPtr<FRigStackEntry>> InstructionToEntry;
+	TMap<int32, TSharedPtr<FRigStackEntry>> CallableToEntry;
+	TArray<TTuple<int32,int32>> Blocks;
+	TMap<TTuple<int32,int32>, TSharedPtr<FRigStackEntry>> BlockToEntry;
 
 	FDelegateHandle OnModelModified;
 	FDelegateHandle OnHostInitializedHandle;
