Control Rig: Fix CR Runtime Asset creating empty comment nodes
#rnx

#Synced-CL 50820870

diff --git a/Engine/Plugins/Runtime/RigVM/Source/RigVMEditor/Private/Editor/RigVMCommentNodeDetailCustomization.cpp b/Engine/Plugins/Runtime/RigVM/Source/RigVMEditor/Private/Editor/RigVMCommentNodeDetailCustomization.cpp
--- a/Engine/Plugins/Runtime/RigVM/Source/RigVMEditor/Private/Editor/RigVMCommentNodeDetailCustomization.cpp
+++ b/Engine/Plugins/Runtime/RigVM/Source/RigVMEditor/Private/Editor/RigVMCommentNodeDetailCustomization.cpp
@@ -29,7 +29,7 @@ void FRigVMCommentNodeDetailCustomization::CustomizeDetails(IDetailLayoutBuilder
 		ObjectsBeingCustomized.Add(Node);
 	}
 
-	if(ObjectsBeingCustomized[0].IsValid())
+	if(!ObjectsBeingCustomized.IsEmpty() && ObjectsBeingCustomized[0].IsValid())
 	{
 		if (IRigVMEditorAssetInterface* InterfacePtr = ObjectsBeingCustomized[0]->GetImplementingOuter<IRigVMEditorAssetInterface>())
 		{
@@ -124,6 +124,14 @@ void FRigVMCommentNodeDetailCustomization::CustomizeDetails(IDetailLayoutBuilder
 	];
 }
 
+void FRigVMCommentNodeDetailCustomization::CleanupInvalidObjects() const
+{
+	ObjectsBeingCustomized = ObjectsBeingCustomized.FilterByPredicate([](const TWeakObjectPtr<URigVMCommentNode>& InNode)
+		{
+			return InNode.IsValid();
+		});
+}
+
 void FRigVMCommentNodeDetailCustomization::GetValuesFromNode(TWeakObjectPtr<URigVMCommentNode> WeakNode)
 {
 	if (TStrongObjectPtr<URigVMCommentNode> Node = WeakNode.Pin())
@@ -149,6 +157,7 @@ void FRigVMCommentNodeDetailCustomization::SetValues(TWeakObjectPtr<URigVMCommen
 
 FText FRigVMCommentNodeDetailCustomization::GetText() const
 {
+	CleanupInvalidObjects();
 	if (!ObjectsBeingCustomized.IsEmpty())
 	{
 		FString Value = ObjectsBeingCustomized[0]->GetCommentText();
@@ -170,6 +179,7 @@ FText FRigVMCommentNodeDetailCustomization::GetText() const
 
 void FRigVMCommentNodeDetailCustomization::SetText(const FText& InNewText, ETextCommit::Type InCommitType)
 {
+	CleanupInvalidObjects();
 	for (TWeakObjectPtr<URigVMCommentNode> WeakNode : ObjectsBeingCustomized)
 	{
 		GetValuesFromNode(WeakNode);
@@ -180,6 +190,7 @@ void FRigVMCommentNodeDetailCustomization::SetText(const FText& InNewText, EText
 
 FLinearColor FRigVMCommentNodeDetailCustomization::GetColor() const
 {
+	CleanupInvalidObjects();
 	if (!ObjectsBeingCustomized.IsEmpty())
 	{
 		FLinearColor Value = ObjectsBeingCustomized[0]->GetNodeColor();
@@ -214,6 +225,7 @@ void FRigVMCommentNodeDetailCustomization::OnColorPicked(FLinearColor LinearColo
 {
 	if (BlueprintBeingCustomized)
 	{
+		CleanupInvalidObjects();
 		for (TWeakObjectPtr<URigVMCommentNode> WeakNode : ObjectsBeingCustomized)
 		{
 			if (TStrongObjectPtr<URigVMCommentNode> Node = WeakNode.Pin())
@@ -228,6 +240,7 @@ void FRigVMCommentNodeDetailCustomization::OnColorPicked(FLinearColor LinearColo
 
 ECheckBoxState FRigVMCommentNodeDetailCustomization::IsShowingBubbleEnabled() const
 {
+	CleanupInvalidObjects();
 	if (!ObjectsBeingCustomized.IsEmpty())
 	{
 		bool Value = ObjectsBeingCustomized[0]->GetCommentBubbleVisible();
@@ -249,6 +262,7 @@ ECheckBoxState FRigVMCommentNodeDetailCustomization::IsShowingBubbleEnabled() co
 
 void FRigVMCommentNodeDetailCustomization::OnShowingBubbleStateChanged(ECheckBoxState InValue)
 {
+	CleanupInvalidObjects();
 	for (TWeakObjectPtr<URigVMCommentNode> WeakNode : ObjectsBeingCustomized)
 	{
 		GetValuesFromNode(WeakNode);
@@ -259,6 +273,7 @@ void FRigVMCommentNodeDetailCustomization::OnShowingBubbleStateChanged(ECheckBox
 
 ECheckBoxState FRigVMCommentNodeDetailCustomization::IsColorBubbleEnabled() const
 {
+	CleanupInvalidObjects();
 	if (!ObjectsBeingCustomized.IsEmpty())
 	{
 		bool Value = ObjectsBeingCustomized[0]->GetCommentColorBubble();
@@ -280,6 +295,7 @@ ECheckBoxState FRigVMCommentNodeDetailCustomization::IsColorBubbleEnabled() cons
 
 void FRigVMCommentNodeDetailCustomization::OnColorBubbleStateChanged(ECheckBoxState InValue)
 {
+	CleanupInvalidObjects();
 	for (TWeakObjectPtr<URigVMCommentNode> WeakNode : ObjectsBeingCustomized)
 	{
 		GetValuesFromNode(WeakNode);
@@ -290,6 +306,7 @@ void FRigVMCommentNodeDetailCustomization::OnColorBubbleStateChanged(ECheckBoxSt
 
 TOptional<int> FRigVMCommentNodeDetailCustomization::GetFontSize() const
 {
+	CleanupInvalidObjects();
 	if (!ObjectsBeingCustomized.IsEmpty())
 	{
 		TOptional<int32> Value = ObjectsBeingCustomized[0]->GetCommentFontSize();
@@ -310,6 +327,7 @@ TOptional<int> FRigVMCommentNodeDetailCustomization::GetFontSize() const
 
 void FRigVMCommentNodeDetailCustomization::OnFontSizeChanged(int32 InValue, ETextCommit::Type Arg)
 {
+	CleanupInvalidObjects();
 	for (TWeakObjectPtr<URigVMCommentNode> WeakNode : ObjectsBeingCustomized)
 	{
 		GetValuesFromNode(WeakNode);
diff --git a/Engine/Plugins/Runtime/RigVM/Source/RigVMEditor/Private/Editor/RigVMEditor.cpp b/Engine/Plugins/Runtime/RigVM/Source/RigVMEditor/Private/Editor/RigVMEditor.cpp
--- a/Engine/Plugins/Runtime/RigVM/Source/RigVMEditor/Private/Editor/RigVMEditor.cpp
+++ b/Engine/Plugins/Runtime/RigVM/Source/RigVMEditor/Private/Editor/RigVMEditor.cpp
@@ -2773,9 +2773,14 @@ void FRigVMEditorBase::OnCreateComment()
 						else
 						{
 							FSlateRect Bounds;
-							GraphEditor->GetBoundsForSelectedNodes(Bounds, 50.f);
-							FVector2D Size = Bounds.GetSize();
-							Controller->AddCommentNode(TEXT("Comment"), Bounds.GetTopLeft(), Size, GetDefault<UGraphEditorSettings>()->DefaultCommentNodeTitleColor, FString(), true, true);
+							FVector2D Size(400, 100);
+							FVector2D Location = GraphEditor->GetPasteLocation2f();
+							if (GraphEditor->GetBoundsForSelectedNodes(Bounds, 50.f))
+							{
+								Location = Bounds.GetTopLeft();
+								Size = Bounds.GetSize();
+							}
+							Controller->AddCommentNode(TEXT("Comment"), Location, Size, GetDefault<UGraphEditorSettings>()->DefaultCommentNodeTitleColor, FString(), true, true);
 						}
 						Controller->CloseUndoBracket();
 					}
diff --git a/Engine/Plugins/Runtime/RigVM/Source/RigVMEditor/Public/Editor/RigVMCommentNodeDetailCustomization.h b/Engine/Plugins/Runtime/RigVM/Source/RigVMEditor/Public/Editor/RigVMCommentNodeDetailCustomization.h
--- a/Engine/Plugins/Runtime/RigVM/Source/RigVMEditor/Public/Editor/RigVMCommentNodeDetailCustomization.h
+++ b/Engine/Plugins/Runtime/RigVM/Source/RigVMEditor/Public/Editor/RigVMCommentNodeDetailCustomization.h
@@ -36,7 +36,9 @@ public:
 private:
 
 	FRigVMEditorAssetInterfacePtr BlueprintBeingCustomized;
-	TArray<TWeakObjectPtr<URigVMCommentNode>> ObjectsBeingCustomized;
+	mutable TArray<TWeakObjectPtr<URigVMCommentNode>> ObjectsBeingCustomized;
+	
+	UE_API void CleanupInvalidObjects() const;
 
 	UE_API void GetValuesFromNode(TWeakObjectPtr<URigVMCommentNode> Node);
 	UE_API void SetValues(TWeakObjectPtr<URigVMCommentNode> Node);
