Control Rig: Fix crash when removing array element pins of type struct
#rnx

#Synced-CL 50861785

diff --git a/Engine/Plugins/Runtime/RigVM/Source/RigVMDeveloper/Private/RigVMModel/RigVMController.cpp b/Engine/Plugins/Runtime/RigVM/Source/RigVMDeveloper/Private/RigVMModel/RigVMController.cpp
--- a/Engine/Plugins/Runtime/RigVM/Source/RigVMDeveloper/Private/RigVMModel/RigVMController.cpp
+++ b/Engine/Plugins/Runtime/RigVM/Source/RigVMDeveloper/Private/RigVMModel/RigVMController.cpp
@@ -11894,23 +11894,12 @@ URigVMPin* URigVMController::InsertArrayPin(URigVMPin* ArrayPin, int32 InIndex,
 		Action.SetTitle(FString::Printf(TEXT("Insert Array Pin")));
 		GetActionStack()->BeginAction(Action);
 	}
-	
-	TFunction<void(URigVMPin*)> ClearPinPathCacheRecursively;
-	ClearPinPathCacheRecursively = [&](URigVMPin* InPin)
-		{
-			InPin->CachedPinPath.ResetCachedValue();
-			InPin->CachedPinPathWithNodePath.ResetCachedValue();
-			for (URigVMPin* SubPin : InPin->GetSubPins())
-			{
-				ClearPinPathCacheRecursively(SubPin);
-			}
-		};
 
 	for (int32 ExistingIndex = ArrayPin->GetSubPins().Num() - 1; ExistingIndex >= InIndex; ExistingIndex--)
 	{
 		URigVMPin* ExistingPin = ArrayPin->GetSubPins()[ExistingIndex];
 		RenameObject(ExistingPin, *FString::FormatAsNumber(ExistingIndex + 1));
-		ClearPinPathCacheRecursively(ExistingPin);
+		ExistingPin->ClearPinPathCache();
 
 		// temporarily force a the index based display name
 		// so that the ed graph node can pick that up during handling the notify
@@ -12062,8 +12051,7 @@ bool URigVMController::RemoveArrayPin(const FString& InArrayElementPinPath, bool
 	{
 		URigVMPin* ExistingPin = ArrayPin->GetSubPins()[ExistingIndex];
 		ExistingPin->SetNameFromIndex();
-		ExistingPin->CachedPinPath.ResetCachedValue();
-		ExistingPin->CachedPinPathWithNodePath.ResetCachedValue();
+		ExistingPin->ClearPinPathCache();
 		Notify(ERigVMGraphNotifType::PinRenamed, ExistingPin);
 	}
 
@@ -13958,8 +13946,7 @@ bool URigVMController::RenameExposedPin(const FName& InOldPinName, const FName&
 				return false;
 			}
 
-			InPin->CachedPinPath.ResetCachedValue();
-			InPin->CachedPinPathWithNodePath.ResetCachedValue();
+			InPin->ClearPinPathCache();
 			InPin->SetDisplayName(InNewName);
 			InPin->IncrementVersion();
 			
diff --git a/Engine/Plugins/Runtime/RigVM/Source/RigVMDeveloper/Private/RigVMModel/RigVMPin.cpp b/Engine/Plugins/Runtime/RigVM/Source/RigVMDeveloper/Private/RigVMModel/RigVMPin.cpp
--- a/Engine/Plugins/Runtime/RigVM/Source/RigVMDeveloper/Private/RigVMModel/RigVMPin.cpp
+++ b/Engine/Plugins/Runtime/RigVM/Source/RigVMDeveloper/Private/RigVMModel/RigVMPin.cpp
@@ -458,6 +458,16 @@ void URigVMPin::GetExposedPinChainImpl(TArray<const URigVMPin*>& OutExposedPins,
 	}		
 }
 
+void URigVMPin::ClearPinPathCache() const
+{
+	for (URigVMPin* SubPin : GetSubPins())
+	{
+		SubPin->ClearPinPathCache();
+	}
+	CachedPinPath.ResetCachedValue();
+	CachedPinPathWithNodePath.ResetCachedValue();
+}
+
 FName URigVMPin::GetDisplayName() const
 {
 	if(CachedDisplayName.IsValid())
diff --git a/Engine/Plugins/Runtime/RigVM/Source/RigVMDeveloper/Public/RigVMModel/RigVMPin.h b/Engine/Plugins/Runtime/RigVM/Source/RigVMDeveloper/Public/RigVMModel/RigVMPin.h
--- a/Engine/Plugins/Runtime/RigVM/Source/RigVMDeveloper/Public/RigVMModel/RigVMPin.h
+++ b/Engine/Plugins/Runtime/RigVM/Source/RigVMDeveloper/Public/RigVMModel/RigVMPin.h
@@ -569,6 +569,8 @@ private:
 	UE_API void IncrementVersion(bool bAffectParentPin = true, bool bAffectSubPins = true);
 
 	UE_API void GetExposedPinChainImpl(TArray<const URigVMPin*>& OutExposedPins, TArray<const URigVMPin*>& VisitedPins) const;
+	
+	UE_API void ClearPinPathCache() const;
 
 	UPROPERTY()
 	FName DisplayName;
