Control Rig: Public functions not showing up in the action context menu
#rnx

#Synced-CL 50867638

diff --git a/Engine/Plugins/Runtime/RigVM/Source/RigVM/Private/RigVMCore/RigVMGraphFunctionDefinition.cpp b/Engine/Plugins/Runtime/RigVM/Source/RigVM/Private/RigVMCore/RigVMGraphFunctionDefinition.cpp
--- a/Engine/Plugins/Runtime/RigVM/Source/RigVM/Private/RigVMCore/RigVMGraphFunctionDefinition.cpp
+++ b/Engine/Plugins/Runtime/RigVM/Source/RigVM/Private/RigVMCore/RigVMGraphFunctionDefinition.cpp
@@ -5,6 +5,11 @@
 #include "RigVMStringUtils.h"
 #include "RigVMCore/RigVMTrait.h"
 
+#if WITH_EDITOR
+#include "IAssetTools.h"
+#include "AssetRegistry/AssetRegistryModule.h"
+#endif
+
 #include UE_INLINE_GENERATED_CPP_BY_NAME(RigVMGraphFunctionDefinition)
 
 TAutoConsoleVariable<bool> CVarRigVMCompileFunctionsToCallables(
@@ -351,6 +356,26 @@ FArchive& operator<<(FArchive& Ar, FRigVMGraphFunctionIdentifier& Data)
 
 TFunction<TArray<FRigVMVariantRef>(const FGuid& InGuid)> FRigVMGraphFunctionIdentifier::GetVariantRefsByGuidFunc;
 
+#if WITH_EDITOR
+FAssetData FRigVMGraphFunctionIdentifier::GetAssetData() const
+{
+	FAssetRegistryModule& AssetRegistryModule = FModuleManager::LoadModuleChecked<FAssetRegistryModule>("AssetRegistry");
+	IAssetTools& AssetTools = IAssetTools::Get();
+	FAssetData AssetData = AssetRegistryModule.Get().GetAssetByObjectPath(HostObject);
+	if (!AssetData.IsValid())
+	{
+		AssetData = AssetRegistryModule.Get().GetAssetByObjectPath(GetNodeSoftPath().GetWithoutSubPath());
+	}
+	
+	if (!AssetData.IsValid() || !AssetTools.IsAssetVisible(AssetData))
+	{
+		return FAssetData();
+	}
+	
+	return AssetData;
+}
+#endif
+
 bool FRigVMGraphFunctionIdentifier::IsVariant() const
 {
 	return !GetVariants(false).IsEmpty();
diff --git a/Engine/Plugins/Runtime/RigVM/Source/RigVM/Public/RigVMCore/RigVMGraphFunctionIdentifier.h b/Engine/Plugins/Runtime/RigVM/Source/RigVM/Public/RigVMCore/RigVMGraphFunctionIdentifier.h
--- a/Engine/Plugins/Runtime/RigVM/Source/RigVM/Public/RigVMCore/RigVMGraphFunctionIdentifier.h
+++ b/Engine/Plugins/Runtime/RigVM/Source/RigVM/Public/RigVMCore/RigVMGraphFunctionIdentifier.h
@@ -89,6 +89,10 @@ public:
 	{
 		return GetLibraryNodePath();
 	}
+	
+#if WITH_EDITOR
+	UE_API FAssetData GetAssetData() const;
+#endif
 
 	RIGVM_API friend FArchive& operator<<(FArchive& Ar, FRigVMGraphFunctionIdentifier& Data);
 
diff --git a/Engine/Plugins/Runtime/RigVM/Source/RigVM/RigVM.Build.cs b/Engine/Plugins/Runtime/RigVM/Source/RigVM/RigVM.Build.cs
--- a/Engine/Plugins/Runtime/RigVM/Source/RigVM/RigVM.Build.cs
+++ b/Engine/Plugins/Runtime/RigVM/Source/RigVM/RigVM.Build.cs
@@ -33,6 +33,7 @@ public class RigVM : ModuleRules
                     "UnrealEd",
                     "BlueprintGraph",
 					"StructUtilsEditor",
+					"AssetTools"
 				}
             );
         }
diff --git a/Engine/Plugins/Runtime/RigVM/Source/RigVMEditor/Private/RigVMEditorModule.cpp b/Engine/Plugins/Runtime/RigVM/Source/RigVMEditor/Private/RigVMEditorModule.cpp
--- a/Engine/Plugins/Runtime/RigVM/Source/RigVMEditor/Private/RigVMEditorModule.cpp
+++ b/Engine/Plugins/Runtime/RigVM/Source/RigVMEditor/Private/RigVMEditorModule.cpp
@@ -457,8 +457,7 @@ void FRigVMEditorModule::GetTypeActions(FRigVMEditorAssetInterfacePtr RigVMBluep
 			// Starting from a library node path
 			const FString LibraryNodePath = FunctionId.GetLibraryNodePath();
 
-			FAssetData AssetData = AssetRegistryModule.Get().GetAssetByObjectPath(FunctionId.HostObject);
-			
+			FAssetData AssetData = FunctionId.GetAssetData();
 			if (AssetData.IsValid())  // Found in registry = exists on disk
 			{
 				ControlRigAssetDatas.FindOrAdd(FunctionId) = AssetData;
