Control Rig: Avoid dirtying status on loading of asset
#rnx

#Synced-CL 50867638

diff --git a/Engine/Plugins/Runtime/RigVM/Source/RigVMEditor/Private/Editor/RigVMEditor.cpp b/Engine/Plugins/Runtime/RigVM/Source/RigVMEditor/Private/Editor/RigVMEditor.cpp
--- a/Engine/Plugins/Runtime/RigVM/Source/RigVMEditor/Private/Editor/RigVMEditor.cpp
+++ b/Engine/Plugins/Runtime/RigVM/Source/RigVMEditor/Private/Editor/RigVMEditor.cpp
@@ -391,7 +391,9 @@ void FRigVMEditorBase::InitRigVMEditor(const EToolkitMode::Type Mode, const TSha
 
 	if (Blueprints.Num() > 0)
 	{
-		if(Blueprints[0]->GetAssetStatus() == RVMA_Error)
+		if(Blueprints[0]->GetAssetStatus() == RVMA_Error || 
+			Blueprints[0]->GetAssetStatus() == RVMA_Dirty || 
+			Blueprints[0]->GetAssetStatus() == RVMA_Unknown)
 		{
 			Compile();
 		}
@@ -425,7 +427,7 @@ void FRigVMEditorBase::InitRigVMEditor(const EToolkitMode::Type Mode, const TSha
 	});
 
 	bRigVMEditorInitialized = true;
-	UpdateStaleWatchedPins();
+	InRigVMBlueprint->ClearPinWatches();
 
 #if WITH_EDITOR
 	FString BlueprintName = InRigVMBlueprint->GetObject()->GetPathName();
