Control Rig: Fix local variables and input arguments not showing up in action context menu
#rnx

#Synced-CL 50904261

diff --git a/Engine/Plugins/Runtime/RigVM/Source/RigVMEditor/Private/RigVMEditorModule.cpp b/Engine/Plugins/Runtime/RigVM/Source/RigVMEditor/Private/RigVMEditorModule.cpp
--- a/Engine/Plugins/Runtime/RigVM/Source/RigVMEditor/Private/RigVMEditorModule.cpp
+++ b/Engine/Plugins/Runtime/RigVM/Source/RigVMEditor/Private/RigVMEditorModule.cpp
@@ -560,9 +560,12 @@ void FRigVMEditorModule::GetInstanceActions(FRigVMEditorAssetInterfacePtr RigVMB
 {
 	// Only register actions for ourselves
 	UClass* InClass = RigVMBlueprint.GetObject()->GetClass();
-	if (TScriptInterface<IRigVMRuntimeAssetInterface> RuntimeAsset = RigVMBlueprint->GetRuntimeAssetInterface())
+	if (!RigVMBlueprint->IsBlueprintAsset())
 	{
-		InClass = RuntimeAsset.GetObject()->GetClass();
+		if (TScriptInterface<IRigVMRuntimeAssetInterface> RuntimeAsset = RigVMBlueprint->GetRuntimeAssetInterface())
+		{
+			InClass = RuntimeAsset.GetObject()->GetClass();
+		}
 	}
 	UClass* BlueprintClass = GetRigVMAssetClass();
 	if (RigVMBlueprint->IsBlueprintAsset() && CVarRigVMBlueprintIndependentAssets->GetBool())
