Remove property path related warning

diff --git a/Engine/Plugins/Runtime/RigVM/Source/RigVM/Private/RigVMCore/RigVM.cpp b/Engine/Plugins/Runtime/RigVM/Source/RigVM/Private/RigVMCore/RigVM.cpp
--- a/Engine/Plugins/Runtime/RigVM/Source/RigVM/Private/RigVMCore/RigVM.cpp
+++ b/Engine/Plugins/Runtime/RigVM/Source/RigVM/Private/RigVMCore/RigVM.cpp
@@ -3676,14 +3676,6 @@ void URigVM::RefreshExternalPropertyPaths()
 			ExternalPropertyPaths[PropertyPathIndex] = FRigVMPropertyPath(
 				ExternalVariables[PropertyIndex].GetProperty(),
 				ExternalPropertyPathDescriptions[PropertyPathIndex].SegmentPath);
-#if WITH_EDITOR
-			if (!ExternalPropertyPaths[PropertyPathIndex].IsValid())
-			{
-#if UE_BUILD_DEBUG
-				UE_LOG(LogRigVM, Warning, TEXT("Unable to resolve property path '%s' for property '%s'."), *ExternalPropertyPathDescriptions[PropertyPathIndex].SegmentPath, *ExternalVariables[PropertyIndex].GetProperty()->GetPathName());
-#endif
-			}
-#endif
 		}
 	}
 }
