Fix RigVM crash on exit
#rnx

#Synced-CL 50908424

diff --git a/Engine/Plugins/Runtime/RigVM/Source/RigVMDeveloper/Private/RigVMDeveloperModule.cpp b/Engine/Plugins/Runtime/RigVM/Source/RigVMDeveloper/Private/RigVMDeveloperModule.cpp
--- a/Engine/Plugins/Runtime/RigVM/Source/RigVMDeveloper/Private/RigVMDeveloperModule.cpp
+++ b/Engine/Plugins/Runtime/RigVM/Source/RigVMDeveloper/Private/RigVMDeveloperModule.cpp
@@ -40,7 +40,11 @@ void FRigVMDeveloperModule::ShutdownModule()
 	{
 		KismetCompilerModule->GetCompilers().Remove(&RigVMBlueprintCompiler);
 	}
-	URigVMBuildData::Get()->TearDownRigVMGraphFunctionPointers();
+
+	if (UObjectInitialized() && !IsEngineExitRequested())
+	{
+		URigVMBuildData::Get()->TearDownRigVMGraphFunctionPointers();
+	}
 }
 
 TSharedPtr<FKismetCompilerContext> FRigVMDeveloperModule::GetRigVMCompiler(UBlueprint* BP, FCompilerResultsLog& InMessageLog, const FKismetCompilerOptions& InCompileOptions)
