RigVM: Make URigVMBuildData::TearDownRigVMGraphFunctionPointers static so it can be called from module shutdown.

#rnx

#Synced-CL 50875397

diff --git a/Engine/Plugins/Runtime/RigVM/Source/RigVMDeveloper/Private/RigVMDeveloperModule.cpp b/Engine/Plugins/Runtime/RigVM/Source/RigVMDeveloper/Private/RigVMDeveloperModule.cpp
--- a/Engine/Plugins/Runtime/RigVM/Source/RigVMDeveloper/Private/RigVMDeveloperModule.cpp
+++ b/Engine/Plugins/Runtime/RigVM/Source/RigVMDeveloper/Private/RigVMDeveloperModule.cpp
@@ -40,11 +40,7 @@ void FRigVMDeveloperModule::ShutdownModule()
 	{
 		KismetCompilerModule->GetCompilers().Remove(&RigVMBlueprintCompiler);
 	}
-
-	if (UObjectInitialized() && !IsEngineExitRequested())
-	{
-		URigVMBuildData::Get()->TearDownRigVMGraphFunctionPointers();
-	}
+	URigVMBuildData::TearDownRigVMGraphFunctionPointers();
 }
 
 TSharedPtr<FKismetCompilerContext> FRigVMDeveloperModule::GetRigVMCompiler(UBlueprint* BP, FCompilerResultsLog& InMessageLog, const FKismetCompilerOptions& InCompileOptions)
diff --git a/Engine/Plugins/Runtime/RigVM/Source/RigVMDeveloper/Public/RigVMModel/RigVMBuildData.h b/Engine/Plugins/Runtime/RigVM/Source/RigVMDeveloper/Public/RigVMModel/RigVMBuildData.h
--- a/Engine/Plugins/Runtime/RigVM/Source/RigVMDeveloper/Public/RigVMModel/RigVMBuildData.h
+++ b/Engine/Plugins/Runtime/RigVM/Source/RigVMDeveloper/Public/RigVMModel/RigVMBuildData.h
@@ -178,7 +178,7 @@ private:
 	static UE_API TArray<UClass*> FindAllRigVMAssetClasses();
 
 	UE_API void SetupRigVMGraphFunctionPointers();
-	UE_API void TearDownRigVMGraphFunctionPointers();
+	UE_API static void TearDownRigVMGraphFunctionPointers();
 	static UE_API TArray<FRigVMGraphFunctionHeader> GetFunctionHeadersForAsset(const FAssetData& InAssetData);
 	
 	static UE_API bool bInitialized;
