[RigVM/ALT]
- Fix read-lock that was taken twice and caused a deadlock if a write lock was taken in-between

#Synced-CL 50902286

diff --git a/Engine/Plugins/Runtime/RigVM/Source/RigVM/Public/RigVMCore/RigVMRegistry.h b/Engine/Plugins/Runtime/RigVM/Source/RigVM/Public/RigVMCore/RigVMRegistry.h
--- a/Engine/Plugins/Runtime/RigVM/Source/RigVM/Public/RigVMCore/RigVMRegistry.h
+++ b/Engine/Plugins/Runtime/RigVM/Source/RigVM/Public/RigVMCore/RigVMRegistry.h
@@ -864,7 +864,6 @@ public:
 	// Returns the type index given a cpp type and a type object
 	TRigVMTypeIndex GetTypeIndex(const FName& InCPPType, UObject* InCPPTypeObject) const
 	{
-		FConditionalReadScopeLock _(*this);
 		return GetTypeIndex(FRigVMTemplateArgumentType(InCPPType, InCPPTypeObject));
 	}
 
