RigVM Registry fix for user defined structs test

#rnx

#Synced-CL 50967379

diff --git a/Engine/Plugins/Runtime/RigVM/Source/RigVM/Private/RigVMCore/RigVMTemplate.cpp b/Engine/Plugins/Runtime/RigVM/Source/RigVM/Private/RigVMCore/RigVMTemplate.cpp
--- a/Engine/Plugins/Runtime/RigVM/Source/RigVM/Private/RigVMCore/RigVMTemplate.cpp
+++ b/Engine/Plugins/Runtime/RigVM/Source/RigVM/Private/RigVMCore/RigVMTemplate.cpp
@@ -1124,9 +1124,11 @@ FRigVMTemplate::FTypeMap FRigVMTemplate::GetArgumentTypesFromString_NoLock(const
 					if (TypeIndex == INDEX_NONE && RigVMTypeUtils::RequiresCPPTypeObject(TypeName))
 					{
 						UObject* CPPTypeObject = RigVMTypeUtils::ObjectFromCPPType(TypeName, true, InTypeResolver);
-						
-						FRigVMTemplateArgumentType ArgType(TypeFName, CPPTypeObject);
-						TypeIndex = InRegistry->FindOrAddType_NoLock(ArgType, false);
+						if (ensure(CPPTypeObject != nullptr))
+						{
+							FRigVMTemplateArgumentType ArgType(TypeFName, CPPTypeObject);
+							TypeIndex = InRegistry->FindOrAddType_NoLock(ArgType, false);
+						}
 					}
 					
 					if(TypeIndex != INDEX_NONE)
diff --git a/Engine/Plugins/Runtime/RigVM/Source/RigVM/Private/RigVMTypeUtils.cpp b/Engine/Plugins/Runtime/RigVM/Source/RigVM/Private/RigVMTypeUtils.cpp
--- a/Engine/Plugins/Runtime/RigVM/Source/RigVM/Private/RigVMTypeUtils.cpp
+++ b/Engine/Plugins/Runtime/RigVM/Source/RigVM/Private/RigVMTypeUtils.cpp
@@ -303,7 +303,12 @@ UObject* RigVMTypeUtils::UserDefinedTypeFromCPPType(FString& InOutCPPType, const
 							{
 								Obj = AssetData.ToSoftObjectPath().ResolveObject();
 							}
-							if(UUserDefinedStruct* UserDefinedStruct = Cast<UUserDefinedStruct>(Obj))
+							UUserDefinedStruct* UserDefinedStruct = Cast<UUserDefinedStruct>(Obj);
+							if (UserDefinedStruct == nullptr)
+							{
+								UserDefinedStruct = Cast<UUserDefinedStruct>(AssetData.GetAsset());
+							}
+							if (UserDefinedStruct)
 							{
 								CPPTypeObject = UserDefinedStruct;
 								InOutCPPType = PostProcessCPPType(OriginalTypeName, CPPTypeObject);
@@ -330,7 +335,12 @@ UObject* RigVMTypeUtils::UserDefinedTypeFromCPPType(FString& InOutCPPType, const
 					{
 						Obj = AssetData.ToSoftObjectPath().ResolveObject();
 					}
-					if(UUserDefinedStruct* UserDefinedStruct = Cast<UUserDefinedStruct>(Obj))
+					UUserDefinedStruct* UserDefinedStruct = Cast<UUserDefinedStruct>(Obj);
+					if (UserDefinedStruct == nullptr)
+					{
+						UserDefinedStruct = Cast<UUserDefinedStruct>(AssetData.GetAsset());
+					}
+					if (UserDefinedStruct)
 					{
 						const FString GuidBasedName = GetUniqueStructTypeName(UserDefinedStruct);
 						if(GuidBasedName == InOutCPPType)
@@ -356,7 +366,12 @@ UObject* RigVMTypeUtils::UserDefinedTypeFromCPPType(FString& InOutCPPType, const
 				{
 					Obj = AssetData.ToSoftObjectPath().ResolveObject();
 				}
-				if(UUserDefinedEnum* UserDefinedEnum = Cast<UUserDefinedEnum>(Obj))
+				UUserDefinedEnum* UserDefinedEnum = Cast<UUserDefinedEnum>(Obj);
+				if (UserDefinedEnum == nullptr)
+				{
+					UserDefinedEnum = Cast<UUserDefinedEnum>(AssetData.GetAsset());
+				}
+				if (UserDefinedEnum)
 				{
 					const FString EnumCPPName = CPPTypeFromEnum(UserDefinedEnum);
 					if(EnumCPPName == InOutCPPType)
