RigVM: Guard action stack controller delegate with scope lock

#rnx

#Synced-CL 51006651

diff --git a/Engine/Plugins/Runtime/RigVM/Source/RigVMDeveloper/Private/RigVMModel/RigVMController.cpp b/Engine/Plugins/Runtime/RigVM/Source/RigVMDeveloper/Private/RigVMModel/RigVMController.cpp
--- a/Engine/Plugins/Runtime/RigVM/Source/RigVMDeveloper/Private/RigVMModel/RigVMController.cpp
+++ b/Engine/Plugins/Runtime/RigVM/Source/RigVMDeveloper/Private/RigVMModel/RigVMController.cpp
@@ -142,10 +142,13 @@ URigVMController::URigVMController(const FObjectInitializer& ObjectInitializer)
 
 URigVMController::~URigVMController()
 {
-	if(URigVMActionStack* ActionStack = WeakActionStack.Get())
 	{
-		ActionStack->OnModified().Remove(ActionStackHandle);
-		ActionStackHandle.Reset();
+		UE::TScopeLock Lock(ActionStackLock);
+		if(URigVMActionStack* ActionStack = WeakActionStack.Get())
+		{
+			ActionStack->OnModified().Remove(ActionStackHandle);
+			ActionStackHandle.Reset();
+		}
 	}
 }
 
@@ -193,9 +196,12 @@ void URigVMController::SetGraph(URigVMGraph* InGraph)
 	}
 
 	// make sure we have "some" action stack. this is mainly relevant for unit tests
-	if(!WeakActionStack.IsValid())
 	{
-		SetActionStack(NewObject<URigVMActionStack>(this, TEXT("ActionStack")));
+		UE::TScopeLock Lock(ActionStackLock);
+		if(!WeakActionStack.IsValid())
+		{
+			SetActionStack_NoLock(NewObject<URigVMActionStack>(this, TEXT("ActionStack")));
+		}
 	}
 
 	HandleModifiedEvent(ERigVMGraphNotifType::GraphChanged, GetGraph(), nullptr);
@@ -24718,6 +24724,12 @@ ERigVMPinDefaultValueType URigVMController::GetDefaultValueType(const URigVMPin*
 #endif
 
 URigVMActionStack* URigVMController::GetActionStack() const
+{
+	UE::TScopeLock Lock(ActionStackLock);
+	return GetActionStack_NoLock();
+}
+
+URigVMActionStack* URigVMController::GetActionStack_NoLock() const
 {
 	if(WeakActionStack.IsValid())
 	{
@@ -24740,6 +24752,12 @@ URigVMActionStack* URigVMController::GetActionStack() const
 }
 
 void URigVMController::SetActionStack(URigVMActionStack* InActionStack)
+{
+	UE::TScopeLock Lock(ActionStackLock);
+	SetActionStack_NoLock(InActionStack);
+}
+
+void URigVMController::SetActionStack_NoLock(URigVMActionStack* InActionStack)
 {
 	if(URigVMActionStack* PreviousActionStack = WeakActionStack.Get())
 	{
diff --git a/Engine/Plugins/Runtime/RigVM/Source/RigVMDeveloper/Public/RigVMModel/RigVMController.h b/Engine/Plugins/Runtime/RigVM/Source/RigVMDeveloper/Public/RigVMModel/RigVMController.h
--- a/Engine/Plugins/Runtime/RigVM/Source/RigVMDeveloper/Public/RigVMModel/RigVMController.h
+++ b/Engine/Plugins/Runtime/RigVM/Source/RigVMDeveloper/Public/RigVMModel/RigVMController.h
@@ -1545,9 +1545,11 @@ public:
 
 	UFUNCTION(BlueprintPure, Category = RigVMController)
 	UE_API URigVMActionStack* GetActionStack() const;
+	UE_API URigVMActionStack* GetActionStack_NoLock() const;
 
 	UFUNCTION(BlueprintCallable, Category = RigVMController)
 	UE_API void SetActionStack(URigVMActionStack* InActionStack);
+	UE_API void SetActionStack_NoLock(URigVMActionStack* InActionStack);
 
 	UE_API URigVMNode* ConvertRerouteNodeToDispatch(URigVMRerouteNode* InRerouteNode, const FName& InTemplateNotation, bool bSetupUndoRedo, bool bPrintPythonCommand);
 
@@ -1692,6 +1694,7 @@ private:
 
 	static UE_API TMap<FRigVMStructPinRedirectorKey, FString> PinPathCoreRedirectors;
 	FTransactionallySafeCriticalSection PinPathCoreRedirectorsLock;
+	mutable FTransactionallySafeCriticalSection ActionStackLock;
 
 	FRigVMUnitNodeCreatedContext UnitNodeCreatedContext;
 
