RigVM: Avoid read locks completely if a write lock has been established on the same thread

#rnx

#Synced-CL 51006651

diff --git a/Engine/Plugins/Runtime/RigVM/Source/RigVM/Public/RigVMCore/RigVMRegistry.h b/Engine/Plugins/Runtime/RigVM/Source/RigVM/Public/RigVMCore/RigVMRegistry.h
--- a/Engine/Plugins/Runtime/RigVM/Source/RigVM/Public/RigVMCore/RigVMRegistry.h
+++ b/Engine/Plugins/Runtime/RigVM/Source/RigVM/Public/RigVMCore/RigVMRegistry.h
@@ -546,11 +546,14 @@ protected:
 			// First, take the actual underlying lock.
 			if (DesiredLockType == LockType_Read)
 			{
-				// Only attempt to lock if we're not already holding the write lock.
-				if (Registry->WriteLockOwnerThread != FPlatformTLS::GetCurrentThreadId())
+				// Only attempt to lock if we're not already holding a write lock.
+				if (Registry->WriteLockOwnerThread == FPlatformTLS::GetCurrentThreadId())
 				{
-					Registry->Lock.ReadLock();
+					bLockEnabled = false;
+					return;
 				}
+				
+				Registry->Lock.ReadLock();
 			}
 			else if (DesiredLockType == LockType_Write)
 			{
@@ -627,10 +630,9 @@ protected:
 				ensure(CurrentLockCount >= 0);
 				
 				// If we were already holding the write lock, we didn't apply a subsequent read lock.
-				if (Registry->WriteLockOwnerThread != FPlatformTLS::GetCurrentThreadId())
-				{
-					Registry->Lock.ReadUnlock();
-				}
+				ensure (Registry->WriteLockOwnerThread != FPlatformTLS::GetCurrentThreadId());
+
+				Registry->Lock.ReadUnlock();
 			}
 			else if (DesiredLockType == LockType_Write)
 			{
@@ -674,6 +676,8 @@ public:
 		: FConditionalReadScopeLock(FRigVMRegistry_RWLock::Get(), bInLockEnabled)
 		{
 		}
+
+		UE_NONCOPYABLE(FConditionalReadScopeLock);
 	};
 
 	class FConditionalWriteScopeLock : public FConditionalScopeLock
@@ -688,6 +692,8 @@ public:
 		: FConditionalWriteScopeLock(FRigVMRegistry_RWLock::Get(), bInLockEnabled)
 		{
 		}
+
+		UE_NONCOPYABLE(FConditionalWriteScopeLock);
 	};
 
 	class FSuspendPluginUnmounting
diff --git a/Engine/Plugins/Runtime/RigVM/Source/RigVM/Public/RigVMCore/RigVMRegistryHandle.h b/Engine/Plugins/Runtime/RigVM/Source/RigVM/Public/RigVMCore/RigVMRegistryHandle.h
--- a/Engine/Plugins/Runtime/RigVM/Source/RigVM/Public/RigVMCore/RigVMRegistryHandle.h
+++ b/Engine/Plugins/Runtime/RigVM/Source/RigVM/Public/RigVMCore/RigVMRegistryHandle.h
@@ -13,7 +13,6 @@ struct FRigVMRegistryHandle
 public:
 
 	UE_API FRigVMRegistryHandle();
-	FRigVMRegistryHandle(FRigVMRegistryHandle& InOther) = delete;
 
 	UE_API bool IsValid() const;
 	
@@ -23,8 +22,7 @@ public:
 	UE_API FRigVMRegistry_NoLock* GetRegistry();
 	UE_API const FRigVMRegistry_NoLock* GetRegistry() const;
 
-	FRigVMRegistryHandle& operator=(FRigVMRegistryHandle& InOther) = delete;
-	FRigVMRegistryHandle& operator=(const FRigVMRegistryHandle& InOther) = delete;
+	UE_NONCOPYABLE(FRigVMRegistryHandle);
 
 protected:
 	UE_API FRigVMRegistryHandle(const FRigVMRegistry_NoLock* InRegistry);
@@ -42,9 +40,11 @@ struct FRigVMRegistryHandleLock : public FRigVMRegistryHandle
 		, bLockEnabled(bInLockEnabled)
 	{
 	}
+
+	UE_NONCOPYABLE(FRigVMRegistryHandleLock);
 	
 protected:
-	const bool bLockEnabled;
+	bool bLockEnabled;
 };
 
 #undef UE_API
