RigVM: Reuse existing readlock during FRigVMTemplate::GetExecuteArguments_NoLock

#rnx

#Synced-CL 51006651

diff --git a/Engine/Plugins/Runtime/RigVM/Source/RigVM/Private/RigVMCore/RigVMTemplate.cpp b/Engine/Plugins/Runtime/RigVM/Source/RigVM/Private/RigVMCore/RigVMTemplate.cpp
--- a/Engine/Plugins/Runtime/RigVM/Source/RigVM/Private/RigVMCore/RigVMTemplate.cpp
+++ b/Engine/Plugins/Runtime/RigVM/Source/RigVM/Private/RigVMCore/RigVMTemplate.cpp
@@ -1624,7 +1624,7 @@ const TArray<FRigVMExecuteArgument>& FRigVMTemplate::GetExecuteArguments_NoLock(
 
 			ExecuteArguments = Factory->GetExecuteArguments_NoLock(InContext, InRegistry);
 		}
-		else if(const FRigVMFunction* PrimaryPermutation = GetPrimaryPermutation_NoLock(FRigVMRegistryReadLock()))
+		else if(const FRigVMFunction* PrimaryPermutation = GetPrimaryPermutation_NoLock(InRegistry))
 		{
 			if(PrimaryPermutation->Struct)
 			{
