RigVMEdGraph call PostEditChange after calling modify in RemoveNode

#tests editor, PIE
#Synced-CL 50972851

diff --git a/Engine/Plugins/Runtime/RigVM/Source/RigVMDeveloper/Private/EdGraph/RigVMEdGraph.cpp b/Engine/Plugins/Runtime/RigVM/Source/RigVMDeveloper/Private/EdGraph/RigVMEdGraph.cpp
--- a/Engine/Plugins/Runtime/RigVM/Source/RigVMDeveloper/Private/EdGraph/RigVMEdGraph.cpp
+++ b/Engine/Plugins/Runtime/RigVM/Source/RigVMDeveloper/Private/EdGraph/RigVMEdGraph.cpp
@@ -1179,6 +1179,8 @@ void URigVMEdGraph::RemoveNode(UEdGraphNode* InNode)
 
 	// this also subsequently calls NotifyGraphChanged
 	Super::RemoveNode(InNode);
+
+	PostEditChange();
 }
 
 void URigVMEdGraph::HandleVMCompiledEvent(UObject* InCompiledObject, URigVM* InVM, FRigVMExtendedExecuteContext& InContext)
