Fix for not being able to clear uobject references on rig vm properties via the details panel

#Synced-CL 51301845

diff --git a/Engine/Plugins/Runtime/RigVM/Source/RigVMDeveloper/Private/RigVMModel/RigVMController.cpp b/Engine/Plugins/Runtime/RigVM/Source/RigVMDeveloper/Private/RigVMModel/RigVMController.cpp
--- a/Engine/Plugins/Runtime/RigVM/Source/RigVMDeveloper/Private/RigVMModel/RigVMController.cpp
+++ b/Engine/Plugins/Runtime/RigVM/Source/RigVMDeveloper/Private/RigVMModel/RigVMController.cpp
@@ -10323,7 +10323,7 @@ bool URigVMController::SetPinDefaultValue(URigVMPin* InPin, const FString& InDef
 				if(URigVMPin* SubPin = InPin->FindSubPin(FString::FromInt(ElementIndex)))
 				{
 					PostProcessDefaultValue(SubPin, Elements[ElementIndex]);
-					if (!Elements[ElementIndex].IsEmpty())
+					if (!Elements[ElementIndex].IsEmpty() || SubPin->IsWildCard() || SubPin->IsUObject())
 					{
 						SetPinDefaultValue(SubPin, Elements[ElementIndex], bResizeArrays, false, false, bSetValueOnLinkedPins);
 						bSetPinDefaultValueSucceeded = true;
@@ -10349,7 +10349,7 @@ bool URigVMController::SetPinDefaultValue(URigVMPin* InPin, const FString& InDef
 				else if( !MemberValue.IsEmpty())
 				{
 					PostProcessDefaultValue(SubPin, MemberValue);
-					if (!MemberValue.IsEmpty())
+					if (!MemberValue.IsEmpty() || SubPin->IsWildCard() || SubPin->IsUObject())
 					{
 						SetPinDefaultValue(SubPin, MemberValue, bResizeArrays, false, false, bSetValueOnLinkedPins);
 						bSetPinDefaultValueSucceeded = true;
@@ -20196,7 +20196,7 @@ void URigVMController::ApplyPinState(URigVMPin* InPin, const FPinState& InPinSta
 		{
 			FString DefaultValue = InPinState.DefaultValue;
 			PostProcessDefaultValue(InPin, DefaultValue);
-			if(!DefaultValue.IsEmpty() || InPin->IsWildCard())
+			if(!DefaultValue.IsEmpty() || InPin->IsWildCard() || InPin->IsUObject())
 			{
 				FRigVMDefaultValueTypeGuard _(this, InPinState.DefaultValueType, true);
 				SetPinDefaultValue(InPin, DefaultValue, true, bSetupUndoRedo, false);
