RigVM: Variable with Guids. Added refresh variables nodes function in the RigVMClient which will not dirty the package while still refreshing the names and types of the variable nodes.

#Synced-CL 51336242

diff --git a/Engine/Plugins/Animation/ControlRig/Source/ControlRig/Private/ModularRig.cpp b/Engine/Plugins/Animation/ControlRig/Source/ControlRig/Private/ModularRig.cpp
--- a/Engine/Plugins/Animation/ControlRig/Source/ControlRig/Private/ModularRig.cpp
+++ b/Engine/Plugins/Animation/ControlRig/Source/ControlRig/Private/ModularRig.cpp
@@ -1098,9 +1098,10 @@ bool UModularRig::SetModuleVariableBindings(const FName& InModuleName, const TMa
 			FRigVMExternalVariable SourceVariable;
 			if (SourceModuleName.IsEmpty())
 			{
-				if (const FProperty* Property = GetClass()->FindPropertyByName(*SourceVariableName))
+				const FRigVMExternalVariable Variable = GetVariableByName(*SourceVariableName);
+				if (const FProperty* Property = Variable.GetProperty())
 				{
-					SourceVariable = FRigVMExternalVariable::Make(Property, (UObject*)this);
+					SourceVariable = FRigVMExternalVariable::Make(Variable.GetGuid(), Property, (UObject*)this);
 				}
 			}
 			else if(const FRigModuleInstance* SourceModule = FindModule(*SourceModuleName))
diff --git a/Engine/Plugins/Experimental/UAF/UAF/Source/UAF/Private/AnimNextRigVMAsset.cpp b/Engine/Plugins/Experimental/UAF/UAF/Source/UAF/Private/AnimNextRigVMAsset.cpp
--- a/Engine/Plugins/Experimental/UAF/UAF/Source/UAF/Private/AnimNextRigVMAsset.cpp
+++ b/Engine/Plugins/Experimental/UAF/UAF/Source/UAF/Private/AnimNextRigVMAsset.cpp
@@ -142,7 +142,7 @@ TArray<FRigVMExternalVariable> UUAFRigVMAsset::GetExternalVariablesImpl(bool bFa
 			for(const FPropertyBagPropertyDesc& Desc : VariableDescs)
 			{
 				const FProperty* Property = Desc.CachedProperty;
-				FRigVMExternalVariable ExternalVariable = FRigVMExternalVariable::Make(Property, Container);
+				FRigVMExternalVariable ExternalVariable = FRigVMExternalVariable::Make(Desc.ID, Property, Container);
 				if(!ExternalVariable.IsValid())
 				{
 					UE_LOG(LogRigVM, Error, TEXT("%s: Property '%s' of type '%s' is not supported."), *GetName(), *Property->GetName(), *Property->GetCPPType());
diff --git a/Engine/Plugins/Experimental/UAF/UAF/Source/UAFEditor/Private/Common/Outliner/OutlinerDragDropOps.cpp b/Engine/Plugins/Experimental/UAF/UAF/Source/UAFEditor/Private/Common/Outliner/OutlinerDragDropOps.cpp
--- a/Engine/Plugins/Experimental/UAF/UAF/Source/UAFEditor/Private/Common/Outliner/OutlinerDragDropOps.cpp
+++ b/Engine/Plugins/Experimental/UAF/UAF/Source/UAFEditor/Private/Common/Outliner/OutlinerDragDropOps.cpp
@@ -77,7 +77,7 @@ FReply FVariableDragDropOp::DroppedOnPin(const FVector2f& ScreenPosition, const
     						uint8* Container = const_cast<uint8*>(Variable->GetPropertyBag().GetValue().GetMemory());
     						for(const FPropertyBagPropertyDesc& Desc : VariableDescs)
     						{
-    							FRigVMExternalVariable ExternalVariable = FRigVMExternalVariable::Make(Desc.CachedProperty, Container);
+    							FRigVMExternalVariable ExternalVariable = FRigVMExternalVariable::Make(Desc.ID, Desc.CachedProperty, Container);
     							if(!ExternalVariable.IsValid())
     							{
     								continue;
@@ -91,7 +91,7 @@ FReply FVariableDragDropOp::DroppedOnPin(const FVector2f& ScreenPosition, const
     			}
     			else if (const FProperty* Property = VariablesItem->PropertyPath.Get())
     			{
-    				CreatedExternalVariable = FRigVMExternalVariable::Make(Property, nullptr);
+    				CreatedExternalVariable = FRigVMExternalVariable::Make(FGuid(), Property, nullptr);
     			}
     		}
     		
diff --git a/Engine/Plugins/Experimental/UAF/UAF/Source/UAFUncookedOnly/Private/UncookedOnlyUtils.cpp b/Engine/Plugins/Experimental/UAF/UAF/Source/UAFUncookedOnly/Private/UncookedOnlyUtils.cpp
--- a/Engine/Plugins/Experimental/UAF/UAF/Source/UAFUncookedOnly/Private/UncookedOnlyUtils.cpp
+++ b/Engine/Plugins/Experimental/UAF/UAF/Source/UAFUncookedOnly/Private/UncookedOnlyUtils.cpp
@@ -515,7 +515,8 @@ void FUtils::GetExternalVariables(const UUAFRigVMAssetEditorData* EditorData, TA
 		for (TFieldIterator<FProperty> It(ReferencedVariableStruct); It; ++It)
 		{
 			const FProperty* Property = *It;
-			const FRigVMExternalVariable Variable = FRigVMExternalVariable::Make(Property, nullptr);
+			const FGuid Guid = IAnimNextUncookedOnlyModule::Get().GetVariableGuidByName(Property->GetFName(), ReferencedVariableStruct);
+			const FRigVMExternalVariable Variable = FRigVMExternalVariable::Make(Guid, Property, nullptr);
 			if (!Variable.IsValid(true))
 			{
 				UE_LOG(LogRigVM, Error, TEXT("%s: Property '%s' of type '%s' is not valid."), *ReferencedVariableStruct->GetName(), *Property->GetName(), *Property->GetCPPType());
@@ -551,7 +552,7 @@ void FUtils::GetInternalVariables(const UUAFRigVMAssetEditorData* EditorData, TA
 	TArray<FRigVMExternalVariable> AssetVariables;
 	EditorData->ForEachEntryOfType<UAnimNextVariableEntry>([AssetName = EditorData->GetAssetName(), &AssetVariables, &VariableInfo](const UAnimNextVariableEntry* Entry)
 	{
-		const FRigVMExternalVariable Variable = FRigVMExternalVariable::Make(Entry->GetVariableName(), Entry->GetType().ToString(), const_cast<UObject*>(Entry->GetType().GetValueTypeObject()), Entry->Access == EAnimNextExportAccessSpecifier::Public, false);
+		const FRigVMExternalVariable Variable = FRigVMExternalVariable::Make(FGuid(), Entry->GetVariableName(), Entry->GetType().ToString(), const_cast<UObject*>(Entry->GetType().GetValueTypeObject()), Entry->Access == EAnimNextExportAccessSpecifier::Public, false);
 		if (!Variable.IsValid(true))
 		{
 			UE_LOG(LogRigVM, Error, TEXT("%s: Property '%s' of type '%s' is not valid."), *AssetName, *Entry->GetName(), *Variable.GetExtendedCPPType().ToString());
@@ -572,7 +573,7 @@ void FUtils::GetInternalVariables(const UUAFRigVMAssetEditorData* EditorData, TA
 		for (const FAnimNextProgrammaticVariable& ProgrammaticVariable : GetVariableCompileContext.GetProgrammaticVariables())
 		{
 			VariableInfo.Add({ProgrammaticVariable.GetVariableName(), ProgrammaticVariable.GetType()});
-			AssetVariables.Add(FRigVMExternalVariable::Make(ProgrammaticVariable.GetVariableName(), ProgrammaticVariable.GetType().ToString(), const_cast<UObject*>(ProgrammaticVariable.GetType().GetValueTypeObject()), false, false));
+			AssetVariables.Add(FRigVMExternalVariable::Make(FGuid(), ProgrammaticVariable.GetVariableName(), ProgrammaticVariable.GetType().ToString(), const_cast<UObject*>(ProgrammaticVariable.GetType().GetValueTypeObject()), false, false));
 		}
 	}
 
@@ -2086,6 +2087,7 @@ void FUtils::ReplaceVariableReferences(UUAFRigVMAssetEditorData* InEditorData, c
 	const UUAFRigVMAsset* FindOuterAsset = Cast<UUAFRigVMAsset>(ReferenceToFind.GetObject());
 	const FAnimNextVariableReference ReferenceToReplaceWith(SoftReferenceToReplaceWith);
 	const UUAFRigVMAsset* ReplaceOuterAsset = Cast<UUAFRigVMAsset>(ReferenceToReplaceWith.GetObject());
+	const FGuid ReferenceGuid = UE::UAF::UncookedOnly::IAnimNextUncookedOnlyModule::Get().GetVariableGuidByName(ReferenceToFind.GetName(), FindOuterAsset);
 
 	const bool bReplacingWithNone = ReferenceToReplaceWith.IsNone();
 
@@ -2183,19 +2185,19 @@ void FUtils::ReplaceVariableReferences(UUAFRigVMAssetEditorData* InEditorData, c
 		{
 			if (bReplacingWithNone)
 			{
-				Controller->OnExternalVariableRemoved(ReferenceToFind.GetName(), bSetupUndoRedo);
+				Controller->OnExternalVariableRemoved(ReferenceGuid, bSetupUndoRedo);
 			}
 			else
 			{
 				if (bVariableTypeChanged)
 				{
 					const FRigVMTemplateArgumentType NewRigVMType = ReplaceReferenceType.ToRigVMTemplateArgument();
-					Controller->OnExternalVariableTypeChanged(ReferenceToFind.GetName(), NewRigVMType.CPPType.ToString(), NewRigVMType.CPPTypeObject.Get(), bSetupUndoRedo);
+					Controller->OnExternalVariableTypeChanged(ReferenceGuid, NewRigVMType.CPPType.ToString(), NewRigVMType.CPPTypeObject.Get(), bSetupUndoRedo);
 				}
 
 				if (bVariableNameChanged)
 				{
-					Controller->OnExternalVariableRenamed(ReferenceToFind.GetName(), ReferenceToReplaceWith.GetName(), bSetupUndoRedo);
+					Controller->OnExternalVariableRenamed(ReferenceGuid, ReferenceToReplaceWith.GetName(), bSetupUndoRedo);
 				}
 			}
 		}
@@ -2274,12 +2276,12 @@ void FUtils::ReplaceVariableReferences(UUAFRigVMAssetEditorData* InEditorData, c
 								if (bVariableTypeChanged)
 								{
 									const FRigVMTemplateArgumentType NewRigVMType = ReplaceReferenceType.ToRigVMTemplateArgument();
-									Controller->OnExternalVariableTypeChanged(ReferenceToFind.GetName(), NewRigVMType.CPPType.ToString(), NewRigVMType.CPPTypeObject.Get(), bSetupUndoRedo);
+									Controller->OnExternalVariableTypeChanged(ReferenceGuid, NewRigVMType.CPPType.ToString(), NewRigVMType.CPPTypeObject.Get(), bSetupUndoRedo);
 								}
 
 								if (bVariableNameChanged)
 								{
-									Controller->OnExternalVariableRenamed(ReferenceToFind.GetName(), ReferenceToReplaceWith.GetName(), bSetupUndoRedo);
+									Controller->OnExternalVariableRenamed(ReferenceGuid, ReferenceToReplaceWith.GetName(), bSetupUndoRedo);
 								}
 							}
 						}
diff --git a/Engine/Plugins/Runtime/RigVM/Source/RigVM/Private/RigVMCore/RigVMExternalVariable.cpp b/Engine/Plugins/Runtime/RigVM/Source/RigVM/Private/RigVMCore/RigVMExternalVariable.cpp
--- a/Engine/Plugins/Runtime/RigVM/Source/RigVM/Private/RigVMCore/RigVMExternalVariable.cpp
+++ b/Engine/Plugins/Runtime/RigVM/Source/RigVM/Private/RigVMCore/RigVMExternalVariable.cpp
@@ -1,6 +1,8 @@
 // Copyright Epic Games, Inc. All Rights Reserved.
 
 #include "RigVMCore/RigVMExternalVariable.h"
+
+#include "RigVMObjectVersion.h"
 #include "RigVMCore/RigVMRegistry.h"
 
 #include UE_INLINE_GENERATED_CPP_BY_NAME(RigVMExternalVariable)
@@ -10,9 +12,10 @@
 #include "EdGraphSchema_K2.h"
 #include "Engine/Blueprint.h"
 
-FRigVMExternalVariable FRigVMExternalVariable::Make(const FName& InName, const FEdGraphPinType& InPinType, bool bInPublic, bool bInReadonly)
+FRigVMExternalVariable FRigVMExternalVariable::Make(const FGuid InGuid, const FName& InName, const FEdGraphPinType& InPinType, bool bInPublic, bool bInReadonly)
 {
 	FRigVMExternalVariable ExternalVariable;
+	ExternalVariable.Guid = InGuid;
 	ExternalVariable.Name = InName;
 	ExternalVariable.bIsPublic = bInPublic;
 	ExternalVariable.bIsReadOnly = bInReadonly;
@@ -103,7 +106,7 @@ FRigVMExternalVariable FRigVMExternalVariable::Make(const FName& InName, const F
 
 #endif
 
-FRigVMExternalVariable FRigVMExternalVariable::Make(const FName& InName, const FString& InCPPTypePath, bool bInPublic, bool bInReadonly)
+FRigVMExternalVariable FRigVMExternalVariable::Make(const FGuid InGuid, const FName& InName, const FString& InCPPTypePath, bool bInPublic, bool bInReadonly)
 {
 	FString BaseCPPTypePath = InCPPTypePath;
 	if (RigVMTypeUtils::IsArrayType(InCPPTypePath))
@@ -112,10 +115,10 @@ FRigVMExternalVariable FRigVMExternalVariable::Make(const FName& InName, const F
 	}
 
 	UObject* CPPTypeObject = RigVMTypeUtils::FindObjectFromCPPTypeObjectPath(BaseCPPTypePath);
-	return Make(InName, InCPPTypePath, CPPTypeObject, bInPublic, bInReadonly);
+	return Make(InGuid, InName, InCPPTypePath, CPPTypeObject, bInPublic, bInReadonly);
 }
 
-FRigVMExternalVariable FRigVMExternalVariable::Make(const FName& InName, const FString& InCPPType, UObject* InCPPTypeObject, bool bInPublic, bool bInReadonly)
+FRigVMExternalVariable FRigVMExternalVariable::Make(const FGuid InGuid, const FName& InName, const FString& InCPPType, UObject* InCPPTypeObject, bool bInPublic, bool bInReadonly)
 {
 	FRigVMExternalVariable Variable;
 	if (InCPPType.StartsWith(TEXT("TMap<")))
@@ -123,6 +126,7 @@ FRigVMExternalVariable FRigVMExternalVariable::Make(const FName& InName, const F
 		return Variable;
 	}
 	
+	Variable.Guid = InGuid;
 	Variable.Name = InName;
 	Variable.bIsPublic = bInPublic;
 	Variable.bIsReadOnly = bInReadonly;
@@ -153,7 +157,8 @@ FRigVMExternalVariable RigVMTypeUtils::ExternalVariableFromBPVariableDescription
 {
 	const bool bIsPublic = !((InVariableDescription.PropertyFlags & CPF_DisableEditOnInstance) == CPF_DisableEditOnInstance);
 	const bool bIsReadOnly = ((InVariableDescription.PropertyFlags & CPF_BlueprintReadOnly) == CPF_BlueprintReadOnly);
-	return ExternalVariableFromPinType(InVariableDescription.VarName, InVariableDescription.VarType, bIsPublic, bIsReadOnly);
+	FRigVMExternalVariable Variable =  ExternalVariableFromPinType(InVariableDescription.VarGuid, InVariableDescription.VarName, InVariableDescription.VarType, bIsPublic, bIsReadOnly);
+	return Variable;
 }
 
 FRigVMExternalVariable RigVMTypeUtils::ExternalVariableFromBPVariableDescription(const FBPVariableDescription& InVariableDescription, const UStruct* VariablesStruct, void* Container)
@@ -161,7 +166,7 @@ FRigVMExternalVariable RigVMTypeUtils::ExternalVariableFromBPVariableDescription
 	const bool bIsPublic = !((InVariableDescription.PropertyFlags & CPF_DisableEditOnInstance) == CPF_DisableEditOnInstance);
 	const bool bIsReadOnly = ((InVariableDescription.PropertyFlags & CPF_BlueprintReadOnly) == CPF_BlueprintReadOnly);
 	
-	FRigVMExternalVariable ExternalVariable = ExternalVariableFromPinType(InVariableDescription.VarName, InVariableDescription.VarType, bIsPublic, bIsReadOnly);
+	FRigVMExternalVariable ExternalVariable = ExternalVariableFromPinType(InVariableDescription.VarGuid, InVariableDescription.VarName, InVariableDescription.VarType, bIsPublic, bIsReadOnly);
 	
 	if (VariablesStruct)
 	{
@@ -179,19 +184,20 @@ FRigVMExternalVariable RigVMTypeUtils::ExternalVariableFromBPVariableDescription
 	return ExternalVariable;
 }
 
-FRigVMExternalVariable RigVMTypeUtils::ExternalVariableFromPinType(const FName& InName, const FEdGraphPinType& InPinType, bool bInPublic, bool bInReadonly)
+FRigVMExternalVariable RigVMTypeUtils::ExternalVariableFromPinType(const FGuid InGuid, const FName& InName, const FEdGraphPinType& InPinType, bool bInPublic, bool bInReadonly)
 {
-	return FRigVMExternalVariable::Make(InName, InPinType, bInPublic, bInReadonly);
+	return FRigVMExternalVariable::Make(InGuid, InName, InPinType, bInPublic, bInReadonly);
 }
 
-FRigVMExternalVariable RigVMTypeUtils::ExternalVariableFromCPPTypePath(const FName& InName, const FString& InCPPTypePath, bool bInPublic, bool bInReadonly)
+FRigVMExternalVariable RigVMTypeUtils::ExternalVariableFromCPPTypePath(const FGuid InGuid, const FName& InName, const FString& InCPPTypePath, bool bInPublic, bool bInReadonly)
 {
-	return FRigVMExternalVariable::Make(InName, InCPPTypePath, bInPublic, bInReadonly);
+	return FRigVMExternalVariable::Make(InGuid, InName, InCPPTypePath, bInPublic, bInReadonly);
 }
 
-FRigVMExternalVariable RigVMTypeUtils::ExternalVariableFromCPPType(const FName& InName, const FString& InCPPType, UObject* InCPPTypeObject, bool bInPublic, bool bInReadonly)
+FRigVMExternalVariable RigVMTypeUtils::ExternalVariableFromCPPType(const FGuid InGuid, const FName& InName, const FString& InCPPType, UObject* InCPPTypeObject, bool bInPublic, bool
+                                                                   bInReadonly)
 {
-	return FRigVMExternalVariable::Make(InName, InCPPType, InCPPTypeObject, bInPublic, bInReadonly);
+	return FRigVMExternalVariable::Make(InGuid, InName, InCPPType, InCPPTypeObject, bInPublic, bInReadonly);
 }
 
 #endif // WITH_EDITOR
@@ -219,6 +225,16 @@ RIGVM_API TArray<FRigVMExternalVariableDef> RigVMTypeUtils::GetExternalVariableD
 
 FArchive& operator<<(FArchive& Ar, FRigVMExternalVariableDef& Variable)
 {
+	Ar.UsingCustomVersion(FRigVMObjectVersion::GUID);
+	if (Ar.IsLoading() && Ar.CustomVer(FRigVMObjectVersion::GUID) < FRigVMObjectVersion::GuidForVariables)
+	{
+		Variable.Guid.Invalidate();
+	}
+	else
+	{
+		Ar << Variable.Guid;
+	}
+	
 	Ar << Variable.Name;
 	Ar << Variable.BaseCPPType;
 
diff --git a/Engine/Plugins/Runtime/RigVM/Source/RigVM/Private/RigVMCore/RigVMGraphFunctionDefinition.cpp b/Engine/Plugins/Runtime/RigVM/Source/RigVM/Private/RigVMCore/RigVMGraphFunctionDefinition.cpp
--- a/Engine/Plugins/Runtime/RigVM/Source/RigVM/Private/RigVMCore/RigVMGraphFunctionDefinition.cpp
+++ b/Engine/Plugins/Runtime/RigVM/Source/RigVM/Private/RigVMCore/RigVMGraphFunctionDefinition.cpp
@@ -250,7 +250,7 @@ FRigVMExternalVariable FRigVMGraphFunctionArgument::ToExternalVariable() const
 	{
 		TypeObject = CPPTypeObject.Get();
 	}
-	return FRigVMExternalVariable::Make(Name, CPPType.ToString(), TypeObject, true, false);
+	return FRigVMExternalVariable::Make(FGuid(), Name, CPPType.ToString(), TypeObject, true, false);
 }
 
 FArchive& operator<<(FArchive& Ar, FRigVMFunctionCompilationData& Data)
diff --git a/Engine/Plugins/Runtime/RigVM/Source/RigVM/Private/RigVMCore/RigVMPropertyBag.cpp b/Engine/Plugins/Runtime/RigVM/Source/RigVM/Private/RigVMCore/RigVMPropertyBag.cpp
--- a/Engine/Plugins/Runtime/RigVM/Source/RigVM/Private/RigVMCore/RigVMPropertyBag.cpp
+++ b/Engine/Plugins/Runtime/RigVM/Source/RigVM/Private/RigVMCore/RigVMPropertyBag.cpp
@@ -271,12 +271,6 @@ FRigVMPropertyBag::FRigVMPropertyBag()
 {
 }
 
-FRigVMPropertyBag::FRigVMPropertyBag(const TArray<FRigVMPropertyDescription>& InPropertyDescriptions)
-	: FInstancedPropertyBag()
-{
-	AddProperties(InPropertyDescriptions);
-}
-
 bool FRigVMPropertyBag::Serialize(FArchive& Ar)
 {
 	UE_RIGVM_ARCHIVETRACE_SCOPE(Ar, TEXT("FRigVMPropertyBag"));
@@ -398,7 +392,7 @@ void FRigVMPropertyBag::AddProperties(const TArray<FRigVMPropertyDescription>& I
 			UE_LOG(LogRigVM, Error, TEXT("Detected duplicated name in the generated properties : [%s]"), *RigVMDescriptor.Name.ToString());
 		}
 
-		const FPropertyBagPropertyDesc& PropertyBagDescriptor = GetPropertyBagDescriptor(RigVMDescriptor);
+		const FPropertyBagPropertyDesc& PropertyBagDescriptor = GeneratePropertyBagDescriptor(RigVMDescriptor, !RigVMDescriptor.Guid.IsValid());
 		if (PropertyBagDescriptor.ValueType != EPropertyBagPropertyType::None)
 		{
 			BagDescriptors.Add(PropertyBagDescriptor);
@@ -734,12 +728,6 @@ void FRigVMPropertyBag::SetDefaultValuesOnTailProperties(const TArray<FRigVMProp
 }
 //****************************************************************************
 
-FPropertyBagPropertyDesc FRigVMPropertyBag::GetPropertyBagDescriptor(const FRigVMPropertyDescription& RigVMDescriptor)
-{
-	FPropertyBagPropertyDesc Result = GeneratePropertyBagDescriptor(RigVMDescriptor);
-	return Result;
-}
-
 /*static*/ FPropertyBagPropertyDesc FRigVMPropertyBag::GeneratePropertyBagDescriptor(const FRigVMPropertyDescription& RigVMDescriptor, bool bComputeID)
 {
 	FPropertyBagPropertyDesc Result;
@@ -789,6 +777,10 @@ FPropertyBagPropertyDesc FRigVMPropertyBag::GetPropertyBagDescriptor(const FRigV
 			FBlake3Hash Hash = Builder.Finalize();
 			Result.ID = FGuid::NewGuidFromHash(Hash);
 		}
+		else
+		{
+			Result.ID = RigVMDescriptor.Guid;
+		}
 	}
 
 	return Result;
diff --git a/Engine/Plugins/Runtime/RigVM/Source/RigVM/Private/RigVMCore/RigVMVariableDescription.cpp b/Engine/Plugins/Runtime/RigVM/Source/RigVM/Private/RigVMCore/RigVMVariableDescription.cpp
--- a/Engine/Plugins/Runtime/RigVM/Source/RigVM/Private/RigVMCore/RigVMVariableDescription.cpp
+++ b/Engine/Plugins/Runtime/RigVM/Source/RigVM/Private/RigVMCore/RigVMVariableDescription.cpp
@@ -34,14 +34,16 @@ FRigVMExternalVariable RigVMVariableUtils::ExternalVariableFromRigVMVariableDesc
 		{
 			CPPTypeName = RigVMTypeUtils::ArrayTypeFromBaseType(CPPTypeName);
 		}
-	}	
-	
-	return FRigVMExternalVariable::Make(
+	}
+
+	FRigVMExternalVariable Variable = FRigVMExternalVariable::Make(
+		InVariableDescription.Guid,
 		InVariableDescription.Name,
 		CPPTypeName,
 		InVariableDescription.CPPTypeObject,
 		InVariableDescription.bPublic,
 		false);
+	return Variable;
 }
 
 FRigVMGraphVariableDescription RigVMVariableUtils::VariableDescriptionFromPropertyDesc(const FPropertyBagPropertyDesc& InPropertyDesc)
@@ -126,6 +128,7 @@ FRigVMGraphVariableDescription RigVMVariableUtils::VariableDescriptionFromProper
 	{
 		Variable.bPublic = true;
 	}
+	Variable.Guid = InPropertyDesc.ID;
 	return Variable;
 }
 
@@ -153,9 +156,10 @@ FRigVMPropertyDescription RigVMVariableUtils::PropertyDescriptionFromVariableDes
 	{
 		MetaData.Add(FBlueprintMetadata::MD_Private, TEXT("true"));
 	}
-	return FRigVMPropertyDescription(InVariable.Name, InVariable.CPPType, InVariable.CPPTypeObject, InVariable.DefaultValue, PropertyFlags, MetaData, bAllowSpacesInName);
+	FRigVMPropertyDescription Result(InVariable.Name, InVariable.CPPType, InVariable.CPPTypeObject, InVariable.DefaultValue, PropertyFlags, MetaData, bAllowSpacesInName);
 #else
-	return FRigVMPropertyDescription(InVariable.Name, InVariable.CPPType, InVariable.CPPTypeObject, InVariable.DefaultValue, PropertyFlags, bAllowSpacesInName);
+	FRigVMPropertyDescription Result(InVariable.Name, InVariable.CPPType, InVariable.CPPTypeObject, InVariable.DefaultValue, PropertyFlags, bAllowSpacesInName);
 #endif
-	
+	Result.Guid = InVariable.Guid;
+	return Result;
 }
diff --git a/Engine/Plugins/Runtime/RigVM/Source/RigVM/Private/RigVMHost.cpp b/Engine/Plugins/Runtime/RigVM/Source/RigVM/Private/RigVMHost.cpp
--- a/Engine/Plugins/Runtime/RigVM/Source/RigVM/Private/RigVMHost.cpp
+++ b/Engine/Plugins/Runtime/RigVM/Source/RigVM/Private/RigVMHost.cpp
@@ -458,11 +458,20 @@ uint8* URigVMHost::GetVariablesMemory()
 	return (uint8*) this;
 }
 
+FRigVMExternalVariable URigVMHost::GetVariableByName(const FName& InVariableName) const
+{
+	if (const FProperty* Property = GetVariableProperty(InVariableName))
+	{
+		return FRigVMExternalVariable::Make(GetVariableGuid(InVariableName), Property, const_cast<URigVMHost*>(this)->GetVariablesMemory());
+	}
+	return FRigVMExternalVariable();
+}
+
 FRigVMExternalVariable URigVMHost::GetPublicVariableByName(const FName& InVariableName) const
 {
 	if (const FProperty* Property = GetPublicVariableProperty(InVariableName))
 	{
-		return FRigVMExternalVariable::Make(Property, const_cast<URigVMHost*>(this)->GetVariablesMemory());
+		return FRigVMExternalVariable::Make(GetVariableGuid(InVariableName), Property, const_cast<URigVMHost*>(this)->GetVariablesMemory());
 	}
 	return FRigVMExternalVariable();
 }
@@ -1383,7 +1392,8 @@ TArray<FRigVMExternalVariable> URigVMHost::GetExternalVariablesImpl(bool bFallba
 				continue;
 			}
 
-			FRigVMExternalVariable ExternalVariable = FRigVMExternalVariable::Make(Property, (UObject*)this);
+			const FGuid Guid = Cast<UBlueprintGeneratedClass>(GetClass())->FindBlueprintPropertyGuidFromName(Property->GetFName());
+			FRigVMExternalVariable ExternalVariable = FRigVMExternalVariable::Make(Guid, Property, (UObject*)this);
 			if(!ExternalVariable.IsValid())
 			{
 				UE_LOG(LogRigVM, Warning, TEXT("%s: Property '%s' of type '%s' is not supported."), *GetClass()->GetName(), *Property->GetName(), *Property->GetCPPType());
diff --git a/Engine/Plugins/Runtime/RigVM/Source/RigVM/Public/RigVMCore/RigVMExternalVariable.h b/Engine/Plugins/Runtime/RigVM/Source/RigVM/Public/RigVMCore/RigVMExternalVariable.h
--- a/Engine/Plugins/Runtime/RigVM/Source/RigVM/Public/RigVMCore/RigVMExternalVariable.h
+++ b/Engine/Plugins/Runtime/RigVM/Source/RigVM/Public/RigVMCore/RigVMExternalVariable.h
@@ -29,7 +29,8 @@ public:
 
 	bool operator == (const FRigVMExternalVariableDef& Other) const
 	{
-		return Name == Other.Name
+		return Guid == Other.Guid
+			&& Name == Other.Name
 			&& Property == Other.Property
 			&& BaseCPPType == Other.BaseCPPType
 			&& CPPTypeObject == Other.CPPTypeObject
@@ -45,6 +46,16 @@ public:
 			BaseCPPType.IsValid() &&
 			!BaseCPPType.IsNone();
 	}
+	
+	const FGuid& GetGuid() const
+	{
+		return Guid;
+	}
+	
+	// void SetGuid(const FGuid& InGuid)
+	// {
+	// 	Guid = InGuid;
+	// }
 
 	const FName& GetName() const
 	{
@@ -132,6 +143,7 @@ protected:
 
 	RIGVM_API void CheckCPPTypeIntegrity() const;
 	
+	FGuid Guid = FGuid();
 	FName Name = NAME_None;
 	const FProperty* Property =  nullptr;
 	FName BaseCPPType = NAME_None;
@@ -208,13 +220,14 @@ public:
 		return ((FRigVMExternalVariableDef)*this) == ((FRigVMExternalVariableDef)Other);  
 	}
 
-	static FRigVMExternalVariable Make(const FProperty* InProperty, void* InContainer, const FName& InOptionalName = NAME_None)
+	static FRigVMExternalVariable Make(const FGuid InGuid, const FProperty* InProperty, void* InContainer, const FName& InOptionalName = NAME_None)
 	{
 		check(InProperty);
 
 		const FProperty* Property = InProperty;
 
 		FRigVMExternalVariable ExternalVariable;
+		ExternalVariable.Guid = InGuid;
 		ExternalVariable.Name = InOptionalName.IsNone() ? InProperty->GetFName() : InOptionalName;
 		ExternalVariable.Property = Property;
 		ExternalVariable.bIsPublic = !InProperty->HasAllPropertyFlags(CPF_DisableEditOnInstance);
@@ -237,9 +250,10 @@ public:
 		return ExternalVariable;
 	}
 
-	static FRigVMExternalVariable Make(const FName& InName, bool& InValue)
+	static FRigVMExternalVariable Make(const FGuid InGuid, const FName& InName, bool& InValue)
 	{
 		FRigVMExternalVariable Variable;
+		Variable.Guid = InGuid;
 		Variable.Name = InName;
 		Variable.BaseCPPType = TEXT("bool");
 		Variable.CPPTypeObject = nullptr;
@@ -249,9 +263,10 @@ public:
 		return Variable;
 	}
 
-	static FRigVMExternalVariable Make(const FName& InName, TArray<bool>& InValue)
+	static FRigVMExternalVariable Make(const FGuid InGuid, const FName& InName, TArray<bool>& InValue)
 	{
 		FRigVMExternalVariable Variable;
+		Variable.Guid = InGuid;
 		Variable.Name = InName;
 		Variable.BaseCPPType = TEXT("bool");
 		Variable.CPPTypeObject = nullptr;
@@ -261,9 +276,10 @@ public:
 		return Variable;
 	}
 
-	static FRigVMExternalVariable Make(const FName& InName, int32& InValue)
+	static FRigVMExternalVariable Make(const FGuid InGuid, const FName& InName, int32& InValue)
 	{
 		FRigVMExternalVariable Variable;
+		Variable.Guid = InGuid;
 		Variable.Name = InName;
 		Variable.BaseCPPType = TEXT("int32");
 		Variable.CPPTypeObject = nullptr;
@@ -273,9 +289,10 @@ public:
 		return Variable;
 	}
 
-	static FRigVMExternalVariable Make(const FName& InName, TArray<int32>& InValue)
+	static FRigVMExternalVariable Make(const FGuid InGuid, const FName& InName, TArray<int32>& InValue)
 	{
 		FRigVMExternalVariable Variable;
+		Variable.Guid = InGuid;
 		Variable.Name = InName;
 		Variable.BaseCPPType = TEXT("int32");
 		Variable.CPPTypeObject = nullptr;
@@ -285,9 +302,10 @@ public:
 		return Variable;
 	}
 
-	static FRigVMExternalVariable Make(const FName& InName, uint8& InValue)
+	static FRigVMExternalVariable Make(const FGuid InGuid, const FName& InName, uint8& InValue)
 	{
 		FRigVMExternalVariable Variable;
+		Variable.Guid = InGuid;
 		Variable.Name = InName;
 		Variable.BaseCPPType = TEXT("uint8");
 		Variable.CPPTypeObject = nullptr;
@@ -297,9 +315,10 @@ public:
 		return Variable;
 	}
 
-	static FRigVMExternalVariable Make(const FName& InName, TArray<uint8>& InValue)
+	static FRigVMExternalVariable Make(const FGuid InGuid, const FName& InName, TArray<uint8>& InValue)
 	{
 		FRigVMExternalVariable Variable;
+		Variable.Guid = InGuid;
 		Variable.Name = InName;
 		Variable.BaseCPPType = TEXT("uint8");
 		Variable.CPPTypeObject = nullptr;
@@ -309,9 +328,10 @@ public:
 		return Variable;
 	}
 
-	static FRigVMExternalVariable Make(const FName& InName, float& InValue)
+	static FRigVMExternalVariable Make(const FGuid InGuid, const FName& InName, float& InValue)
 	{
 		FRigVMExternalVariable Variable;
+		Variable.Guid = InGuid;
 		Variable.Name = InName;
 		Variable.BaseCPPType = TEXT("float");
 		Variable.CPPTypeObject = nullptr;
@@ -321,9 +341,10 @@ public:
 		return Variable;
 	}
 
-	static FRigVMExternalVariable Make(const FName& InName, TArray<float>& InValue)
+	static FRigVMExternalVariable Make(const FGuid InGuid, const FName& InName, TArray<float>& InValue)
 	{
 		FRigVMExternalVariable Variable;
+		Variable.Guid = InGuid;
 		Variable.Name = InName;
 		Variable.BaseCPPType = TEXT("float");
 		Variable.CPPTypeObject = nullptr;
@@ -333,9 +354,10 @@ public:
 		return Variable;
 	}
 
-	static FRigVMExternalVariable Make(const FName& InName, double& InValue)
+	static FRigVMExternalVariable Make(const FGuid InGuid, const FName& InName, double& InValue)
 	{
 		FRigVMExternalVariable Variable;
+		Variable.Guid = InGuid;
 		Variable.Name = InName;
 		Variable.BaseCPPType = TEXT("double");
 		Variable.CPPTypeObject = nullptr;
@@ -345,9 +367,10 @@ public:
 		return Variable;
 	}
 
-	static FRigVMExternalVariable Make(const FName& InName, TArray<double>& InValue)
+	static FRigVMExternalVariable Make(const FGuid InGuid, const FName& InName, TArray<double>& InValue)
 	{
 		FRigVMExternalVariable Variable;
+		Variable.Guid = InGuid;
 		Variable.Name = InName;
 		Variable.BaseCPPType = TEXT("double");
 		Variable.CPPTypeObject = nullptr;
@@ -357,9 +380,10 @@ public:
 		return Variable;
 	}
 
-	static FRigVMExternalVariable Make(const FName& InName, FString& InValue)
+	static FRigVMExternalVariable Make(const FGuid InGuid, const FName& InName, FString& InValue)
 	{
 		FRigVMExternalVariable Variable;
+		Variable.Guid = InGuid;
 		Variable.Name = InName;
 		Variable.BaseCPPType = TEXT("FString");
 		Variable.CPPTypeObject = nullptr;
@@ -369,9 +393,10 @@ public:
 		return Variable;
 	}
 
-	static FRigVMExternalVariable Make(const FName& InName, TArray<FString>& InValue)
+	static FRigVMExternalVariable Make(const FGuid InGuid, const FName& InName, TArray<FString>& InValue)
 	{
 		FRigVMExternalVariable Variable;
+		Variable.Guid = InGuid;
 		Variable.Name = InName;
 		Variable.BaseCPPType = TEXT("FString");
 		Variable.CPPTypeObject = nullptr;
@@ -381,9 +406,10 @@ public:
 		return Variable;
 	}
 
-	static FRigVMExternalVariable Make(const FName& InName, FName& InValue)
+	static FRigVMExternalVariable Make(const FGuid InGuid, const FName& InName, FName& InValue)
 	{
 		FRigVMExternalVariable Variable;
+		Variable.Guid = InGuid;
 		Variable.Name = InName;
 		Variable.BaseCPPType = TEXT("FName");
 		Variable.CPPTypeObject = nullptr;
@@ -393,9 +419,10 @@ public:
 		return Variable;
 	}
 
-	static FRigVMExternalVariable Make(const FName& InName, TArray<FName>& InValue)
+	static FRigVMExternalVariable Make(const FGuid InGuid, const FName& InName, TArray<FName>& InValue)
 	{
 		FRigVMExternalVariable Variable;
+		Variable.Guid = InGuid;
 		Variable.Name = InName;
 		Variable.BaseCPPType = TEXT("FName");
 		Variable.CPPTypeObject = nullptr;
@@ -409,9 +436,10 @@ public:
 		typename T,
 		typename TEnableIf<TIsEnum<T>::Value>::Type* = nullptr
 	>
-	static FRigVMExternalVariable Make(const FName& InName, T& InValue)
+	static FRigVMExternalVariable Make(const FGuid InGuid, const FName& InName, T& InValue)
 	{
 		FRigVMExternalVariable Variable;
+		Variable.Guid = InGuid;
 		Variable.Name = InName;
 		Variable.BaseCPPType = StaticEnum<T>()->GetFName();
 		Variable.CPPTypeObject = StaticEnum<T>();
@@ -425,9 +453,10 @@ public:
 		typename T,
 		typename TEnableIf<TIsEnum<T>::Value>::Type* = nullptr
 	>
-	static FRigVMExternalVariable Make(const FName& InName, TArray<T>& InValue)
+	static FRigVMExternalVariable Make(const FGuid InGuid, const FName& InName, TArray<T>& InValue)
 	{
 		FRigVMExternalVariable Variable;
+		Variable.Guid = InGuid;
 		Variable.Name = InName;
 		Variable.BaseCPPType = StaticEnum<T>()->GetFName();
 		Variable.CPPTypeObject = StaticEnum<T>();
@@ -441,9 +470,10 @@ public:
 		typename T,
 		typename TEnableIf<TRigVMIsBaseStructure<T>::Value, T>::Type* = nullptr
 	>
-	static FRigVMExternalVariable Make(const FName& InName, T& InValue)
+	static FRigVMExternalVariable Make(const FGuid InGuid, const FName& InName, T& InValue)
 	{
 		FRigVMExternalVariable Variable;
+		Variable.Guid = InGuid;
 		Variable.Name = InName;
 		Variable.BaseCPPType = *RigVMTypeUtils::GetUniqueStructTypeName(TBaseStructure<T>::Get());
 		Variable.CPPTypeObject = TBaseStructure<T>::Get();
@@ -457,9 +487,10 @@ public:
 		typename T,
 		typename TEnableIf<TRigVMIsBaseStructure<T>::Value, T>::Type* = nullptr
 	>
-	static FRigVMExternalVariable Make(const FName& InName, TArray<T>& InValue)
+	static FRigVMExternalVariable Make(const FGuid InGuid, const FName& InName, TArray<T>& InValue)
 	{
 		FRigVMExternalVariable Variable;
+		Variable.Guid = InGuid;
 		Variable.Name = InName;
 		Variable.BaseCPPType = *RigVMTypeUtils::GetUniqueStructTypeName(TBaseStructure<T>::Get());
 		Variable.CPPTypeObject = TBaseStructure<T>::Get();
@@ -473,9 +504,10 @@ public:
 		typename T,
 		typename TEnableIf<TModels_V<CRigVMUStruct, T>>::Type * = nullptr
 	>
-	static FRigVMExternalVariable Make(const FName& InName, T& InValue)
+	static FRigVMExternalVariable Make(const FGuid InGuid, const FName& InName, T& InValue)
 	{
 		FRigVMExternalVariable Variable;
+		Variable.Guid = InGuid;
 		Variable.Name = InName;
 		Variable.BaseCPPType = *RigVMTypeUtils::GetUniqueStructTypeName(T::StaticStruct());
 		Variable.CPPTypeObject = T::StaticStruct();
@@ -489,9 +521,10 @@ public:
 		typename T,
 		typename TEnableIf<TModels_V<CRigVMUStruct, T>>::Type * = nullptr
 	>
-	static FRigVMExternalVariable Make(const FName& InName, TArray<T>& InValue)
+	static FRigVMExternalVariable Make(const FGuid InGuid, const FName& InName, TArray<T>& InValue)
 	{
 		FRigVMExternalVariable Variable;
+		Variable.Guid = InGuid;
 		Variable.Name = InName;
 		Variable.BaseCPPType = *RigVMTypeUtils::GetUniqueStructTypeName(T::StaticStruct());
 		Variable.CPPTypeObject = T::StaticStruct();
@@ -505,9 +538,10 @@ public:
 		typename T,
 		typename TEnableIf<TModels_V<CRigVMUClass, T>>::Type * = nullptr
 	>
-	static FRigVMExternalVariable Make(const FName& InName, T& InValue)
+	static FRigVMExternalVariable Make(const FGuid InGuid, const FName& InName, T& InValue)
 	{
 		FRigVMExternalVariable Variable;
+		Variable.Guid = InGuid;
 		Variable.Name = InName;
 		Variable.BaseCPPType = *T::StaticClass()->GetStructCPPName();
 		Variable.CPPTypeObject = T::StaticClass();
@@ -521,9 +555,10 @@ public:
 		typename T,
 		typename TEnableIf<TModels_V<CRigVMUClass, T>>::Type * = nullptr
 	>
-	static FRigVMExternalVariable Make(const FName& InName, TArray<T>& InValue)
+	static FRigVMExternalVariable Make(const FGuid InGuid, const FName& InName, TArray<T>& InValue)
 	{
 		FRigVMExternalVariable Variable;
+		Variable.Guid = InGuid;
 		Variable.Name = InName;
 		Variable.BaseCPPType = *T::StaticClass()->GetStructCPPName();
 		Variable.CPPTypeObject = T::StaticClass();
@@ -613,13 +648,13 @@ public:
 
 #if WITH_EDITOR
 
-	RIGVM_API static FRigVMExternalVariable Make(const FName& InName, const FEdGraphPinType& InPinType, bool bInPublic = false, bool bInReadonly = false);
+	RIGVM_API static FRigVMExternalVariable Make(const FGuid InGuid, const FName& InName, const FEdGraphPinType& InPinType, bool bInPublic = false, bool bInReadonly = false);
 
 #endif
 	
-	RIGVM_API static FRigVMExternalVariable Make(const FName& InName, const FString& InCPPTypePath, bool bInPublic = false, bool bInReadonly = false);
+	RIGVM_API static FRigVMExternalVariable Make(const FGuid InGuid, const FName& InName, const FString& InCPPTypePath, bool bInPublic = false, bool bInReadonly = false);
 
-	RIGVM_API static FRigVMExternalVariable Make(const FName& InName, const FString& InCPPType, UObject* InCPPTypeObject, bool bInPublic = false, bool bInReadonly = false);
+	RIGVM_API static FRigVMExternalVariable Make(const FGuid InGuid, const FName& InName, const FString& InCPPType, UObject* InCPPTypeObject, bool bInPublic = false, bool bInReadonly = false);
 
 private:
 	
@@ -635,11 +670,11 @@ namespace RigVMTypeUtils
 
 	RIGVM_API FRigVMExternalVariable ExternalVariableFromBPVariableDescription(const FBPVariableDescription& InVariableDescription, const UStruct* VariablesStruct, void* Container);
 
-	RIGVM_API FRigVMExternalVariable ExternalVariableFromPinType(const FName& InName, const FEdGraphPinType& InPinType, bool bInPublic = false, bool bInReadonly = false);
+	RIGVM_API FRigVMExternalVariable ExternalVariableFromPinType(const FGuid InGuid, const FName& InName, const FEdGraphPinType& InPinType, bool bInPublic = false, bool bInReadonly = false);
 
-	RIGVM_API FRigVMExternalVariable ExternalVariableFromCPPTypePath(const FName& InName, const FString& InCPPTypePath, bool bInPublic = false, bool bInReadonly = false);
+	RIGVM_API FRigVMExternalVariable ExternalVariableFromCPPTypePath(const FGuid InGuid, const FName& InName, const FString& InCPPTypePath, bool bInPublic = false, bool bInReadonly = false);
 
-	RIGVM_API FRigVMExternalVariable ExternalVariableFromCPPType(const FName& InName, const FString& InCPPType, UObject* InCPPTypeObject, bool bInPublic = false, bool bInReadonly = false);
+	RIGVM_API FRigVMExternalVariable ExternalVariableFromCPPType(const FGuid InGuid, const FName& InName, const FString& InCPPType, UObject* InCPPTypeObject, bool bInPublic = false, bool bInReadonly = false);
 #endif // WITH_EDITOR
 
 	RIGVM_API TArray<FRigVMExternalVariableDef> GetExternalVariableDefs(const TArray<FRigVMExternalVariable>& ExternalVariables);
diff --git a/Engine/Plugins/Runtime/RigVM/Source/RigVM/Public/RigVMCore/RigVMPropertyBag.h b/Engine/Plugins/Runtime/RigVM/Source/RigVM/Public/RigVMCore/RigVMPropertyBag.h
--- a/Engine/Plugins/Runtime/RigVM/Source/RigVM/Public/RigVMCore/RigVMPropertyBag.h
+++ b/Engine/Plugins/Runtime/RigVM/Source/RigVM/Public/RigVMCore/RigVMPropertyBag.h
@@ -30,6 +30,9 @@ struct FRigVMPropertyDescription
 	
 public:
 	
+	// The guid of the property to create
+	FGuid Guid;
+	
 	// The name of the property to create
 	FName Name;
 
@@ -109,7 +112,6 @@ struct FRigVMPropertyBag : public FInstancedPropertyBag
 	GENERATED_BODY()
 
 	UE_API FRigVMPropertyBag();
-	UE_API FRigVMPropertyBag(const TArray<FRigVMPropertyDescription>& InPropertyDescriptions);
 
 	UE_API bool Serialize(FArchive& Ar);
 	UE_API void AddStructReferencedObjects(FReferenceCollector& Collector);
@@ -347,7 +349,6 @@ struct FRigVMPropertyBag : public FInstancedPropertyBag
 
 
 	// Helper function to generate a property bag desc from a FRigVMPropertyDescription
-	UE_API FPropertyBagPropertyDesc GetPropertyBagDescriptor(const FRigVMPropertyDescription& RigVMDescriptor);
 	static UE_API FPropertyBagPropertyDesc GeneratePropertyBagDescriptor(const FRigVMPropertyDescription& RigVMDescriptor, bool bComputeId = true);
 
 protected:
diff --git a/Engine/Plugins/Runtime/RigVM/Source/RigVM/Public/RigVMCore/RigVMVariableDescription.h b/Engine/Plugins/Runtime/RigVM/Source/RigVM/Public/RigVMCore/RigVMVariableDescription.h
--- a/Engine/Plugins/Runtime/RigVM/Source/RigVM/Public/RigVMCore/RigVMVariableDescription.h
+++ b/Engine/Plugins/Runtime/RigVM/Source/RigVM/Public/RigVMCore/RigVMVariableDescription.h
@@ -1,6 +1,7 @@
 // Copyright Epic Games, Inc. All Rights Reserved.
 
 #pragma once
+#include "Misc/Guid.h"
 
 #include "RigVMVariableDescription.generated.h"
 
@@ -23,8 +24,16 @@ public:
 	// comparison operator
 	bool operator ==(const FRigVMGraphVariableDescription& Other) const
 	{
+		if (Guid.IsValid() && Other.Guid.IsValid())
+		{
+			return Guid == Other.Guid;
+		}
 		return Name == Other.Name;
 	}
+	
+	// The guid of the variable
+	UPROPERTY(EditAnywhere, BlueprintReadOnly, Category = RigVMGraphVariableDescription)
+	FGuid Guid;
 
 	// The name of the variable
 	UPROPERTY(EditAnywhere, BlueprintReadOnly, Category = RigVMGraphVariableDescription)
diff --git a/Engine/Plugins/Runtime/RigVM/Source/RigVM/Public/RigVMHost.h b/Engine/Plugins/Runtime/RigVM/Source/RigVM/Public/RigVMHost.h
--- a/Engine/Plugins/Runtime/RigVM/Source/RigVM/Public/RigVMHost.h
+++ b/Engine/Plugins/Runtime/RigVM/Source/RigVM/Public/RigVMHost.h
@@ -151,6 +151,9 @@ public:
 	/** The memory where the variables are stored (for the BP-based host, the memory is this, for the BP-independent, it would be the raw ptr of the property bag) */
 	UE_API uint8* GetVariablesMemory();
 
+	/** Returns a variable given its name */
+	UE_API FRigVMExternalVariable GetVariableByName(const FName& InVariableName) const;
+	
 	/** Returns a public variable given its name */
 	UE_API FRigVMExternalVariable GetPublicVariableByName(const FName& InVariableName) const;
 
@@ -502,6 +505,20 @@ protected:
 		return nullptr;
 	}
 	
+	FGuid GetVariableGuid(const FName& InVariableName) const
+	{
+		const TArray<FRigVMExternalVariable> ExternalVariables = GetExternalVariables();
+		const FRigVMExternalVariable* Found = ExternalVariables.FindByPredicate([InVariableName](const FRigVMExternalVariable& Variable)
+			{
+				return InVariableName == Variable.GetName();
+			});
+		if (Found)
+		{
+			return Found->GetGuid();
+		}
+		return FGuid();
+	}
+	
 public:
 
 	UPROPERTY()
diff --git a/Engine/Plugins/Runtime/RigVM/Source/RigVM/Public/RigVMObjectVersion.h b/Engine/Plugins/Runtime/RigVM/Source/RigVM/Public/RigVMObjectVersion.h
--- a/Engine/Plugins/Runtime/RigVM/Source/RigVM/Public/RigVMObjectVersion.h
+++ b/Engine/Plugins/Runtime/RigVM/Source/RigVM/Public/RigVMObjectVersion.h
@@ -75,7 +75,7 @@ struct FRigVMObjectVersion
 		// The VM stores a relative seek offset to be able to skip the registry during load
 		LocalizedRegistryWithRelativeSeekOffset,
 
-		// Function arguments can now represent an input variable (an external variable passed into a funciton)
+		// Function arguments can now represent an input variable (an external variable passed into a function)
 		FunctionArgumentCanRepresentInputVariable,
 
 		// Object archive is now storing the version container
@@ -86,6 +86,9 @@ struct FRigVMObjectVersion
 
 		// Introduction of callables to the rigvm bytecode
 		RigVMCallables,
+		
+		// Referencing variables through Guids
+		GuidForVariables,
 
 		// -----<new versions can be added above this line>-------------------------------------------------
 		VersionPlusOne,
diff --git a/Engine/Plugins/Runtime/RigVM/Source/RigVMDeveloper/Private/EdGraph/RigVMEdGraphNode.cpp b/Engine/Plugins/Runtime/RigVM/Source/RigVMDeveloper/Private/EdGraph/RigVMEdGraphNode.cpp
--- a/Engine/Plugins/Runtime/RigVM/Source/RigVMDeveloper/Private/EdGraph/RigVMEdGraphNode.cpp
+++ b/Engine/Plugins/Runtime/RigVM/Source/RigVMDeveloper/Private/EdGraph/RigVMEdGraphNode.cpp
@@ -628,6 +628,7 @@ void URigVMEdGraphNode::SynchronizeGraphPinNameWithModelPin(const URigVMPin* InM
 				
 				GraphPin->PinName = *NewPinName;
 				GraphPin->PinFriendlyName = DisplayName;
+				GraphPin->bOrphanedPin = InModelPin->IsOrphanPin() ? 1 : 0;
 
 				PinPathToModelPin.Add(NewPinName, (URigVMPin*)InModelPin);
 				return true;
diff --git a/Engine/Plugins/Runtime/RigVM/Source/RigVMDeveloper/Private/EdGraph/RigVMEdGraphSchema.cpp b/Engine/Plugins/Runtime/RigVM/Source/RigVMDeveloper/Private/EdGraph/RigVMEdGraphSchema.cpp
--- a/Engine/Plugins/Runtime/RigVM/Source/RigVMDeveloper/Private/EdGraph/RigVMEdGraphSchema.cpp
+++ b/Engine/Plugins/Runtime/RigVM/Source/RigVMDeveloper/Private/EdGraph/RigVMEdGraphSchema.cpp
@@ -399,6 +399,7 @@ UEdGraphNode* FRigVMEdGraphSchemaAction_PromoteToVariable::PerformAction(UEdGrap
 		}
 
 		FRigVMExternalVariable ExternalVariable = FRigVMExternalVariable::Make(
+			FGuid::NewGuid(),
 			FromPin->GetFName(),
 			ModelPin->GetCPPType(),
 			ModelPin->GetCPPTypeObject(),
@@ -2087,7 +2088,7 @@ UEdGraphPin* URigVMEdGraphSchema::DropPinOnNode(UEdGraphNode* InTargetNode, cons
 				{
 					ensure(!Model->IsTopLevelGraph());
 
-					const FRigVMExternalVariable ExternalVar = RigVMTypeUtils::ExternalVariableFromPinType(InSourcePinName, InSourcePinType);
+					const FRigVMExternalVariable ExternalVar = RigVMTypeUtils::ExternalVariableFromPinType(FGuid(), InSourcePinName, InSourcePinType);
 					if (ExternalVar.IsValid(true /* allow null memory */))
 					{
 						if (URigVMController* Controller = RigBlueprint->GetController(Model))
@@ -2413,7 +2414,7 @@ bool URigVMEdGraphSchema::SafeDeleteNodeFromGraph(UEdGraph* Graph, UEdGraphNode*
 
 bool URigVMEdGraphSchema::CanVariableBeDropped(UEdGraph* InGraph, FProperty* InVariableToDrop) const
 {
-	const FRigVMExternalVariable ExternalVariable = FRigVMExternalVariable::Make(InVariableToDrop, nullptr);
+	const FRigVMExternalVariable ExternalVariable = FRigVMExternalVariable::Make(FGuid(), InVariableToDrop, nullptr);
 	return ExternalVariable.IsValid(true /* allow nullptr */);
 }
 
@@ -2477,7 +2478,9 @@ bool URigVMEdGraphSchema::RequestVariableDropOnPin(UEdGraph* InGraph, FProperty*
 		{
 			if (const URigVMPin* ModelPin = Graph->GetModel()->FindPin(InPin->GetName()))
 			{
-				const FRigVMExternalVariable ExternalVariable = FRigVMExternalVariable::Make(InVariableToDrop, nullptr);
+				URigVMController* Controller = Graph->GetController();
+				const FRigVMExternalVariable ExternalVariable = Controller->GetVariableByName(InVariableToDrop->GetFName());
+				ensure(ExternalVariable.GetProperty() == InVariableToDrop);
 				return RequestVariableDropOnPin(InGraph, ExternalVariable, InPin, FDeprecateSlateVector2D(InDropPosition), FDeprecateSlateVector2D(InScreenPosition));
 			}
 		}
diff --git a/Engine/Plugins/Runtime/RigVM/Source/RigVMDeveloper/Private/RigVMBlueprintLegacy.cpp b/Engine/Plugins/Runtime/RigVM/Source/RigVMDeveloper/Private/RigVMBlueprintLegacy.cpp
--- a/Engine/Plugins/Runtime/RigVM/Source/RigVMDeveloper/Private/RigVMBlueprintLegacy.cpp
+++ b/Engine/Plugins/Runtime/RigVM/Source/RigVMDeveloper/Private/RigVMBlueprintLegacy.cpp
@@ -55,7 +55,7 @@ FEdGraphPinType FRigVMOldPublicFunctionArg::GetPinType() const
 	{
 		CPPTypeObject = RigVMTypeUtils::FindObjectFromCPPTypeObjectPath(CPPTypeObjectPath.ToString());
 	}
-	FRigVMExternalVariable Variable = FRigVMExternalVariable::Make(Name, CPPType.ToString(), CPPTypeObject, false, false);
+	FRigVMExternalVariable Variable = FRigVMExternalVariable::Make(FGuid(), Name, CPPType.ToString(), CPPTypeObject, false, false);
 	return RigVMTypeUtils::PinTypeFromExternalVariable(Variable);
 }
 
@@ -1136,7 +1136,7 @@ TArray<FString> URigVMBlueprint::GeneratePythonContextCommands(const FString InN
 
 FName URigVMBlueprint::AddMemberVariable(const FName& InName, const FString& InCPPType, bool bIsPublic, bool bIsReadOnly, FString InDefaultValue)
 {
-	FRigVMExternalVariable Variable = RigVMTypeUtils::ExternalVariableFromCPPTypePath(InName, InCPPType, bIsPublic, bIsReadOnly);
+	FRigVMExternalVariable Variable = RigVMTypeUtils::ExternalVariableFromCPPTypePath(FGuid::NewGuid(), InName, InCPPType, bIsPublic, bIsReadOnly);
 	FName Result = AddHostMemberVariableFromExternal(Variable, InDefaultValue);
 	if (!Result.IsNone())
 	{
@@ -1212,7 +1212,7 @@ bool URigVMBlueprint::ChangeMemberVariableType(const FName& InName, const FStrin
 	}
 
 	UObject* CPPTypeObject = RigVMTypeUtils::FindObjectFromCPPTypeObjectPath<UScriptStruct>(BaseCPPType);
-	const FRigVMExternalVariable Variable = FRigVMExternalVariable::Make(InName, InCPPType, CPPTypeObject, bIsPublic, bIsReadOnly);
+	const FRigVMExternalVariable Variable = FRigVMExternalVariable::Make(FGuid(), InName, InCPPType, CPPTypeObject, bIsPublic, bIsReadOnly);
 	FEdGraphPinType PinType = RigVMTypeUtils::PinTypeFromExternalVariable(Variable);
 	if (!PinType.PinCategory.IsValid())
 	{
@@ -1267,12 +1267,12 @@ FName URigVMBlueprint::FindHostMemberVariableUniqueName(TSharedPtr<FKismetNameVa
 	return *KismetName;
 }
 
-int32 URigVMBlueprint::AddHostMemberVariable(URigVMBlueprint* InBlueprint, const FName& InVarName, FEdGraphPinType InVarType, bool bIsPublic, bool bIsReadOnly, FString InDefaultValue)
+int32 URigVMBlueprint::AddHostMemberVariable(URigVMBlueprint* InBlueprint, const FGuid InVarGuid, const FName& InVarName, FEdGraphPinType InVarType, bool bIsPublic, bool bIsReadOnly, FString InDefaultValue)
 {
 	FBPVariableDescription NewVar;
 
 	NewVar.VarName = InVarName;
-	NewVar.VarGuid = FGuid::NewGuid();
+	NewVar.VarGuid = InVarGuid;
 	NewVar.FriendlyName = FName::NameToDisplayString(InVarName.ToString(), (InVarType.PinCategory == UEdGraphSchema_K2::PC_Boolean) ? true : false);
 	NewVar.VarType = InVarType;
 
@@ -1317,7 +1317,7 @@ FName URigVMBlueprint::AddHostMemberVariableFromExternal(FRigVMExternalVariable
 
 	TSharedPtr<FKismetNameValidator> NameValidator = MakeShareable(new FKismetNameValidator(this, NAME_None, nullptr));
 	FName VarName = FindHostMemberVariableUniqueName(NameValidator, InVariableToCreate.GetName().ToString());
-	int32 VariableIndex = AddHostMemberVariable(this, VarName, PinType, InVariableToCreate.IsPublic(), InVariableToCreate.IsReadOnly(), InDefaultValue);
+	int32 VariableIndex = AddHostMemberVariable(this, InVariableToCreate.GetGuid(), VarName, PinType, InVariableToCreate.IsPublic(), InVariableToCreate.IsReadOnly(), InDefaultValue);
 	if (VariableIndex != INDEX_NONE)
 	{
 		FBlueprintEditorUtils::MarkBlueprintAsStructurallyModified(this);
@@ -1369,7 +1369,7 @@ void URigVMBlueprint::HandlePostVariableChange(UBlueprint* InBlueprint)
 	{
 		if (!NewVariablesByGuid.Contains(OldVariable.VarGuid))
 		{
-			OnVariableRemoved(OldVariable.VarName);
+			OnVariableRemoved(OldVariable.VarGuid);
 			continue;
 		}
 	}
@@ -1386,12 +1386,12 @@ void URigVMBlueprint::HandlePostVariableChange(UBlueprint* InBlueprint)
 		const FBPVariableDescription& OldVariable = LocalLastNewVariables[OldVarIndex];
 		if (OldVariable.VarName != NewVariable.VarName)
 		{
-			OnVariableRenamed(OldVariable.VarName, NewVariable.VarName);
+			OnVariableRenamed(OldVariable.VarGuid, NewVariable.VarName);
 		}
 
 		if (OldVariable.VarType != NewVariable.VarType)
 		{
-			OnVariableTypeChanged(NewVariable.VarName, OldVariable.VarType, NewVariable.VarType);
+			OnVariableTypeChanged(NewVariable.VarGuid, OldVariable.VarType, NewVariable.VarType);
 		}
 	}
 
diff --git a/Engine/Plugins/Runtime/RigVM/Source/RigVMDeveloper/Private/RigVMEditorAsset.cpp b/Engine/Plugins/Runtime/RigVM/Source/RigVMDeveloper/Private/RigVMEditorAsset.cpp
--- a/Engine/Plugins/Runtime/RigVM/Source/RigVMDeveloper/Private/RigVMEditorAsset.cpp
+++ b/Engine/Plugins/Runtime/RigVM/Source/RigVMDeveloper/Private/RigVMEditorAsset.cpp
@@ -246,6 +246,19 @@ void IRigVMEditorAssetInterface::RequestClientHostRigVMInit()
 	}
 }
 
+FRigVMGraphVariableDescription IRigVMEditorAssetInterface::FindAssetVariableByGuid(const FGuid& InGuid) const
+{
+	const TArray<FRigVMGraphVariableDescription> Variables = GetAssetVariables();
+	for (const FRigVMGraphVariableDescription& Variable : Variables)
+	{
+		if (Variable.Guid == InGuid)
+		{
+			return Variable;
+		}
+	}
+	return FRigVMGraphVariableDescription();
+}
+
 void IRigVMEditorAssetInterface::CollectExternalDependencies(TArray<FRigVMExternalDependency>& OutDependencies, const FName& InCategory, const FRigVMClient* InClient) const
 {
 	check(InClient);
@@ -4001,7 +4014,7 @@ void IRigVMEditorAssetInterface::OnVariableAdded(const FName& InVarName)
 	BroadcastExternalVariablesChangedEvent();
 }
 
-void IRigVMEditorAssetInterface::OnVariableRemoved(const FName& InVarName)
+void IRigVMEditorAssetInterface::OnVariableRemoved(const FGuid& InVarGuid)
 {
 	TArray<URigVMGraph*> AllGraphs = GetAllModels();
 	for (URigVMGraph* Graph : AllGraphs)
@@ -4013,18 +4026,20 @@ void IRigVMEditorAssetInterface::OnVariableRemoved(const FName& InVarName)
 #else
 			const bool bSetupUndoRedo = false;
 #endif
-			Controller->OnExternalVariableRemoved(InVarName, bSetupUndoRedo);
+			Controller->OnExternalVariableRemoved(InVarGuid, bSetupUndoRedo);
 		}
 	}
+	
+	const FRigVMGraphVariableDescription Variable = FindAssetVariableByGuid(InVarGuid);
 
 	RigVMPythonUtils::Print(GetObject()->GetFName().ToString(),
 		FString::Printf(TEXT("asset.remove_member_variable('%s')"),
-			*InVarName.ToString()));
+			*Variable.Name.ToString()));
 	
 	BroadcastExternalVariablesChangedEvent();
 }
 
-void IRigVMEditorAssetInterface::OnVariableRenamed(const FName& InOldVarName, const FName& InNewVarName)
+void IRigVMEditorAssetInterface::OnVariableRenamed(const FGuid& InVarGuid, const FName& InNewVarName)
 {
 	TArray<URigVMGraph*> AllGraphs = GetAllModels();
 	for (URigVMGraph* Graph : AllGraphs)
@@ -4036,19 +4051,20 @@ void IRigVMEditorAssetInterface::OnVariableRenamed(const FName& InOldVarName, co
 #else
 			const bool bSetupUndoRedo = false;
 #endif
-			Controller->OnExternalVariableRenamed(InOldVarName, InNewVarName, bSetupUndoRedo);
+			Controller->OnExternalVariableRenamed(InVarGuid, InNewVarName, bSetupUndoRedo);
 		}
 	}
 
+	const FRigVMGraphVariableDescription Variable = FindAssetVariableByGuid(InVarGuid);
 	RigVMPythonUtils::Print(GetObject()->GetFName().ToString(),
 		FString::Printf(TEXT("asset.rename_member_variable('%s', '%s')"),
-			*InOldVarName.ToString(),
+			*Variable.Name.ToString(),
 			*InNewVarName.ToString()));
 	
 	BroadcastExternalVariablesChangedEvent();
 }
 
-void IRigVMEditorAssetInterface::OnVariableTypeChanged(const FName& InVarName, FEdGraphPinType InOldPinType, FEdGraphPinType InNewPinType)
+void IRigVMEditorAssetInterface::OnVariableTypeChanged(const FGuid& InVarGuid, FEdGraphPinType InOldPinType, FEdGraphPinType InNewPinType)
 {
 	FString CPPType;
 	UObject* CPPTypeObject = nullptr;
@@ -4067,11 +4083,11 @@ void IRigVMEditorAssetInterface::OnVariableTypeChanged(const FName& InVarName, F
 
 			if (!CPPType.IsEmpty())
 			{
-				Controller->OnExternalVariableTypeChanged(InVarName, CPPType, CPPTypeObject, bSetupUndoRedo);
+				Controller->OnExternalVariableTypeChanged(InVarGuid, CPPType, CPPTypeObject, bSetupUndoRedo);
 			}
 			else
 			{
-				Controller->OnExternalVariableRemoved(InVarName, bSetupUndoRedo);
+				Controller->OnExternalVariableRemoved(InVarGuid, bSetupUndoRedo);
 			}
 		}
 	}
@@ -4080,7 +4096,7 @@ void IRigVMEditorAssetInterface::OnVariableTypeChanged(const FName& InVarName, F
 	{
 		for (auto Var : GetAssetVariables())
 		{
-			if (Var.Name == InVarName)
+			if (Var.Guid == InVarGuid)
 			{
 				CPPType = ScriptStruct->GetName();
 			}
@@ -4090,7 +4106,7 @@ void IRigVMEditorAssetInterface::OnVariableTypeChanged(const FName& InVarName, F
 	{
 		for (auto Var : GetAssetVariables())
 		{
-			if (Var.Name == InVarName)
+			if (Var.Guid == InVarGuid)
 			{
 				CPPType = Enum->GetName();
 			}
@@ -4099,10 +4115,12 @@ void IRigVMEditorAssetInterface::OnVariableTypeChanged(const FName& InVarName, F
 
 	// register the type in the registry
 	FRigVMRegistry::Get().FindOrAddType(FRigVMTemplateArgumentType(*CPPType, CPPTypeObject));
+	
+	const FRigVMGraphVariableDescription Variable = FindAssetVariableByGuid(InVarGuid);
 
 	RigVMPythonUtils::Print(GetObject()->GetFName().ToString(),
 		FString::Printf(TEXT("asset.change_member_variable_type('%s', '%s')"),
-		*InVarName.ToString(),
+		*Variable.Name.ToString(),
 		*CPPType));
 
 	BroadcastExternalVariablesChangedEvent();
@@ -4615,7 +4633,7 @@ FName URigVMEditorAsset::AddHostMemberVariableFromExternal(FRigVMExternalVariabl
 
 FName URigVMEditorAsset::AddMemberVariable(const FName& InName, const FString& InCPPType, bool bIsPublic, bool bIsReadOnly, FString InDefaultValue)
 {
-	FRigVMExternalVariable Variable = RigVMTypeUtils::ExternalVariableFromCPPTypePath(InName, InCPPType, bIsPublic, bIsReadOnly);
+	FRigVMExternalVariable Variable = RigVMTypeUtils::ExternalVariableFromCPPTypePath(FGuid::NewGuid(), InName, InCPPType, bIsPublic, bIsReadOnly);
 	FName Result = AddHostMemberVariableFromExternal(Variable, InDefaultValue);
 	return Result;
 }
@@ -4624,9 +4642,10 @@ bool URigVMEditorAsset::RemoveMemberVariable(const FName& InName)
 {
 	if (URigVMRuntimeAsset* RuntimeAsset = GetRuntimeAsset())
 	{
+		const FRigVMGraphVariableDescription Variable = RuntimeAsset->FindAssetVariable(InName);
 		if (RuntimeAsset->RemoveMemberVariable(InName))
 		{
-			OnVariableRemoved(InName);
+			OnVariableRemoved(Variable.Guid);
 			return true;
 		}
 	}
@@ -4637,11 +4656,17 @@ bool URigVMEditorAsset::BulkRemoveMemberVariables(const TArray<FName>& InNames)
 {
 	if (URigVMRuntimeAsset* RuntimeAsset = GetRuntimeAsset())
 	{
+		TMap<FName, FGuid> NameToGuid;
+		for (const FName& Name : InNames)
+		{
+			const FRigVMGraphVariableDescription Variable = RuntimeAsset->FindAssetVariable(Name);
+			NameToGuid.Add(Name, Variable.Guid);
+		}
 		if (RuntimeAsset->BulkRemoveMemberVariables(InNames))
 		{
 			for (const FName& Name : InNames)
 			{
-				OnVariableRemoved(Name);
+				OnVariableRemoved(NameToGuid.FindChecked(Name));
 			}
 			return true;
 		}
@@ -4655,7 +4680,8 @@ bool URigVMEditorAsset::RenameMemberVariable(const FName& InOldName, const FName
 	{
 		if (RuntimeAsset->RenameMemberVariable(InOldName, InNewName))
 		{
-			OnVariableRenamed(InOldName, InNewName);
+			const FRigVMGraphVariableDescription Variable = RuntimeAsset->FindAssetVariable(InNewName);
+			OnVariableRenamed(Variable.Guid, InNewName);
 			return true;
 		}
 	}
@@ -4685,7 +4711,7 @@ bool URigVMEditorAsset::ChangeMemberVariableType(const FName& InName, const FStr
 		{
 			FRigVMGraphVariableDescription NewVariable = RuntimeAsset->FindAssetVariable(InName);
 			FEdGraphPinType NewPinType = RigVMTypeUtils::PinTypeFromRigVMVariableDescription(NewVariable);
-			OnVariableTypeChanged(InName, OldPinType, NewPinType);
+			OnVariableTypeChanged(OldVariable.Guid, OldPinType, NewPinType);
 			return true;
 		}
 	}
diff --git a/Engine/Plugins/Runtime/RigVM/Source/RigVMDeveloper/Private/RigVMModel/Nodes/RigVMCollapseNode.cpp b/Engine/Plugins/Runtime/RigVM/Source/RigVMDeveloper/Private/RigVMModel/Nodes/RigVMCollapseNode.cpp
--- a/Engine/Plugins/Runtime/RigVM/Source/RigVMDeveloper/Private/RigVMModel/Nodes/RigVMCollapseNode.cpp
+++ b/Engine/Plugins/Runtime/RigVM/Source/RigVMDeveloper/Private/RigVMModel/Nodes/RigVMCollapseNode.cpp
@@ -51,3 +51,12 @@ FString URigVMCollapseNode::GetEditorSubGraphName() const
 {
 	return FString::Printf(TEXT("%s_SubGraph"), *GetName());
 }
+
+FGuid URigVMCollapseNode::FindPinGuid(const URigVMPin* InPin) const
+{
+	if (const FGuid* Guid = PinNameToGuid.Find(InPin->GetFName()))
+	{
+		return *Guid;
+	}
+	return FGuid();
+}
diff --git a/Engine/Plugins/Runtime/RigVM/Source/RigVMDeveloper/Private/RigVMModel/Nodes/RigVMFunctionInterfaceNode.cpp b/Engine/Plugins/Runtime/RigVM/Source/RigVMDeveloper/Private/RigVMModel/Nodes/RigVMFunctionInterfaceNode.cpp
--- a/Engine/Plugins/Runtime/RigVM/Source/RigVMDeveloper/Private/RigVMModel/Nodes/RigVMFunctionInterfaceNode.cpp
+++ b/Engine/Plugins/Runtime/RigVM/Source/RigVMDeveloper/Private/RigVMModel/Nodes/RigVMFunctionInterfaceNode.cpp
@@ -114,6 +114,11 @@ bool URigVMFunctionInterfaceNode::IsInterfacePinUsed(const FName& InInterfacePin
 	return false;
 }
 
+const URigVMCollapseNode* URigVMFunctionInterfaceNode::FindReferencedNode() const
+{
+	return Cast<URigVMCollapseNode>(GetGraph()->GetOuter());
+}
+
 const URigVMPin* URigVMFunctionInterfaceNode::FindReferencedPin(const URigVMPin* InPin) const
 {
 	return FindReferencedPin(InPin->GetSegmentPath(true));
@@ -121,7 +126,7 @@ const URigVMPin* URigVMFunctionInterfaceNode::FindReferencedPin(const URigVMPin*
 
 const URigVMPin* URigVMFunctionInterfaceNode::FindReferencedPin(const FString& InPinPath) const
 {
-	if(const URigVMCollapseNode* OuterNode = Cast<URigVMCollapseNode>(GetGraph()->GetOuter()))
+	if(const URigVMCollapseNode* OuterNode = FindReferencedNode())
 	{
 		return OuterNode->FindPin(InPinPath);
 	}
diff --git a/Engine/Plugins/Runtime/RigVM/Source/RigVMDeveloper/Private/RigVMModel/RigVMClient.cpp b/Engine/Plugins/Runtime/RigVM/Source/RigVMDeveloper/Private/RigVMModel/RigVMClient.cpp
--- a/Engine/Plugins/Runtime/RigVM/Source/RigVMDeveloper/Private/RigVMModel/RigVMClient.cpp
+++ b/Engine/Plugins/Runtime/RigVM/Source/RigVMDeveloper/Private/RigVMModel/RigVMClient.cpp
@@ -357,6 +357,15 @@ void FRigVMClient::RefreshAllModels(ERigVMLoadType InLoadType, bool bEnablePostL
 					TemplateNode->InvalidateCache();
 					TemplateNode->PostLoad();
 				}
+				
+				if (URigVMVariableNode* VariableNode = Cast<URigVMVariableNode>(ModelNode))
+				{
+					if (!VariableNode->GetVariableGuid().IsValid())
+					{
+						FName VarName = VariableNode->GetVariableName();
+						Controller->RefreshVariableNode(VariableNode->GetFName(), VarName, VariableNode->GetCPPType(), VariableNode->GetCPPTypeObject(), false, false);
+					}
+				}
 			}
 		}
 
@@ -424,6 +433,30 @@ void FRigVMClient::RefreshAllModels(ERigVMLoadType InLoadType, bool bEnablePostL
 	}
 }
 
+void FRigVMClient::RefreshVariableNodes(const TArray<FRigVMExternalVariable> InPublicVariables)
+{
+	TArray<URigVMGraph*> AllModels = GetModels();
+	for(URigVMGraph* Model : AllModels)
+	{
+		URigVMController* Controller = GetOrCreateController(Model);
+		TGuardValue<bool> DirtyGuard(Controller->bSuspendMarkingPackageDirty, true);
+		for (URigVMNode* ModelNode : Model->GetNodes())
+		{
+			if (URigVMVariableNode* VariableNode = Cast<URigVMVariableNode>(ModelNode))
+			{
+				for (const FRigVMExternalVariable& PublicVariable : InPublicVariables)
+				{
+					if (VariableNode->GetVariableGuid() == PublicVariable.GetGuid())
+					{
+						FRigVMTemplateArgumentType Type = FRigVMRegistry::Get().GetType(PublicVariable.GetTypeIndex());
+						Controller->RefreshVariableNode(VariableNode->GetFName(), PublicVariable.GetName(), Type.CPPType.ToString(), Type.CPPTypeObject, false, true);
+					}
+				}
+			}
+		}
+	}
+}
+
 const TArray<TObjectPtr<URigVMGraph>>& FRigVMClient::GetModels() const
 {
 	return ExternalModelHost != nullptr ? ExternalModelHost->GetExternalModels() : Models;
@@ -1333,6 +1366,7 @@ FRigVMClientPatchResult FRigVMClient::PatchModelsOnLoad()
 		TGuardValue<bool> GuardIsTransacting(Controller->bIsTransacting, true);
 		{
 			FRigVMDefaultValueTypeGuard DefaultValueTypeGuard(Controller, ERigVMPinDefaultValueType::KeepValueType);
+			Result.Merge(Controller->PatchPinGuidInCollapseNodes());
 			Result.Merge(Controller->PatchLocalVariableTypes());
 			Result.Merge(Controller->PatchRerouteNodesOnLoad());
 			Result.Merge(Controller->PatchUnitNodesOnLoad());
diff --git a/Engine/Plugins/Runtime/RigVM/Source/RigVMDeveloper/Private/RigVMModel/RigVMController.cpp b/Engine/Plugins/Runtime/RigVM/Source/RigVMDeveloper/Private/RigVMModel/RigVMController.cpp
--- a/Engine/Plugins/Runtime/RigVM/Source/RigVMDeveloper/Private/RigVMModel/RigVMController.cpp
+++ b/Engine/Plugins/Runtime/RigVM/Source/RigVMDeveloper/Private/RigVMModel/RigVMController.cpp
@@ -96,6 +96,7 @@ void FRigVMClientPatchResult::Merge(const FRigVMClientPatchResult& InOther)
 URigVMController::URigVMController()
 	: bValidatePinDefaults(true)
 	, bSuspendNotifications(false)
+	, bSuspendMarkingPackageDirty(false)
 	, bSuspendRefreshingFunctionReferences(false)
 	, bReportWarningsAndErrors(true)
 	, bIgnoreRerouteCompactnessChanges(false)
@@ -693,6 +694,10 @@ FRigVMPinInfoArray::FRigVMPinInfoArray(const URigVMNode* InNode, URigVMControlle
 		{
 			Pins[Index].bIsDynamicArray = Pins[Index].bIsArray;
 		}
+		if (const URigVMCollapseNode* CollapseNode = Cast<URigVMCollapseNode>(InNode))
+		{
+			Pins[Index].Guid = CollapseNode->FindPinGuid(Pin);
+		}
 		Pins[Index].bIsInputVariable = Pin->IsDefinedAsInputVariable();
 	}
 }
@@ -1794,7 +1799,7 @@ URigVMUnitNode* URigVMController::AddUnitNodeWithDefaults(UScriptStruct* InScrip
 	
 	if (!bSuspendNotifications)
 	{
-		(void)Graph->MarkPackageDirty();
+		(void)MarkPackageDirty(Graph);
 	}
 
 	FRigVMControllerCompileBracketScope CompileScope(this);
@@ -2018,6 +2023,7 @@ URigVMVariableNode* URigVMController::AddVariableNode(const FName& InVariableNam
 
 	FString Name = GetSchema()->GetValidNodeName(Graph, InNodeName.IsEmpty() ? FString(TEXT("VariableNode")) : InNodeName);
 	URigVMVariableNode* Node = NewObject<URigVMVariableNode>(Graph, InVariableNodeClass.Get(), *Name);
+	Node->VariableGuid = ExternalVariable.GetGuid();
 	Node->Position = InPosition;
 
 	if (!bIsGetter)
@@ -2071,7 +2077,7 @@ URigVMVariableNode* URigVMController::AddVariableNode(const FName& InVariableNam
 
 	if (!bSuspendNotifications)
 	{
-		(void)Graph->MarkPackageDirty();
+		(void)MarkPackageDirty(Graph);
 	}
 
 	FRigVMControllerCompileBracketScope CompileScope(this);
@@ -2138,6 +2144,12 @@ URigVMVariableNode* URigVMController::AddVariableNodeFromObjectPath(const FName&
 }
 
 void URigVMController::RefreshVariableNode(const FName& InNodeName, const FName& InVariableName, const FString& InCPPType, UObject* InCPPTypeObject, bool bSetupUndoRedo, bool bSetupOrphanPins)
+{
+	const FRigVMExternalVariable Variable = GetVariableByName(InVariableName);
+	RefreshVariableNode(InNodeName, Variable.GetGuid(), InVariableName, InCPPType, InCPPTypeObject, bSetupUndoRedo, bSetupOrphanPins);
+}
+
+void URigVMController::RefreshVariableNode(const FName& InNodeName, const FGuid& InVariableGuid, const FName& InVariableName, const FString& InCPPType, UObject* InCPPTypeObject, bool bSetupUndoRedo, bool bSetupOrphanPins)
 {
 	if (!IsValidGraph())
 	{
@@ -2148,12 +2160,26 @@ void URigVMController::RefreshVariableNode(const FName& InNodeName, const FName&
 	{
 		return;
 	}
+	
+	if (!InVariableGuid.IsValid())
+	{
+		return;
+	}
 
 	URigVMGraph* Graph = GetGraph();
 	check(Graph);
 
 	if (URigVMVariableNode* VariableNode = Cast<URigVMVariableNode>(Graph->FindNodeByName(InNodeName)))
 	{
+		if (VariableNode->VariableGuid != InVariableGuid)
+		{
+			if (bSetupUndoRedo)
+			{
+				VariableNode->Modify();
+			}
+			VariableNode->VariableGuid = InVariableGuid;
+		}
+		
 		if (URigVMPin* VariablePin = VariableNode->FindPin(URigVMVariableNode::VariableName))
 		{
 			if (VariablePin->Direction == ERigVMPinDirection::Visible)
@@ -2204,7 +2230,7 @@ void URigVMController::RefreshVariableNode(const FName& InNodeName, const FName&
 	}
 }
 
-void URigVMController::OnExternalVariableRemoved(const FName& InVarName, bool bSetupUndoRedo)
+void URigVMController::OnExternalVariableRemoved(const FGuid& InVarGuid, bool bSetupUndoRedo)
 {
 	if (!IsValidGraph())
 	{
@@ -2216,7 +2242,7 @@ void URigVMController::OnExternalVariableRemoved(const FName& InVarName, bool bS
 		return;
 	}
 
-	if (!InVarName.IsValid())
+	if (!InVarGuid.IsValid())
 	{
 		return;
 	}
@@ -2232,14 +2258,12 @@ void URigVMController::OnExternalVariableRemoved(const FName& InVarName, bool bS
 
 	for (const FRigVMGraphVariableDescription& LocalVariable : Graph->GetLocalVariables(true))
 	{
-		if (InVarName == LocalVariable.Name)
+		if (InVarGuid == LocalVariable.Guid)
 		{
 			return;
 		}
 	}
 	
-	const FString VarNameStr = InVarName.ToString();
-
 	FRigVMControllerCompileBracketScope CompileScope(this);
 	if (bSetupUndoRedo)
 	{
@@ -2251,13 +2275,10 @@ void URigVMController::OnExternalVariableRemoved(const FName& InVarName, bool bS
 	{
 		if (URigVMVariableNode* VariableNode = ExactCast<URigVMVariableNode>(Node))
 		{
-			if (URigVMPin* VariablePin = VariableNode->FindPin(URigVMVariableNode::VariableName))
+			if (VariableNode->GetVariableGuid() == InVarGuid)
 			{
-				if (VariablePin->GetDefaultValue() == VarNameStr)
-				{
-					RemoveNode(Node, bSetupUndoRedo, true);
-					continue;
-				}
+				RemoveNode(Node, bSetupUndoRedo, true);
+				continue;
 			}
 		}
 		else if(const URigVMCollapseNode* CollapseNode = Cast<URigVMCollapseNode>(Node))
@@ -2268,7 +2289,7 @@ void URigVMController::OnExternalVariableRemoved(const FName& InVarName, bool bS
 				TGuardValue<bool> GuardEditGraph(CollapseNode->ContainedGraph->bEditable, true);
 
 				// call this function for the contained graph recursively 
-				CollapseController->OnExternalVariableRemoved(InVarName, bSetupUndoRedo);
+				CollapseController->OnExternalVariableRemoved(InVarGuid, bSetupUndoRedo);
 			}
 		}			
 	}
@@ -2279,7 +2300,7 @@ void URigVMController::OnExternalVariableRemoved(const FName& InVarName, bool bS
 	}
 }
 
-bool URigVMController::OnExternalVariableRenamed(const FName& InOldVarName, const FName& InNewVarName, bool bSetupUndoRedo)
+bool URigVMController::OnExternalVariableRenamed(const FGuid& InVarGuid, const FName& InNewVarName, bool bSetupUndoRedo)
 {
 	if (!IsValidGraph())
 	{
@@ -2291,7 +2312,7 @@ bool URigVMController::OnExternalVariableRenamed(const FName& InOldVarName, cons
 		return false;
 	}
 
-	if (!InOldVarName.IsValid() || !InNewVarName.IsValid())
+	if (!InVarGuid.IsValid() || !InNewVarName.IsValid())
 	{
 		return false;
 	}
@@ -2301,14 +2322,12 @@ bool URigVMController::OnExternalVariableRenamed(const FName& InOldVarName, cons
 
 	for (const FRigVMGraphVariableDescription& LocalVariable : Graph->GetLocalVariables(true))
 	{
-		if (InOldVarName == LocalVariable.Name)
+		if (InVarGuid == LocalVariable.Guid)
 		{
 			return false;
 		}
 	}
 
-	const FString VarNameStr = InOldVarName.ToString();
-
 	FRigVMControllerCompileBracketScope CompileScope(this);
 	if (bSetupUndoRedo)
 	{
@@ -2320,13 +2339,10 @@ bool URigVMController::OnExternalVariableRenamed(const FName& InOldVarName, cons
 	{
 		if (URigVMVariableNode* VariableNode = Cast<URigVMVariableNode>(Node))
 		{
-			if (URigVMPin* VariablePin = VariableNode->FindPin(URigVMVariableNode::VariableName))
+			if (VariableNode->GetVariableGuid() == InVarGuid)
 			{
-				if (VariablePin->GetDefaultValue() == VarNameStr)
-				{
-					RefreshVariableNode(Node->GetFName(), InNewVarName, FString(), nullptr, bSetupUndoRedo, false);
-					continue;
-				}
+				RefreshVariableNode(Node->GetFName(), InVarGuid, InNewVarName, FString(), nullptr, bSetupUndoRedo, false);
+				continue;
 			}
 		}
 		else if(const URigVMCollapseNode* CollapseNode = Cast<URigVMCollapseNode>(Node))
@@ -2335,7 +2351,7 @@ bool URigVMController::OnExternalVariableRenamed(const FName& InOldVarName, cons
 			{
 				FRigVMControllerCompileBracketScope CollapseCompileScope(CollapseController);
 				TGuardValue<bool> GuardEditGraph(CollapseNode->ContainedGraph->bEditable, true);
-				CollapseController->OnExternalVariableRenamed(InOldVarName, InNewVarName, bSetupUndoRedo);
+				CollapseController->OnExternalVariableRenamed(InVarGuid, InNewVarName, bSetupUndoRedo);
 			}
 		}
 	}
@@ -2349,6 +2365,15 @@ bool URigVMController::OnExternalVariableRenamed(const FName& InOldVarName, cons
 }
 
 void URigVMController::OnExternalVariableTypeChanged(const FName& InVarName, const FString& InCPPType, UObject* InCPPTypeObject, bool bSetupUndoRedo)
+{
+	FRigVMExternalVariable Variable = GetVariableByName(InVarName);
+	if (Variable.GetGuid().IsValid())
+	{
+		OnExternalVariableTypeChanged(Variable.GetGuid(), InCPPType, InCPPTypeObject, bSetupUndoRedo);
+	}
+}
+
+void URigVMController::OnExternalVariableTypeChanged(const FGuid& InVarGuid, const FString& InCPPType, UObject* InCPPTypeObject, bool bSetupUndoRedo)
 {
 	if (!IsValidGraph())
 	{
@@ -2360,7 +2385,7 @@ void URigVMController::OnExternalVariableTypeChanged(const FName& InVarName, con
 		return;
 	}
 
-	if (!InVarName.IsValid())
+	if (!InVarGuid.IsValid())
 	{
 		return;
 	}
@@ -2370,13 +2395,13 @@ void URigVMController::OnExternalVariableTypeChanged(const FName& InVarName, con
 
 	for (const FRigVMGraphVariableDescription& LocalVariable : Graph->GetLocalVariables(true))
 	{
-		if (InVarName == LocalVariable.Name)
+		if (InVarGuid == LocalVariable.Guid)
 		{
 			return;
 		}
 	}
 
-	const FString VarNameStr = InVarName.ToString();
+	const FRigVMExternalVariable Variable = GetVariableByGuid(InVarGuid);
 
 	FRigVMControllerCompileBracketScope CompileScope(this);
 	if (bSetupUndoRedo)
@@ -2389,13 +2414,11 @@ void URigVMController::OnExternalVariableTypeChanged(const FName& InVarName, con
 	{
 		if (const URigVMVariableNode* VariableNode = Cast<URigVMVariableNode>(Node))
 		{
-			if (const URigVMPin* VariablePin = VariableNode->FindPin(URigVMVariableNode::VariableName))
+			const bool bMatches = InVarGuid.IsValid() && VariableNode->GetVariableGuid().IsValid() ? VariableNode->GetVariableGuid() == InVarGuid : VariableNode->GetVariableName() == Variable.GetName();
+			if (bMatches) 
 			{
-				if (VariablePin->GetDefaultValue() == VarNameStr)
-				{
-					RefreshVariableNode(Node->GetFName(), InVarName, InCPPType, InCPPTypeObject, bSetupUndoRedo, false);
-					continue;
-				}
+				RefreshVariableNode(Node->GetFName(), InVarGuid, Variable.GetName(), InCPPType, InCPPTypeObject, bSetupUndoRedo, true);
+				continue;
 			}
 		}
 		else if(const URigVMCollapseNode* CollapseNode = Cast<URigVMCollapseNode>(Node))
@@ -2404,14 +2427,14 @@ void URigVMController::OnExternalVariableTypeChanged(const FName& InVarName, con
 			{
 				FRigVMControllerCompileBracketScope CollapseCompileScope(CollapseController);
 				TGuardValue<bool> GuardEditGraph(CollapseNode->ContainedGraph->bEditable, true);
-				CollapseController->OnExternalVariableTypeChanged(InVarName, InCPPType, InCPPTypeObject, bSetupUndoRedo);
+				CollapseController->OnExternalVariableTypeChanged(InVarGuid, InCPPType, InCPPTypeObject, bSetupUndoRedo);
 			}
 		}
 
 		TArray<URigVMPin*> AllPins = Node->GetAllPinsRecursively();
 		for (URigVMPin* Pin : AllPins)
 		{
-			if (Pin->GetBoundVariableName() == InVarName.ToString())
+			if (Pin->GetBoundVariableName() == Variable.GetName().ToString())
 			{
 				FString BoundVariablePath = Pin->GetBoundVariablePath();
 				UnbindPinFromVariable(Pin, bSetupUndoRedo);
@@ -2428,6 +2451,15 @@ void URigVMController::OnExternalVariableTypeChanged(const FName& InVarName, con
 }
 
 void URigVMController::OnExternalVariableTypeChangedFromObjectPath(const FName& InVarName, const FString& InCPPType, const FString& InCPPTypeObjectPath, bool bSetupUndoRedo)
+{
+	FRigVMExternalVariable Variable = GetVariableByName(InVarName);
+	if (Variable.GetGuid().IsValid())
+	{
+		OnExternalVariableTypeChangedFromObjectPath(Variable.GetGuid(), InCPPType, InCPPTypeObjectPath, bSetupUndoRedo);
+	}
+}
+
+void URigVMController::OnExternalVariableTypeChangedFromObjectPath(const FGuid& InVarGuid, const FString& InCPPType, const FString& InCPPTypeObjectPath, bool bSetupUndoRedo)
 {
 	if (!IsValidGraph())
 	{
@@ -2450,7 +2482,7 @@ void URigVMController::OnExternalVariableTypeChangedFromObjectPath(const FName&
 		}
 	}
 
-	OnExternalVariableTypeChanged(InVarName, InCPPType, CPPTypeObject, bSetupUndoRedo);
+	OnExternalVariableTypeChanged(InVarGuid, InCPPType, CPPTypeObject, bSetupUndoRedo);
 }
 
 URigVMVariableNode* URigVMController::ReplaceParameterNodeWithVariable(const FName& InNodeName, const FName& InVariableName, const FString& InCPPType, UObject* InCPPTypeObject, bool bSetupUndoRedo)
@@ -3050,7 +3082,7 @@ URigVMCommentNode* URigVMController::AddCommentNode(const FString& InCommentText
 	
 	if (!bSuspendNotifications)
 	{
-		(void)Graph->MarkPackageDirty();
+		(void)MarkPackageDirty(Graph);
 	}
 
 	FRigVMControllerCompileBracketScope CompileScope(this);
@@ -3238,7 +3270,7 @@ URigVMRerouteNode* URigVMController::AddRerouteNodeOnPin(const FString& InPinPat
 	
 	if (!bSuspendNotifications)
 	{
-		(void)Graph->MarkPackageDirty();
+		(void)MarkPackageDirty(Graph);
 	}
 
 	if (bSetupUndoRedo)
@@ -3467,7 +3499,7 @@ URigVMInjectionInfo* URigVMController::AddInjectedNode(const FString& InPinPath,
 
 	if (!bSuspendNotifications)
 	{
-		(void)Graph->MarkPackageDirty();
+		(void)MarkPackageDirty(Graph);
 	}
 	
 	if (bSetupUndoRedo)
@@ -3606,7 +3638,7 @@ bool URigVMController::RemoveInjectedNode(const FString& InPinPath, bool bAsInpu
 
 	if (!bSuspendNotifications)
 	{
-		(void)Graph->MarkPackageDirty();
+		(void)MarkPackageDirty(Graph);
 	}
 	
 	if (bSetupUndoRedo)
@@ -3759,7 +3791,7 @@ URigVMInjectionInfo* URigVMController::InjectNodeIntoPin(URigVMPin* InPin, bool
 
 	if (!bSuspendNotifications)
 	{
-		(void)Graph->MarkPackageDirty();
+		(void)MarkPackageDirty(Graph);
 	}
 
 	if (bSetupUndoRedo)
@@ -3894,7 +3926,7 @@ URigVMNode* URigVMController::EjectNodeFromPin(URigVMPin* InPin, bool bSetupUndo
 		
 	if (!bSuspendNotifications)
 	{
-		(void)Graph->MarkPackageDirty();
+		(void)MarkPackageDirty(Graph);
 	}
 
 	if (bSetupUndoRedo)
@@ -7217,7 +7249,9 @@ bool URigVMController::RemoveNodes(TArray<URigVMNode*> InNodes, bool bSetupUndoR
 
 		SelectNode(InNode, false, false);
 
-		for (URigVMPin* Pin : InNode->GetPins())
+		TArray<URigVMPin*> Pins = InNode->GetPins();
+		Pins.Append(InNode->GetOrphanedPins());
+		for (URigVMPin* Pin : Pins)
 		{
 			// Remove injected nodes
 			while (Pin->HasInjectedNodes())
@@ -7260,7 +7294,7 @@ bool URigVMController::RemoveNodes(TArray<URigVMNode*> InNodes, bool bSetupUndoR
 
 	if (!bSuspendNotifications)
 	{
-		(void)Graph->MarkPackageDirty();
+		(void)MarkPackageDirty(Graph);
 	}
 
 	if (bPrintPythonCommand)
@@ -8364,19 +8398,29 @@ bool URigVMController::SetCommentTextByName(const FName& InNodeName, const FStri
 
 bool URigVMController::RenameVariable(const FName& InOldName, const FName& InNewName, bool bSetupUndoRedo)
 {
-	if(!IsValidGraph())
+	if (InOldName == InNewName)
 	{
+		ReportWarning(TEXT("RenameVariable: InOldName and InNewName are equal."));
 		return false;
 	}
+	
+	const FRigVMExternalVariable Variable = GetVariableByName(InOldName);
+	if (Variable.GetGuid().IsValid())
+	{
+		return RenameVariable(Variable.GetGuid(), InNewName, bSetupUndoRedo);
+	}
+	return false;
+}
 
-	if (!bIsTransacting && !IsGraphEditable())
+bool URigVMController::RenameVariable(const FGuid& InGuid, const FName& InNewName, bool bSetupUndoRedo)
+{
+	if(!IsValidGraph())
 	{
 		return false;
 	}
 
-	if (InOldName == InNewName)
+	if (!bIsTransacting && !IsGraphEditable())
 	{
-		ReportWarning(TEXT("RenameVariable: InOldName and InNewName are equal."));
 		return false;
 	}
 
@@ -8396,17 +8440,19 @@ bool URigVMController::RenameVariable(const FName& InOldName, const FName& InNew
 	// If there is a local variable with the old name, a rename of the blueprint member variable does not affect this graph
 	for (FRigVMGraphVariableDescription& LocalVariable : Graph->GetLocalVariables(true))
 	{
-		if (LocalVariable.Name == InOldName)
+		if (LocalVariable.Guid == InGuid)
 		{
 			return false;
 		}
 	}
 
+	const FRigVMExternalVariable Variable = GetVariableByGuid(InGuid);
+
 	FRigVMControllerCompileBracketScope CompileScope(this);
 	FRigVMRenameVariableAction Action;
 	if (bSetupUndoRedo)
 	{
-		Action = FRigVMRenameVariableAction(this, InOldName, InNewName);
+		Action = FRigVMRenameVariableAction(this, InGuid, Variable.GetName(), InNewName);
 		Action.SetTitle(FString::Printf(TEXT("Rename Variable")));
 		GetActionStack()->BeginAction(Action);
 	}
@@ -8416,7 +8462,7 @@ bool URigVMController::RenameVariable(const FName& InOldName, const FName& InNew
 	{
 		if(URigVMVariableNode* VariableNode = Cast<URigVMVariableNode>(Node))
 		{
-			if (VariableNode->GetVariableName() == InOldName)
+			if (VariableNode->GetVariableGuid() == InGuid)
 			{
 				VariableNode->FindPin(URigVMVariableNode::VariableName)->DefaultValue = InNewName.ToString();
 				RenamedNodes.Add(Node);
@@ -8429,7 +8475,7 @@ bool URigVMController::RenameVariable(const FName& InOldName, const FName& InNew
 		Notify(ERigVMGraphNotifType::VariableRenamed, RenamedNode);
 		if (!bSuspendNotifications)
 		{
-			(void)Graph->MarkPackageDirty();
+			(void)MarkPackageDirty(Graph);
 		}
 	}
 
@@ -8520,7 +8566,7 @@ bool URigVMController::SetPinExpansion(URigVMPin* InPin, bool bIsExpanded, bool
 	Notify(ERigVMGraphNotifType::PinExpansionChanged, InPin);
 	if (!bSuspendNotifications)
 	{
-		(void)Graph->MarkPackageDirty();
+		(void)MarkPackageDirty(Graph);
 	}
 
 	if (bSetupUndoRedo)
@@ -8595,7 +8641,7 @@ bool URigVMController::SetPinIsWatched(URigVMPin* InPin, bool bIsWatched, bool b
 	Notify(ERigVMGraphNotifType::PinWatchedChanged, InPin);
 	if (!bSuspendNotifications)
 	{
-		(void)Graph->MarkPackageDirty();
+		(void)MarkPackageDirty(Graph);
 	}
 
 	if (bSetupUndoRedo)
@@ -8700,7 +8746,7 @@ bool URigVMController::SetPinDisplayName(URigVMPin* InPin, const FString& InDisp
 	{
 		const URigVMGraph* Graph = GetGraph();
 		check(Graph);
-		(void)Graph->MarkPackageDirty();
+		(void)MarkPackageDirty(Graph);
 	}
 
 	if (bSetupUndoRedo)
@@ -8824,7 +8870,7 @@ bool URigVMController::AddEmptyPinCategory(const URigVMNode* InNode, const FStri
 	{
 		const URigVMGraph* Graph = GetGraph();
 		check(Graph);
-		(void)Graph->MarkPackageDirty();
+		(void)MarkPackageDirty(Graph);
 	}
 
 	if (bSetupUndoRedo)
@@ -9106,7 +9152,7 @@ bool URigVMController::SetPinCategory(URigVMPin* InPin, const FString& InCategor
 	{
 		const URigVMGraph* Graph = GetGraph();
 		check(Graph);
-		(void)Graph->MarkPackageDirty();
+		(void)MarkPackageDirty(Graph);
 	}
 
 	if (bSetupUndoRedo)
@@ -9248,7 +9294,7 @@ bool URigVMController::RemovePinCategory(const URigVMNode* InNode, const FString
 	{
 		const URigVMGraph* Graph = GetGraph();
 		check(Graph);
-		(void)Graph->MarkPackageDirty();
+		(void)MarkPackageDirty(Graph);
 	}
 
 	if (bSetupUndoRedo)
@@ -9447,7 +9493,7 @@ bool URigVMController::RenamePinCategory(const URigVMNode* InNode, const FString
 	{
 		const URigVMGraph* Graph = GetGraph();
 		check(Graph);
-		(void)Graph->MarkPackageDirty();
+		(void)MarkPackageDirty(Graph);
 	}
 
 	if (bSetupUndoRedo)
@@ -9661,7 +9707,7 @@ bool URigVMController::SetPinCategoryExpansion(const URigVMNode* InNode, const F
 	{
 		const URigVMGraph* Graph = GetGraph();
 		check(Graph);
-		(void)Graph->MarkPackageDirty();
+		(void)MarkPackageDirty(Graph);
 	}
 
 	if (bSetupUndoRedo)
@@ -9800,7 +9846,7 @@ bool URigVMController::SetPinIndexInCategory(URigVMPin* InPin, int32 InIndexInCa
 	{
 		const URigVMGraph* Graph = GetGraph();
 		check(Graph);
-		(void)Graph->MarkPackageDirty();
+		(void)MarkPackageDirty(Graph);
 	}
 
 	if (bSetupUndoRedo)
@@ -10141,7 +10187,7 @@ bool URigVMController::SetPinCategories(const URigVMNode* InNode, const TArray<F
 	{
 		const URigVMGraph* Graph = GetGraph();
 		check(Graph);
-		(void)Graph->MarkPackageDirty();
+		(void)MarkPackageDirty(Graph);
 	}
 
 	if (bSetupUndoRedo)
@@ -10401,7 +10447,7 @@ bool URigVMController::SetPinDefaultValue(URigVMPin* InPin, const FString& InDef
 			Notify(ERigVMGraphNotifType::PinDefaultValueChanged, InPin);
 			if (!bSuspendNotifications)
 			{
-				(void)Graph->MarkPackageDirty();
+				(void)MarkPackageDirty(Graph);
 			}
 			bSetPinDefaultValueSucceeded = true;
 		}
@@ -11905,7 +11951,6 @@ URigVMPin* URigVMController::InsertArrayPin(URigVMPin* ArrayPin, int32 InIndex,
 	{
 		URigVMPin* ExistingPin = ArrayPin->GetSubPins()[ExistingIndex];
 		RenameObject(ExistingPin, *FString::FormatAsNumber(ExistingIndex + 1));
-		ExistingPin->ClearPinPathCache();
 
 		// temporarily force a the index based display name
 		// so that the ed graph node can pick that up during handling the notify
@@ -11974,7 +12019,7 @@ URigVMPin* URigVMController::InsertArrayPin(URigVMPin* ArrayPin, int32 InIndex,
 	Notify(ERigVMGraphNotifType::PinArraySizeChanged, ArrayPin);
 	if (!bSuspendNotifications)
 	{
-		(void)Graph->MarkPackageDirty();
+		(void)MarkPackageDirty(Graph);
 	}
 
 	if (bSetupUndoRedo)
@@ -12063,7 +12108,7 @@ bool URigVMController::RemoveArrayPin(const FString& InArrayElementPinPath, bool
 
 	if (!bSuspendNotifications)
 	{
-		(void)Graph->MarkPackageDirty();
+		(void)MarkPackageDirty(Graph);
 	}
 
 	// set the array pin's default value type based on the resulting array pin list
@@ -13139,7 +13184,7 @@ bool URigVMController::AddLink(URigVMPin* OutputPin, URigVMPin* InputPin, bool b
 
 	if (!bSuspendNotifications)
 	{
-		(void)Graph->MarkPackageDirty();
+		(void)MarkPackageDirty(Graph);
 	}
 	Notify(ERigVMGraphNotifType::LinkAdded, Link);
 
@@ -13357,7 +13402,7 @@ bool URigVMController::BreakLink(URigVMPin* OutputPin, URigVMPin* InputPin, bool
 			
 			if (!bSuspendNotifications)
 			{
-				(void)Graph->MarkPackageDirty();
+				(void)MarkPackageDirty(Graph);
 			}
 			Notify(ERigVMGraphNotifType::LinkRemoved, Link);
 
@@ -13425,6 +13470,11 @@ bool URigVMController::BreakAllLinks(URigVMPin* Pin, bool bAsInput, bool bSetupU
 		return false;
 	}
 	
+	if (!IsValid(Pin))
+	{
+		return false;
+	}
+	
 	if(!Pin->IsLinked(false))
 	{
 		return false;
@@ -13498,7 +13548,8 @@ bool URigVMController::BreakAllLinksRecursive(URigVMPin* Pin, bool bAsInput, boo
 		}
 		else
 		{
-			for (URigVMPin* SubPin : Pin->SubPins)
+			TArray<URigVMPin*> SubPins = Pin->SubPins;
+			for (URigVMPin* SubPin : SubPins)
 			{
 				bBrokenLinks |= BreakAllLinks(SubPin, bAsInput, bSetupUndoRedo);
 				bBrokenLinks |= BreakAllLinksRecursive(SubPin, bAsInput, bTowardsParent, bSetupUndoRedo);
@@ -13619,6 +13670,10 @@ FName URigVMController::AddExposedPin(const FName& InPinName, ERigVMPinDirection
 	Pin->CPPTypeObjectPath = InCPPTypeObjectPath;
 	Pin->bIsConstant = false;
 	Pin->bIsInputVariable = bIsInputVariable;
+	if (URigVMCollapseNode* CollapseNode = Cast<URigVMCollapseNode>(LibraryNode))
+	{
+		CollapseNode->PinNameToGuid.FindOrAdd(PinName) = FGuid::NewGuid();
+	}
 
 	Pin->Direction = InDirection;
 	Pin->bIsDynamicArray = RigVMTypeUtils::IsArrayType(CPPType);
@@ -13803,6 +13858,12 @@ bool URigVMController::RemoveExposedPin(const FName& InPinName, bool bSetupUndoR
 
 	bool bSuccessfullyRemovedPin = false;
 	{
+		FGuid Guid;
+		if (URigVMCollapseNode* CollapseNode = Cast<URigVMCollapseNode>(LibraryNode))
+		{
+			CollapseNode->PinNameToGuid.RemoveAndCopyValue(InPinName, Guid);
+		}
+		
 		if(URigVMController* LibraryController = GetControllerForGraph(LibraryNode->GetGraph()))
 		{
 			bSuccessfullyRemovedPin = LibraryController->RemovePin(Pin, bSetupUndoRedo, true);
@@ -13813,9 +13874,19 @@ bool URigVMController::RemoveExposedPin(const FName& InPinName, bool bSetupUndoR
 		{
 			if (URigVMVariableNode* VariableNode = Cast<URigVMVariableNode>(Node))
 			{
-				if (VariableNode->GetVariableName() == InPinName)
+				if (Guid.IsValid())
 				{
-					NodesToRemove.Add(VariableNode);
+					if (VariableNode->GetVariableGuid() == Guid)
+					{
+						NodesToRemove.Add(VariableNode);
+					}
+				}
+				else
+				{
+					if (VariableNode->GetVariableName() == InPinName)
+					{
+						NodesToRemove.Add(VariableNode);
+					}
 				}
 			}
 		}
@@ -13894,6 +13965,14 @@ bool URigVMController::RenameExposedPin(const FName& InOldPinName, const FName&
 	{
 		return false;
 	}
+	
+	if (URigVMCollapseNode* CollapseNode = Cast<URigVMCollapseNode>(LibraryNode))
+	{
+		FGuid Guid;
+		CollapseNode->PinNameToGuid.RemoveAndCopyValue(InOldPinName, Guid);
+		ensure(Guid.IsValid());
+		CollapseNode->PinNameToGuid.FindOrAdd(InNewPinName) = Guid;
+	}
 
 	if(bSetupUndoRedo)
 	{
@@ -13952,7 +14031,6 @@ bool URigVMController::RenameExposedPin(const FName& InOldPinName, const FName&
 				return false;
 			}
 
-			InPin->ClearPinPathCache();
 			InPin->SetDisplayName(InNewName);
 			InPin->IncrementVersion();
 			
@@ -14400,7 +14478,7 @@ bool URigVMController::MarkExposedPinAsInputVariable(const FName& InPinName, boo
 		return false;
 	}
 
-	URigVMLibraryNode* LibraryNode = Cast<URigVMLibraryNode>(Graph->GetOuter());
+	URigVMCollapseNode* LibraryNode = Cast<URigVMCollapseNode>(Graph->GetOuter());
 	if (LibraryNode == nullptr)
 	{
 		ReportError(TEXT("Graph is not under a Collapse Node"));
@@ -14457,7 +14535,11 @@ bool URigVMController::MarkExposedPinAsInputVariable(const FName& InPinName, boo
 			// we need to run the code below to clean up obsolete variable nodes (like SetVariable)
 			if (!bDefinedAsExternalVariable)
 			{
-				LibraryController->OnExternalVariableRemoved(Pin->GetFName(), bSetupUndoRedo);
+				FGuid Guid = LibraryNode->FindPinGuid(Pin);
+				if (Guid.IsValid())
+				{
+					LibraryController->OnExternalVariableRemoved(Guid, bSetupUndoRedo);
+				}
 			}
 		}
 	}
@@ -15158,7 +15240,7 @@ bool URigVMController::MarkFunctionAsPublic(const FName& InFunctionName, bool bI
 
 	if (!bSuspendNotifications)
 	{
-		(void)Graph->MarkPackageDirty();
+		(void)MarkPackageDirty(Graph);
 	}
 
 	if (bPrintPythonCommand)
@@ -15717,6 +15799,7 @@ FRigVMGraphVariableDescription URigVMController::AddLocalVariable(const FName& I
 		return true;
 	}, false, true);
 
+	NewVariable.Guid = FGuid::NewGuid();
 	NewVariable.Name = VariableName;
 	NewVariable.CPPType = InCPPType;
 	NewVariable.CPPTypeObject = InCPPTypeObject;
@@ -15736,7 +15819,7 @@ FRigVMGraphVariableDescription URigVMController::AddLocalVariable(const FName& I
 		{
 			if (VariableName == VariableNode->GetVariableName())
 			{
-				RefreshVariableNode(VariableNode->GetFName(), VariableName, InCPPType, InCPPTypeObject, bSetupUndoRedo, false);
+				RefreshVariableNode(VariableNode->GetFName(), NewVariable.Guid, VariableName, InCPPType, InCPPTypeObject, bSetupUndoRedo, false);
 			}
 		}
 	}
@@ -15750,7 +15833,7 @@ FRigVMGraphVariableDescription URigVMController::AddLocalVariable(const FName& I
 
 	if (!bSuspendNotifications)
 	{
-		(void)Graph->MarkPackageDirty();
+		(void)MarkPackageDirty(Graph);
 
 		if (URigVMLibraryNode* LibraryNode = Cast<URigVMLibraryNode>(Graph->GetOuter()))
 		{
@@ -15888,7 +15971,7 @@ bool URigVMController::RemoveLocalVariable(const FName& InVariableName, bool bSe
 					{
 						if (VariablePin->GetDefaultValue() == VarNameStr)
 						{
-							RefreshVariableNode(VariableNode->GetFName(), ExternalVariableToSwitch.GetName(), ExternalVariableToSwitch.GetExtendedCPPType().ToString(), ExternalVariableToSwitch.GetCPPTypeObject(), bSetupUndoRedo, false);
+							RefreshVariableNode(VariableNode->GetFName(),  ExternalVariableToSwitch.GetGuid(), ExternalVariableToSwitch.GetName(), ExternalVariableToSwitch.GetExtendedCPPType().ToString(), ExternalVariableToSwitch.GetCPPTypeObject(), bSetupUndoRedo, false);
 							continue;
 						}
 					}
@@ -15910,7 +15993,7 @@ bool URigVMController::RemoveLocalVariable(const FName& InVariableName, bool bSe
 
 		if (!bSuspendNotifications)
 		{
-			(void)Graph->MarkPackageDirty();
+			(void)MarkPackageDirty(Graph);
 			if (URigVMLibraryNode* LibraryNode = Cast<URigVMLibraryNode>(Graph->GetOuter()))
 			{
 				Notify(ERigVMGraphNotifType::LocalVariableRemoved, LibraryNode);
@@ -15984,7 +16067,7 @@ bool URigVMController::RenameLocalVariable(const FName& InVariableName, const FN
 	
 	if (!bSuspendNotifications)
 	{
-		(void)Graph->MarkPackageDirty();
+		(void)MarkPackageDirty(Graph);
 	}
 
 	FRigVMControllerCompileBracketScope CompileScope(this);
@@ -16015,7 +16098,7 @@ bool URigVMController::RenameLocalVariable(const FName& InVariableName, const FN
 
 	if (!bSuspendNotifications)
 	{
-		(void)Graph->MarkPackageDirty();
+		(void)MarkPackageDirty(Graph);
 		if (URigVMLibraryNode* LibraryNode = Cast<URigVMLibraryNode>(Graph->GetOuter()))
 		{
 			Notify(ERigVMGraphNotifType::LocalVariableRenamed, LibraryNode);
@@ -16074,7 +16157,7 @@ bool URigVMController::SetLocalVariableType(const FName& InVariableName, const F
 
 	if (!bSuspendNotifications)
 	{
-		(void)Graph->MarkPackageDirty();
+		(void)MarkPackageDirty(Graph);
 		if (URigVMLibraryNode* LibraryNode = Cast<URigVMLibraryNode>(Graph->GetOuter()))
 		{
 			Notify(ERigVMGraphNotifType::LocalVariableTypeChanged, LibraryNode);
@@ -16124,7 +16207,7 @@ bool URigVMController::SetLocalVariableType(const FName& InVariableName, const F
 			{
 				if (VariablePin->GetDefaultValue() == InVariableName.ToString())
 				{
-					RefreshVariableNode(Node->GetFName(), InVariableName, InCPPType, InCPPTypeObject, bSetupUndoRedo, false);
+					RefreshVariableNode(Node->GetFName(), LocalVariables[FoundIndex].Guid, InVariableName, InCPPType, InCPPTypeObject, bSetupUndoRedo, false);
 					continue;
 				}
 			}
@@ -16221,7 +16304,7 @@ bool URigVMController::SetLocalVariableDefaultValue(const FName& InVariableName,
 	
 	if (!bSuspendNotifications)
 	{
-		(void)Graph->MarkPackageDirty();
+		(void)MarkPackageDirty(Graph);
 
 		if (URigVMLibraryNode* LibraryNode = Cast<URigVMLibraryNode>(Graph->GetOuter()))
 		{
@@ -16531,7 +16614,7 @@ bool URigVMController::SetVariableName(URigVMVariableNode* InVariableNode, const
 		NameToIndex.Add(Descriptions[VariableIndex].GetName(), VariableIndex);
 	}
 
-	const FRigVMExternalVariable VariableType = RigVMTypeUtils::ExternalVariableFromCPPType(InVariableName, InVariableNode->GetCPPType(), InVariableNode->GetCPPTypeObject());
+	const FRigVMExternalVariable VariableType = RigVMTypeUtils::ExternalVariableFromCPPType(FGuid(), InVariableName, InVariableNode->GetCPPType(), InVariableNode->GetCPPTypeObject());
 	FName VariableName = URigVMSchema::GetUniqueName(InVariableName, [Descriptions, NameToIndex, VariableType](const FName& InName) {
 		const int32* FoundIndex = NameToIndex.Find(InName);
 		if (FoundIndex == nullptr)
@@ -17530,7 +17613,7 @@ URigVMInvokeEntryNode* URigVMController::AddInvokeEntryNode(const FName& InEntry
 
 	if (!bSuspendNotifications)
 	{
-		(void)Graph->MarkPackageDirty();
+		(void)MarkPackageDirty(Graph);
 	}
 
 	FRigVMControllerCompileBracketScope CompileScope(this);
@@ -19196,7 +19279,7 @@ bool URigVMController::GenerateNewPinInfos(const FRigVMRegistry& Registry, URigV
 				{
 					if (URigVMPin* EntryPin = GraphEntryNode->FindPin(VariableDescription.Name.ToString()))
 					{
-						Variable = RigVMTypeUtils::ExternalVariableFromCPPType(VariableDescription.Name, EntryPin->GetCPPType(), EntryPin->GetCPPTypeObject());
+						Variable = RigVMTypeUtils::ExternalVariableFromCPPType(VariableDescription.Guid, VariableDescription.Name, EntryPin->GetCPPType(), EntryPin->GetCPPTypeObject());
 					}
 				}
 			}
@@ -19204,34 +19287,34 @@ bool URigVMController::GenerateNewPinInfos(const FRigVMRegistry& Registry, URigV
 			{
 				for (const FRigVMExternalVariable& ExternalVariable : ExternalVariables)
 				{
-					if (ExternalVariable.GetName() == CurrentExternalVariable.GetName())
+					if (ExternalVariable.GetGuid().IsValid() && CurrentExternalVariable.GetGuid().IsValid())
 					{
-						Variable = ExternalVariable;
-						break;
+						if (ExternalVariable.GetGuid() == CurrentExternalVariable.GetGuid())
+						{
+							Variable = ExternalVariable;
+							break;
+						}
+					}
+					else
+					{
+						if (ExternalVariable.GetName() == CurrentExternalVariable.GetName())
+						{
+							Variable = ExternalVariable;
+							break;
+						}
 					}
 				}
 			}
 
 			if (Variable.IsValid(true))
 			{
-				if (!Variable.IsSameType(CurrentExternalVariable))
+				if (!VariableNode->IsGetter())
 				{
-					FString CPPType;
-					UObject* CPPTypeObject;
-
-					if (RigVMTypeUtils::CPPTypeFromExternalVariable(Variable, CPPType, &CPPTypeObject))
-					{
-						RefreshVariableNode(VariableNode->GetFName(), Variable.GetName(), CPPType, Variable.GetCPPTypeObject(), false, bSetupOrphanPinsForThisNode);
-					}
-					else
-					{
-						ReportAndNotifyErrorf(
-							TEXT("Control Rig '%s', Type of Variable '%s' cannot be resolved."),
-							*InNode->GetOutermost()->GetPathName(),
-							*Variable.GetName().ToString()
-						);
-					}
+					(void) NewPinInfos.AddPin(VariableNode->FindExecutePin(), INDEX_NONE, ERigVMPinDirection::IO, ERigVMPinDefaultValueType::AutoDetect);
 				}
+				URigVMPin* NamePin = VariableNode->GetVariableNamePin();
+				(void) NewPinInfos.AddPin(NamePin, INDEX_NONE, NamePin->GetDirection(), NamePin->GetDefaultValueType());
+				(void) NewPinInfos.AddPin(this, INDEX_NONE, URigVMVariableNode::ValueName, ValuePin->GetDirection(), Variable.GetTypeIndex(), ValuePin->GetDefaultValue(), ValuePin->GetDefaultValueType(), nullptr, &PreviousPinInfos, true);
 			}
 			else
 			{
@@ -19242,8 +19325,11 @@ bool URigVMController::GenerateNewPinInfos(const FRigVMRegistry& Registry, URigV
 				);
 			}
 		}
+		else
+		{
+			NewPinInfos = FRigVMPinInfoArray(InNode, this, &PreviousPinInfos);
+		}
 
-		NewPinInfos = FRigVMPinInfoArray(InNode, this, &PreviousPinInfos);
 	}
 	else if (EntryNode || ReturnNode)
 	{
@@ -19598,12 +19684,12 @@ void URigVMController::RepopulatePinsOnNode(const FRigVMRegistry& Registry, cons
 			if (InNode->FindPin(OrphanedName) == nullptr)
 			{
 				Pin->DisplayName = Pin->GetFName();
-				RenameObject(Pin, *OrphanedName, nullptr);
 				InNode->Pins.Remove(Pin);
+				
+				RenameObject(Pin, *OrphanedName, nullptr);
 
-				Notify(ERigVMGraphNotifType::PinRemoved, Pin);
 				InNode->OrphanedPins.Add(Pin);
-				Notify(ERigVMGraphNotifType::PinAdded, Pin);
+				Notify(ERigVMGraphNotifType::PinRenamed, Pin);
 			}
 			else
 			{
@@ -21892,12 +21978,15 @@ bool URigVMController::ChangePinType(URigVMPin* InPin, TRigVMTypeIndex InTypeInd
 				ConfigurePinFromPin(OrphanedPin, InPin);
 				OrphanedPin->DisplayName = InPin->GetFName();
 
+				InPin->GetNode()->OrphanedPins.Add(OrphanedPin);
+				Notify(ERigVMGraphNotifType::PinAdded, OrphanedPin);
+
 				if(OrphanedPin->IsStruct())
 				{
 					AddPinsForStruct(OrphanedPin->GetScriptStruct(), OrphanedPin->GetNode(), OrphanedPin, OrphanedPin->Direction, OrphanedPin->GetDefaultValue(), false);
 				}
 				
-				InPin->GetNode()->OrphanedPins.Add(OrphanedPin);
+				OrphanedPin->ClearPinPathCache();
 			}
 		}
 	}
@@ -22218,6 +22307,15 @@ void URigVMController::RewireLinks(URigVMPin* InOldPin, URigVMPin* InNewPin, boo
 
 #endif
 
+bool URigVMController::MarkPackageDirty(const UObject* InObjectToDirty) const
+{
+	if (!bSuspendMarkingPackageDirty && InObjectToDirty)
+	{
+		return InObjectToDirty->MarkPackageDirty();
+	}
+	return false;
+}
+
 bool URigVMController::RenameObject(UObject* InObjectToRename, const TCHAR* InNewName, UObject* InNewOuter, ERenameFlags InFlags) const
 {
 	// Renames are commonly done to temporarily move objects to a new outer and then move back (often during loading). 
@@ -22236,6 +22334,12 @@ bool URigVMController::RenameObject(UObject* InObjectToRename, const TCHAR* InNe
 			}
 		}
 	}
+	
+	if (URigVMPin* Pin = Cast<URigVMPin>(InObjectToRename))
+	{
+		Pin->ClearPinPathCache();
+	}
+	
 	return bSuccess;
 }
 
@@ -22515,6 +22619,20 @@ FRigVMExternalVariable URigVMController::GetVariableByName(const FName& InExtern
 	return FRigVMExternalVariable();
 }
 
+FRigVMExternalVariable URigVMController::GetVariableByGuid(const FGuid& InExternalVariableGuid, const bool bIncludeInputArguments) const
+{
+	TArray<FRigVMExternalVariable> Variables = GetAllVariables(bIncludeInputArguments);
+	for (const FRigVMExternalVariable& Variable : Variables)
+	{
+		if (Variable.GetGuid() == InExternalVariableGuid)
+		{
+			return Variable;
+		}
+	}
+	
+	return FRigVMExternalVariable();
+}
+
 TArray<FRigVMExternalVariable> URigVMController::GetAllVariables(const bool bIncludeInputArguments) const
 {
 	TArray<FRigVMExternalVariable> ExternalVariables;
@@ -24609,6 +24727,37 @@ FRigVMClientPatchResult URigVMController::PatchLocalVariableTypes()
 	return Result;
 }
 
+FRigVMClientPatchResult URigVMController::PatchPinGuidInCollapseNodes()
+{
+	FRigVMClientPatchResult Result;
+	if (URigVMGraph* Graph = GetGraph())
+	{
+		for (URigVMNode* Node : Graph->GetNodes())
+		{
+			if (URigVMCollapseNode* CollapseNode = Cast<URigVMCollapseNode>(Node))
+			{
+				if (CollapseNode->PinNameToGuid.IsEmpty())
+				{
+					for (URigVMPin* Pin : CollapseNode->GetPins())
+					{
+						CollapseNode->PinNameToGuid.FindOrAdd(Pin->GetFName()) = FGuid::NewDeterministicGuid(Pin->GetPathName());
+						Result.bChangedContent = true;
+					}
+				}
+			}
+		}
+		
+		for (FRigVMGraphVariableDescription& LocalVariable : Graph->LocalVariables)
+		{
+			const FString VariablePath = FString::Printf(TEXT("%s.%s:%s"), *Graph->GetPathName(), *LocalVariable.Name.ToString(), *LocalVariable.CPPType);
+			LocalVariable.Guid = FGuid::NewDeterministicGuid(VariablePath);
+			Result.bChangedContent = true;
+		}
+	}
+
+	return Result;
+}
+
 FRigVMClientPatchResult URigVMController::PatchExternalVariablesToArguments(bool bInSetupUndoRedo)
 {
 	FRigVMClientPatchResult Result;
diff --git a/Engine/Plugins/Runtime/RigVM/Source/RigVMDeveloper/Private/RigVMModel/RigVMControllerActions.cpp b/Engine/Plugins/Runtime/RigVM/Source/RigVMDeveloper/Private/RigVMModel/RigVMControllerActions.cpp
--- a/Engine/Plugins/Runtime/RigVM/Source/RigVMDeveloper/Private/RigVMModel/RigVMControllerActions.cpp
+++ b/Engine/Plugins/Runtime/RigVM/Source/RigVMDeveloper/Private/RigVMModel/RigVMControllerActions.cpp
@@ -1079,8 +1079,9 @@ FRigVMRenameVariableAction::FRigVMRenameVariableAction()
 {
 }
 
-FRigVMRenameVariableAction::FRigVMRenameVariableAction(URigVMController* InController, const FName& InOldVariableName, const FName& InNewVariableName)
+FRigVMRenameVariableAction::FRigVMRenameVariableAction(URigVMController* InController, const FGuid& InVariableGuid, const FName& InOldVariableName, const FName& InNewVariableName)
 : FRigVMBaseAction(InController)
+, VariableGuid(InVariableGuid)
 , OldVariableName(InOldVariableName.ToString())
 , NewVariableName(InNewVariableName.ToString())
 {
@@ -1092,7 +1093,7 @@ bool FRigVMRenameVariableAction::Undo()
 	{
 		return false;
 	}
-	return GetController()->OnExternalVariableRenamed(*NewVariableName, *OldVariableName, false);
+	return GetController()->OnExternalVariableRenamed(VariableGuid, *OldVariableName, false);
 }
 
 bool FRigVMRenameVariableAction::Redo()
@@ -1101,7 +1102,7 @@ bool FRigVMRenameVariableAction::Redo()
 	{
 		return false;
 	}
-	if(!GetController()->OnExternalVariableRenamed(*OldVariableName, *NewVariableName, false))
+	if(!GetController()->OnExternalVariableRenamed(VariableGuid, *NewVariableName, false))
 	{
 		return false;
 	}
diff --git a/Engine/Plugins/Runtime/RigVM/Source/RigVMDeveloper/Private/RigVMModel/RigVMGraph.cpp b/Engine/Plugins/Runtime/RigVM/Source/RigVMDeveloper/Private/RigVMModel/RigVMGraph.cpp
--- a/Engine/Plugins/Runtime/RigVM/Source/RigVMDeveloper/Private/RigVMModel/RigVMGraph.cpp
+++ b/Engine/Plugins/Runtime/RigVM/Source/RigVMDeveloper/Private/RigVMModel/RigVMGraph.cpp
@@ -418,6 +418,13 @@ TArray<FRigVMGraphVariableDescription> URigVMGraph::GetInputArguments() const
 				continue;
 			}
 			FRigVMGraphVariableDescription Description;
+			if (const URigVMPin* ReferencedPin = EntryNode->FindReferencedPin(Pin))
+			{
+				if (const URigVMCollapseNode* CollapseNode = EntryNode->FindReferencedNode())
+				{
+					Description.Guid = CollapseNode->FindPinGuid(ReferencedPin);
+				}
+			}
 			Description.Name = Pin->GetFName();
 			Description.CPPType = Pin->GetCPPType();
 			Description.CPPTypeObject = Pin->GetCPPTypeObject();
diff --git a/Engine/Plugins/Runtime/RigVM/Source/RigVMDeveloper/Private/RigVMModel/RigVMPin.cpp b/Engine/Plugins/Runtime/RigVM/Source/RigVMDeveloper/Private/RigVMModel/RigVMPin.cpp
--- a/Engine/Plugins/Runtime/RigVM/Source/RigVMDeveloper/Private/RigVMModel/RigVMPin.cpp
+++ b/Engine/Plugins/Runtime/RigVM/Source/RigVMDeveloper/Private/RigVMModel/RigVMPin.cpp
@@ -565,6 +565,22 @@ FRigVMGraphVariableDescription URigVMPin::GetInputVariableDescription() const
 	FRigVMGraphVariableDescription Description;
 	if (IsDefinedAsInputVariable())
 	{
+		const URigVMPin* PinToQueryGuid = this;
+		const URigVMCollapseNode* CollapseNode = Cast<URigVMCollapseNode>(GetNode());
+		if (!CollapseNode)
+		{
+			if (URigVMFunctionInterfaceNode* InterfaceNode = Cast<URigVMFunctionInterfaceNode>(GetNode()))
+			{
+				CollapseNode = InterfaceNode->FindReferencedNode();
+				PinToQueryGuid = InterfaceNode->FindReferencedPin(this);
+			}
+		}
+		
+		if (CollapseNode && PinToQueryGuid)
+		{
+			Description.Guid = CollapseNode->FindPinGuid(PinToQueryGuid);
+		}
+		
 		Description.Name = GetFName();
 		Description.CPPType = GetCPPType();
 		Description.CPPTypeObject = GetCPPTypeObject();
@@ -1912,7 +1928,7 @@ FRigVMExternalVariable URigVMPin::ToExternalVariable() const
 		VariableName = VariableName.Replace(TEXT("."), TEXT("_"));
 	}
 
-	ExternalVariable = RigVMTypeUtils::ExternalVariableFromCPPType(*VariableName, CPPType, GetCPPTypeObject(), false, false);
+	ExternalVariable = RigVMTypeUtils::ExternalVariableFromCPPType(FGuid(), *VariableName, CPPType, GetCPPTypeObject(), false, false);
 	
 	return ExternalVariable;
 }
diff --git a/Engine/Plugins/Runtime/RigVM/Source/RigVMDeveloper/Public/RigVMBlueprintLegacy.h b/Engine/Plugins/Runtime/RigVM/Source/RigVMDeveloper/Public/RigVMBlueprintLegacy.h
--- a/Engine/Plugins/Runtime/RigVM/Source/RigVMDeveloper/Public/RigVMBlueprintLegacy.h
+++ b/Engine/Plugins/Runtime/RigVM/Source/RigVMDeveloper/Public/RigVMBlueprintLegacy.h
@@ -482,7 +482,7 @@ protected:
 	virtual TObjectPtr<URigVMHost>& GetEditorHost() override { return EditorHost; }
 
 	static FName FindHostMemberVariableUniqueName(TSharedPtr<FKismetNameValidator> InNameValidator, const FString& InBaseName);
-	static int32 AddHostMemberVariable(URigVMBlueprint* InBlueprint, const FName& InVarName, FEdGraphPinType InVarType, bool bIsPublic, bool bIsReadOnly, FString InDefaultValue);
+	static int32 AddHostMemberVariable(URigVMBlueprint* InBlueprint, const FGuid InVarGuid, const FName& InVarName, FEdGraphPinType InVarType, bool bIsPublic, bool bIsReadOnly, FString InDefaultValue);
 	UE_API virtual FName AddHostMemberVariableFromExternal(FRigVMExternalVariable InVariableToCreate, FString InDefaultValue = FString()) override;
 public:
 
diff --git a/Engine/Plugins/Runtime/RigVM/Source/RigVMDeveloper/Public/RigVMEditorAsset.h b/Engine/Plugins/Runtime/RigVM/Source/RigVMDeveloper/Public/RigVMEditorAsset.h
--- a/Engine/Plugins/Runtime/RigVM/Source/RigVMDeveloper/Public/RigVMEditorAsset.h
+++ b/Engine/Plugins/Runtime/RigVM/Source/RigVMDeveloper/Public/RigVMEditorAsset.h
@@ -198,6 +198,7 @@ public:
 
 	virtual void GetAllEdGraphs(TArray<UEdGraph*>& Graphs) const = 0;
 	virtual TArray<FRigVMGraphVariableDescription> GetAssetVariables() const = 0;
+	UE_API virtual FRigVMGraphVariableDescription FindAssetVariableByGuid(const FGuid& InGuid) const;
 	virtual bool IsRegeneratingOnLoad() const = 0;
 
 	virtual FOnRigVMPreVariablesChanged& OnPreVariablesChanged() = 0;
@@ -657,9 +658,9 @@ protected:
 	FOnRigVMExternalVariablesChanged& OnExternalVariablesChanged() { return ExternalVariablesChangedEvent; }
 
 	UE_API virtual void OnVariableAdded(const FName& InVarName);
-	UE_API virtual void OnVariableRemoved(const FName& InVarName);
-	UE_API virtual void OnVariableRenamed(const FName& InOldVarName, const FName& InNewVarName);
-	UE_API virtual void OnVariableTypeChanged(const FName& InVarName, FEdGraphPinType InOldPinType, FEdGraphPinType InNewPinType);
+	UE_API virtual void OnVariableRemoved(const FGuid& InVarGuid);
+	UE_API virtual void OnVariableRenamed(const FGuid& InVarGuid, const FName& InNewVarName);
+	UE_API virtual void OnVariableTypeChanged(const FGuid& InVarGuid, FEdGraphPinType InOldPinType, FEdGraphPinType InNewPinType);
 
 	UE_API FName AddAssetVariableFromPinType(const FName& InName, const FEdGraphPinType& InType, bool bIsPublic = false, bool bIsReadOnly = false, FString InDefaultValue = FString());
 	
diff --git a/Engine/Plugins/Runtime/RigVM/Source/RigVMDeveloper/Public/RigVMModel/Nodes/RigVMCollapseNode.h b/Engine/Plugins/Runtime/RigVM/Source/RigVMDeveloper/Public/RigVMModel/Nodes/RigVMCollapseNode.h
--- a/Engine/Plugins/Runtime/RigVM/Source/RigVMDeveloper/Public/RigVMModel/Nodes/RigVMCollapseNode.h
+++ b/Engine/Plugins/Runtime/RigVM/Source/RigVMDeveloper/Public/RigVMModel/Nodes/RigVMCollapseNode.h
@@ -35,6 +35,7 @@ public:
 
 	UE_API FString GetEditorSubGraphName() const;
 
+	UE_API FGuid FindPinGuid(const URigVMPin* InPin) const;
 private:
 
 	UPROPERTY()
@@ -48,6 +49,9 @@ private:
 
 	UPROPERTY()
 	FString NodeDescription;
+	
+	UPROPERTY()
+	TMap<FName, FGuid> PinNameToGuid;
 
 	friend class URigVMController;
 	friend class FRigVMMinimalEnvironment;
diff --git a/Engine/Plugins/Runtime/RigVM/Source/RigVMDeveloper/Public/RigVMModel/Nodes/RigVMFunctionInterfaceNode.h b/Engine/Plugins/Runtime/RigVM/Source/RigVMDeveloper/Public/RigVMModel/Nodes/RigVMFunctionInterfaceNode.h
--- a/Engine/Plugins/Runtime/RigVM/Source/RigVMDeveloper/Public/RigVMModel/Nodes/RigVMFunctionInterfaceNode.h
+++ b/Engine/Plugins/Runtime/RigVM/Source/RigVMDeveloper/Public/RigVMModel/Nodes/RigVMFunctionInterfaceNode.h
@@ -7,6 +7,8 @@
 
 #define UE_API RIGVMDEVELOPER_API
 
+class URigVMCollapseNode;
+
 /**
  * The Function Interface node is is used as the base class for
  * both the entry and return nodes.
@@ -31,7 +33,11 @@ public:
 
 	// returns true if the interface pin is used for anything within the function / collapse graph
 	UE_API virtual bool IsInterfacePinUsed(const FName& InInterfacePinName) const;
+	
+	// find outer node which contains this interface node in its graph
+	UE_API const URigVMCollapseNode* FindReferencedNode() const;
 
+	// find pin in referenced node that represents this inner pin
 	UE_API const URigVMPin* FindReferencedPin(const URigVMPin* InPin) const;
 	UE_API const URigVMPin* FindReferencedPin(const FString& InPinPath) const;
 
diff --git a/Engine/Plugins/Runtime/RigVM/Source/RigVMDeveloper/Public/RigVMModel/Nodes/RigVMParameterNode.h b/Engine/Plugins/Runtime/RigVM/Source/RigVMDeveloper/Public/RigVMModel/Nodes/RigVMParameterNode.h
--- a/Engine/Plugins/Runtime/RigVM/Source/RigVMDeveloper/Public/RigVMModel/Nodes/RigVMParameterNode.h
+++ b/Engine/Plugins/Runtime/RigVM/Source/RigVMDeveloper/Public/RigVMModel/Nodes/RigVMParameterNode.h
@@ -61,7 +61,7 @@ public:
 	// Returns nullptr external variable matching this description
 	FRigVMExternalVariable ToExternalVariable() const
 	{
-		return FRigVMExternalVariable::Make(Name, CPPType, CPPTypeObject, false, false);
+		return FRigVMExternalVariable::Make(FGuid(), Name, CPPType, CPPTypeObject, false, false);
 	}
 };
 
diff --git a/Engine/Plugins/Runtime/RigVM/Source/RigVMDeveloper/Public/RigVMModel/Nodes/RigVMVariableNode.h b/Engine/Plugins/Runtime/RigVM/Source/RigVMDeveloper/Public/RigVMModel/Nodes/RigVMVariableNode.h
--- a/Engine/Plugins/Runtime/RigVM/Source/RigVMDeveloper/Public/RigVMModel/Nodes/RigVMVariableNode.h
+++ b/Engine/Plugins/Runtime/RigVM/Source/RigVMDeveloper/Public/RigVMModel/Nodes/RigVMVariableNode.h
@@ -80,8 +80,14 @@ public:
 
 	// Get the subtitle of the node
 	virtual FString GetNodeSubTitle() const override;
+	
+	// Get variable Guid
+	FGuid GetVariableGuid() const { return VariableGuid; }
 
 private:
+	
+	UPROPERTY()
+	FGuid VariableGuid;
 
 	static const inline TCHAR* VariableName = TEXT("Variable");
 	static const inline TCHAR* ValueName = TEXT("Value");
diff --git a/Engine/Plugins/Runtime/RigVM/Source/RigVMDeveloper/Public/RigVMModel/RigVMClient.h b/Engine/Plugins/Runtime/RigVM/Source/RigVMDeveloper/Public/RigVMModel/RigVMClient.h
--- a/Engine/Plugins/Runtime/RigVM/Source/RigVMDeveloper/Public/RigVMModel/RigVMClient.h
+++ b/Engine/Plugins/Runtime/RigVM/Source/RigVMDeveloper/Public/RigVMModel/RigVMClient.h
@@ -255,6 +255,7 @@ public:
 	UE_API URigVMGraph* GetModel(const FString& InNodePathOrName) const;
 	UE_API URigVMGraph* GetModel(const UObject* InEditorSideObject) const;
 	UE_API void RefreshAllModels(ERigVMLoadType InLoadType, bool bEnablePostLoadHashing, bool& bIsCompiling);
+	UE_API void RefreshVariableNodes(const TArray<FRigVMExternalVariable> InPublicVariables);
 	UE_API const TArray<TObjectPtr<URigVMGraph>>& GetModels() const;
 	UE_API TArray<URigVMGraph*> GetAllModels(bool bIncludeFunctionLibrary, bool bRecursive) const;
 	UE_API TArray<URigVMGraph*> GetAllModelsLeavesFirst(bool bIncludeFunctionLibrary) const;
diff --git a/Engine/Plugins/Runtime/RigVM/Source/RigVMDeveloper/Public/RigVMModel/RigVMController.h b/Engine/Plugins/Runtime/RigVM/Source/RigVMDeveloper/Public/RigVMModel/RigVMController.h
--- a/Engine/Plugins/Runtime/RigVM/Source/RigVMDeveloper/Public/RigVMModel/RigVMController.h
+++ b/Engine/Plugins/Runtime/RigVM/Source/RigVMDeveloper/Public/RigVMModel/RigVMController.h
@@ -180,6 +180,7 @@ struct FRigVMPinInfo
 	UE_API void CorrectExecuteTypeIndex();
 
 	int32 ParentIndex;
+	FGuid Guid;
 	FName Name;
 	ERigVMPinDirection Direction;
 	TRigVMTypeIndex TypeIndex;
@@ -407,18 +408,26 @@ public:
 	UE_API URigVMVariableNode* AddVariableNodeFromObjectPath(const FName& InVariableName, const FString& InCPPType, const FString& InCPPTypeObjectPath, bool bIsGetter, const FString& InDefaultValue, const FVector2D& InPosition = FVector2D::ZeroVector, const FString& InNodeName = TEXT(""), bool bSetupUndoRedo = true, bool bPrintPythonCommand = false);
 
 	// Refreshes the variable node with the new data
-	UFUNCTION(BlueprintCallable, Category = RigVMController)
+	UFUNCTION(BlueprintCallable, Category = RigVMController, meta = (DeprecatedFunction, DeprecationMessage = "Use RefreshVariableNode that takes a FGuid"))
 	UE_API void RefreshVariableNode(const FName& InNodeName, const FName& InVariableName, const FString& InCPPType, UObject* InCPPTypeObject, bool bSetupUndoRedo, bool bSetupOrphanPins = true);
+	
+	// Refreshes the variable node with the new data
+	UE_API void RefreshVariableNode(const FName& InNodeName, const FGuid& InVariableGuid, const FName& InVariableName, const FString& InCPPType, UObject* InCPPTypeObject, bool bSetupUndoRedo, bool bSetupOrphanPins = true);
 
 	// Removes all nodes related to a given variable
-	UE_API void OnExternalVariableRemoved(const FName& InVarName, bool bSetupUndoRedo);
+	UE_API void OnExternalVariableRemoved(const FGuid& InVarGuid, bool bSetupUndoRedo);
 
 	// Renames the variable name in all relevant nodes
-	UE_API bool OnExternalVariableRenamed(const FName& InOldVarName, const FName& InNewVarName, bool bSetupUndoRedo);
+	UE_API bool OnExternalVariableRenamed(const FGuid& InVarGuid, const FName& InNewVarName, bool bSetupUndoRedo);
 
 	// Changes the data type of all nodes matching a given variable name
+	UE_DEPRECATED(5.8, "Use OnExternalVariableTypeChanged that receives the variable Guid")
 	UE_API void OnExternalVariableTypeChanged(const FName& InVarName, const FString& InCPPType, UObject* InCPPTypeObject, bool bSetupUndoRedo);
+	UE_API void OnExternalVariableTypeChanged(const FGuid& InVarGuid, const FString& InCPPType, UObject* InCPPTypeObject, bool bSetupUndoRedo);
+	UE_DEPRECATED(5.8, "Use OnExternalVariableTypeChangedFromObjectPath that receives the variable Guid")
 	UE_API void OnExternalVariableTypeChangedFromObjectPath(const FName& InVarName, const FString& InCPPType, const FString& InCPPTypeObjectPath, bool bSetupUndoRedo);
+	UE_API void OnExternalVariableTypeChangedFromObjectPath(const FGuid& InVarGuid, const FString& InCPPType, const FString& InCPPTypeObjectPath, bool bSetupUndoRedo);
+	
 
 	// Refreshes the variable node with the new data
 	UFUNCTION(BlueprintCallable, Category = RigVMController)
@@ -826,6 +835,7 @@ public:
 	// This causes a VariableRenamed modified event.
 	UFUNCTION(BlueprintCallable, Category = RigVMController, meta=(DeprecatedFunction))
 	UE_API bool RenameVariable(const FName& InOldName, const FName& InNewName, bool bSetupUndoRedo = true);
+	UE_API bool RenameVariable(const FGuid& InGuid, const FName& InNewName, bool bSetupUndoRedo = true);
 
 	// Renames a parameter in the graph.
 	// This causes a ParameterRenamed modified event.
@@ -1450,6 +1460,7 @@ private:
 	UE_API bool ChangePinType(URigVMPin* InPin, const FString& InCPPType,UObject* InCPPTypeObject, bool bSetupUndoRedo, bool bSetupOrphanPins = true, bool bBreakLinks = true, bool bRemoveSubPins = true, bool bInitializeDefaultValue = true);
 	UE_API bool ChangePinType(URigVMPin* InPin, TRigVMTypeIndex InTypeIndex, bool bSetupUndoRedo, bool bSetupOrphanPins = true, bool bBreakLinks = true, bool bRemoveSubPins = true, bool bInitializeDefaultValue = true);
 
+	UE_API bool MarkPackageDirty(const UObject* InObjectToDirty) const;
 	UE_API bool RenameObject(UObject* InObjectToRename, const TCHAR* InNewName, UObject* InNewOuter = nullptr, ERenameFlags InFlags = REN_None) const;
 	UE_API void DestroyObject(UObject* InObjectToDestroy) const ;
 	static UE_API URigVMPin* MakeExecutePin(URigVMNode* InNode, const FName& InName);
@@ -1464,13 +1475,15 @@ private:
 	// they can get out of sync when the variable references a user defined struct
 	UE_API bool EnsureLocalVariableValidity();
 	
-	UE_API FRigVMExternalVariable GetVariableByName(const FName& InExternalVariableName, const bool bIncludeInputArguments = false) const;
 	UE_API TArray<FRigVMExternalVariable> GetAllVariables(const bool bIncludeInputArguments = false) const;
 
 	UE_API void RefreshFunctionReferences(const URigVMLibraryNode* InFunctionDefinition, bool bSetupUndoRedo, bool bLoadIfNecessary);
 	UE_API void PropagateNotificationToFunctionReferences(const URigVMLibraryNode* InFunctionDefinition, ERigVMGraphNotifType InNotifType, UObject* InSubject, bool bLoadIfNecessary);
 
 public:
+	
+	UE_API FRigVMExternalVariable GetVariableByName(const FName& InExternalVariableName, const bool bIncludeInputArguments = false) const;
+	UE_API FRigVMExternalVariable GetVariableByGuid(const FGuid& InExternalVariableGuid, const bool bIncludeInputArguments = false) const;
 
 	struct FLinkedPath
 	{
@@ -1574,6 +1587,7 @@ protected:
 	UE_API FRigVMClientPatchResult PatchPinDefaultValues();
 	UE_API FRigVMClientPatchResult PatchUserDefinedStructPinNames();
 	UE_API FRigVMClientPatchResult PatchLocalVariableTypes();
+	UE_API FRigVMClientPatchResult PatchPinGuidInCollapseNodes();
 	UE_API FRigVMClientPatchResult PatchExternalVariablesToArguments(bool bInSetupUndoRedo=false);
 
 	UE_API ERigVMPinDefaultValueType GetDefaultValueType(const URigVMPin* InPin, const FString& InDefaultValue) const;
@@ -1647,6 +1661,7 @@ private:
 	mutable FDelegateHandle ActionStackHandle;
 
 	bool bSuspendNotifications;
+	bool bSuspendMarkingPackageDirty;
 	bool bSuspendRefreshingFunctionReferences;
 	bool bReportWarningsAndErrors;
 	bool bIgnoreRerouteCompactnessChanges;
diff --git a/Engine/Plugins/Runtime/RigVM/Source/RigVMDeveloper/Public/RigVMModel/RigVMControllerActions.h b/Engine/Plugins/Runtime/RigVM/Source/RigVMDeveloper/Public/RigVMModel/RigVMControllerActions.h
--- a/Engine/Plugins/Runtime/RigVM/Source/RigVMDeveloper/Public/RigVMModel/RigVMControllerActions.h
+++ b/Engine/Plugins/Runtime/RigVM/Source/RigVMDeveloper/Public/RigVMModel/RigVMControllerActions.h
@@ -811,11 +811,14 @@ struct FRigVMRenameVariableAction : public FRigVMBaseAction
 public:
 
 	FRigVMRenameVariableAction();
-	FRigVMRenameVariableAction(URigVMController* InController, const FName& InOldVariableName, const FName& InNewVariableName);
+	FRigVMRenameVariableAction(URigVMController* InController, const FGuid& InVariableGuid, const FName& InOldVariableName, const FName& InNewVariableName);
 	virtual ~FRigVMRenameVariableAction() = default;
 	virtual UScriptStruct* GetScriptStruct() const override { return FRigVMRenameVariableAction::StaticStruct(); }
 	virtual bool Undo() override;
 	virtual bool Redo() override;
+	
+	UPROPERTY()
+	FGuid VariableGuid;
 
 	UPROPERTY()
 	FString OldVariableName;
diff --git a/Engine/Plugins/Runtime/RigVM/Source/RigVMDeveloper/Public/RigVMModel/RigVMPin.h b/Engine/Plugins/Runtime/RigVM/Source/RigVMDeveloper/Public/RigVMModel/RigVMPin.h
--- a/Engine/Plugins/Runtime/RigVM/Source/RigVMDeveloper/Public/RigVMModel/RigVMPin.h
+++ b/Engine/Plugins/Runtime/RigVM/Source/RigVMDeveloper/Public/RigVMModel/RigVMPin.h
@@ -675,6 +675,8 @@ private:
 	friend class URigVMLibraryNode;
 	friend class URigVMEdGraphNode;
 	friend struct FRigVMClient;
+	friend class SRigVMGraphNode;
+	friend struct FRigVMPinInfoArray;
 };
 
 class FRigVMPinDefaultValueImportErrorContext : public FOutputDevice
diff --git a/Engine/Plugins/Runtime/RigVM/Source/RigVMEditor/Private/Editor/RigVMEditor.cpp b/Engine/Plugins/Runtime/RigVM/Source/RigVMEditor/Private/Editor/RigVMEditor.cpp
--- a/Engine/Plugins/Runtime/RigVM/Source/RigVMEditor/Private/Editor/RigVMEditor.cpp
+++ b/Engine/Plugins/Runtime/RigVM/Source/RigVMEditor/Private/Editor/RigVMEditor.cpp
@@ -3822,7 +3822,8 @@ void FRigVMEditorBase::HandleVariableDroppedFromBlueprint(UObject* InSubject, FP
 	URigVMController* Controller = GetFocusedController();
 	check(Controller);
 
-	FRigVMExternalVariable ExternalVariable = FRigVMExternalVariable::Make(InVariableToDrop, nullptr);
+	FRigVMExternalVariable ExternalVariable = Controller->GetVariableByName(InVariableToDrop->GetFName());
+	ensure(ExternalVariable.GetProperty() == InVariableToDrop);
 	if (!ExternalVariable.IsValid(true /* allow null ptr */))
 	{
 		return;
diff --git a/Engine/Plugins/Runtime/RigVM/Source/RigVMEditor/Private/Widgets/SRigVMGraphPinUserDataPath.cpp b/Engine/Plugins/Runtime/RigVM/Source/RigVMEditor/Private/Widgets/SRigVMGraphPinUserDataPath.cpp
--- a/Engine/Plugins/Runtime/RigVM/Source/RigVMEditor/Private/Widgets/SRigVMGraphPinUserDataPath.cpp
+++ b/Engine/Plugins/Runtime/RigVM/Source/RigVMEditor/Private/Widgets/SRigVMGraphPinUserDataPath.cpp
@@ -157,7 +157,7 @@ FLinearColor SRigVMUserDataPath::GetUserDataColor(const UNameSpacedUserData::FUs
 {
 	if (InUserData)
 	{
-		const FRigVMExternalVariable ExternalVariable = FRigVMExternalVariable::Make(InUserData->GetProperty(), nullptr);
+		const FRigVMExternalVariable ExternalVariable = FRigVMExternalVariable::Make(FGuid(), InUserData->GetProperty(), nullptr);
 		if(ExternalVariable.IsValid(true))
 		{
 			const URigVMEdGraphSchema* Schema = GetDefault<URigVMEdGraphSchema>();
@@ -199,7 +199,7 @@ void SRigVMUserDataPath::FillUserDataPathMenu(FMenuBuilder& InMenuBuilder, FStri
 
 		for(const UNameSpacedUserData::FUserData* UserData : UserDataArray)
 		{
-			const FRigVMExternalVariable ExternalVariable = FRigVMExternalVariable::Make(UserData->GetProperty(),nullptr);
+			const FRigVMExternalVariable ExternalVariable = FRigVMExternalVariable::Make(FGuid(), UserData->GetProperty(),nullptr);
 
 			FEdGraphPinType PinType;
 
@@ -296,7 +296,7 @@ void SRigVMUserDataPath::HandleSetUserDataPath(FString InUserDataPath)
 
 		if(const UNameSpacedUserData::FUserData* UserData = GetUserData())
 		{
-			const FRigVMExternalVariable ExternalVariable = FRigVMExternalVariable::Make(UserData->GetProperty(), nullptr);
+			const FRigVMExternalVariable ExternalVariable = FRigVMExternalVariable::Make(FGuid(), UserData->GetProperty(), nullptr);
 			if(ExternalVariable.IsValid(true))
 			{
 				const TRigVMTypeIndex ExpectedTypeIndex = FRigVMRegistry::Get().GetTypeIndexFromCPPType(ExternalVariable.GetExtendedCPPType().ToString()); 
diff --git a/Engine/Plugins/Runtime/RigVM/Source/RigVMEditor/Private/Widgets/SRigVMGraphPinVariableBinding.cpp b/Engine/Plugins/Runtime/RigVM/Source/RigVMEditor/Private/Widgets/SRigVMGraphPinVariableBinding.cpp
--- a/Engine/Plugins/Runtime/RigVM/Source/RigVMEditor/Private/Widgets/SRigVMGraphPinVariableBinding.cpp
+++ b/Engine/Plugins/Runtime/RigVM/Source/RigVMEditor/Private/Widgets/SRigVMGraphPinVariableBinding.cpp
@@ -188,7 +188,7 @@ bool SRigVMGraphVariableBinding::OnCanBindProperty(FProperty* InProperty) const
 
 	if (InProperty)
 	{
-		const FRigVMExternalVariable ExternalVariable = FRigVMExternalVariable::Make(InProperty, nullptr);
+		const FRigVMExternalVariable ExternalVariable = FRigVMExternalVariable::Make(FGuid(), InProperty, nullptr);
 		if(ModelPins.Num() > 0)
 		{
 			return ModelPins[0]->CanBeBoundToVariable(ExternalVariable);
