Control Rig: Fix variable Guid not being set in the RigVMBlueprint class when getting variable descriptions
#rnx

#Synced-CL 51339418

diff --git a/Engine/Plugins/Runtime/RigVM/Source/RigVMDeveloper/Private/RigVMBlueprintLegacy.cpp b/Engine/Plugins/Runtime/RigVM/Source/RigVMDeveloper/Private/RigVMBlueprintLegacy.cpp
--- a/Engine/Plugins/Runtime/RigVM/Source/RigVMDeveloper/Private/RigVMBlueprintLegacy.cpp
+++ b/Engine/Plugins/Runtime/RigVM/Source/RigVMDeveloper/Private/RigVMBlueprintLegacy.cpp
@@ -875,6 +875,7 @@ TArray<FRigVMGraphVariableDescription> URigVMBlueprint::GetAssetVariables() cons
 	for (const FBPVariableDescription& BPVariable : NewVariables)
 	{
 		FRigVMGraphVariableDescription NewVariable;
+		NewVariable.Guid = BPVariable.VarGuid;
 		NewVariable.Name = BPVariable.VarName;
 		NewVariable.DefaultValue = BPVariable.DefaultValue;
 		if (NewVariable.DefaultValue.IsEmpty() && GeneratedClass)
diff --git a/Engine/Plugins/Runtime/RigVM/Source/RigVMDeveloper/Private/RigVMModel/Nodes/RigVMCollapseNode.cpp b/Engine/Plugins/Runtime/RigVM/Source/RigVMDeveloper/Private/RigVMModel/Nodes/RigVMCollapseNode.cpp
--- a/Engine/Plugins/Runtime/RigVM/Source/RigVMDeveloper/Private/RigVMModel/Nodes/RigVMCollapseNode.cpp
+++ b/Engine/Plugins/Runtime/RigVM/Source/RigVMDeveloper/Private/RigVMModel/Nodes/RigVMCollapseNode.cpp
@@ -54,9 +54,12 @@ FString URigVMCollapseNode::GetEditorSubGraphName() const
 
 FGuid URigVMCollapseNode::FindPinGuid(const URigVMPin* InPin) const
 {
-	if (const FGuid* Guid = PinNameToGuid.Find(InPin->GetFName()))
+	if (URigVMPin* RootPin = InPin->GetRootPin())
 	{
-		return *Guid;
+		if (const FGuid* Guid = PinNameToGuid.Find(RootPin->GetFName()))
+		{
+			return *Guid;
+		}
 	}
 	return FGuid();
 }
diff --git a/Engine/Plugins/Runtime/RigVM/Source/RigVMDeveloper/Private/RigVMModel/RigVMLink.cpp b/Engine/Plugins/Runtime/RigVM/Source/RigVMDeveloper/Private/RigVMModel/RigVMLink.cpp
--- a/Engine/Plugins/Runtime/RigVM/Source/RigVMDeveloper/Private/RigVMModel/RigVMLink.cpp
+++ b/Engine/Plugins/Runtime/RigVM/Source/RigVMDeveloper/Private/RigVMModel/RigVMLink.cpp
@@ -220,8 +220,11 @@ void URigVMLink::UpdatePinPaths()
 
 void URigVMLink::UpdatePinPointers() const
 {
-	SourcePin = GetGraph()->FindPin(SourcePinPath);
-	TargetPin = GetGraph()->FindPin(TargetPinPath);
+	if (URigVMGraph* Graph = GetGraph())
+	{
+		SourcePin = Graph->FindPin(SourcePinPath);
+		TargetPin = Graph->FindPin(TargetPinPath);
+	}
 }
 
 URigVMPin* URigVMLink::GetOppositePin(const URigVMPin* InPin) const
