Control Rig: Fix cannot promote to variable in control rig asset
#rnx

#Synced-CL 51339699

diff --git a/Engine/Plugins/Runtime/RigVM/Source/RigVM/Private/RigVMRuntimeAsset.cpp b/Engine/Plugins/Runtime/RigVM/Source/RigVM/Private/RigVMRuntimeAsset.cpp
--- a/Engine/Plugins/Runtime/RigVM/Source/RigVM/Private/RigVMRuntimeAsset.cpp
+++ b/Engine/Plugins/Runtime/RigVM/Source/RigVM/Private/RigVMRuntimeAsset.cpp
@@ -440,11 +440,30 @@ FName URigVMRuntimeAsset::AddHostMemberVariableFromExternal(FRigVMExternalVariab
 	FRigVMAssetVariablesChangeScope VariablesChangeScope(this, InVariableToCreate.GetName());
 	Modify();
 	
+	TSet<FString> VariableNames;
+	for (int32 i=0; i<Variables.Num(); i++)
+	{
+		const FProperty* Property = Variables.GetProperty(i);
+		if (Property->HasMetaData(FRigVMPropertyDescription::MD_DisplayName))
+		{
+			VariableNames.Add(Property->GetMetaData(FRigVMPropertyDescription::MD_DisplayName));
+		}
+	}
+	
+	FString VariableName = InVariableToCreate.GetName().ToString();
+	FRigVMPropertyDescription::SanitizeName(VariableName, true);
+	int32 Index=1;
+	while (VariableNames.Contains(VariableName))
+	{
+		VariableName = FString::Printf(TEXT("%s_%d"), *InVariableToCreate.GetName().ToString(), Index++);
+		FRigVMPropertyDescription::SanitizeName(VariableName, true);
+	}
+	
 	FString CPPType;
 	UObject* CPPTypeObject = nullptr;
 	RigVMTypeUtils::CPPTypeFromExternalVariable(InVariableToCreate, CPPType, &CPPTypeObject);
 
-	FRigVMPropertyDescription Description(InVariableToCreate.GetName(), CPPType, CPPTypeObject, InDefaultValue, true);
+	FRigVMPropertyDescription Description(*VariableName, CPPType, CPPTypeObject, InDefaultValue, true);
 	Description.PropertyFlags = CPF_Edit | CPF_BlueprintVisible | CPF_DisableEditOnInstance | CPF_Interp;
 	if (InVariableToCreate.IsPublic())
 	{
diff --git a/Engine/Plugins/Runtime/RigVM/Source/RigVMDeveloper/Private/EdGraph/RigVMEdGraphSchema.cpp b/Engine/Plugins/Runtime/RigVM/Source/RigVMDeveloper/Private/EdGraph/RigVMEdGraphSchema.cpp
--- a/Engine/Plugins/Runtime/RigVM/Source/RigVMDeveloper/Private/EdGraph/RigVMEdGraphSchema.cpp
+++ b/Engine/Plugins/Runtime/RigVM/Source/RigVMDeveloper/Private/EdGraph/RigVMEdGraphSchema.cpp
@@ -2479,7 +2479,7 @@ bool URigVMEdGraphSchema::RequestVariableDropOnPin(UEdGraph* InGraph, FProperty*
 			if (const URigVMPin* ModelPin = Graph->GetModel()->FindPin(InPin->GetName()))
 			{
 				URigVMController* Controller = Graph->GetController();
-				const FRigVMExternalVariable ExternalVariable = Controller->GetVariableByName(InVariableToDrop->GetFName());
+				const FRigVMExternalVariable ExternalVariable = Controller->GetVariableByProperty(InVariableToDrop);
 				ensure(ExternalVariable.GetProperty() == InVariableToDrop);
 				return RequestVariableDropOnPin(InGraph, ExternalVariable, InPin, FDeprecateSlateVector2D(InDropPosition), FDeprecateSlateVector2D(InScreenPosition));
 			}
diff --git a/Engine/Plugins/Runtime/RigVM/Source/RigVMDeveloper/Private/RigVMModel/RigVMController.cpp b/Engine/Plugins/Runtime/RigVM/Source/RigVMDeveloper/Private/RigVMModel/RigVMController.cpp
--- a/Engine/Plugins/Runtime/RigVM/Source/RigVMDeveloper/Private/RigVMModel/RigVMController.cpp
+++ b/Engine/Plugins/Runtime/RigVM/Source/RigVMDeveloper/Private/RigVMModel/RigVMController.cpp
@@ -22605,6 +22605,23 @@ bool URigVMController::EnsureLocalVariableValidity()
 	return bChanged;
 }
 
+FRigVMExternalVariable URigVMController::GetVariableByProperty(const FProperty* InExternalVariableProperty) const
+{
+	if (GetExternalVariablesDelegate.IsBound())
+	{
+		const TArray<FRigVMExternalVariable> Variables = GetExternalVariablesDelegate.Execute(GetGraph());
+		for (const FRigVMExternalVariable& Variable : Variables)
+		{
+			if (Variable.GetProperty() == InExternalVariableProperty)
+			{
+				return Variable;
+			}
+		}
+	}
+	
+	return FRigVMExternalVariable();
+}
+
 FRigVMExternalVariable URigVMController::GetVariableByName(const FName& InExternalVariableName, const bool bIncludeInputArguments) const
 {
 	TArray<FRigVMExternalVariable> Variables = GetAllVariables(bIncludeInputArguments);
diff --git a/Engine/Plugins/Runtime/RigVM/Source/RigVMDeveloper/Public/RigVMModel/RigVMController.h b/Engine/Plugins/Runtime/RigVM/Source/RigVMDeveloper/Public/RigVMModel/RigVMController.h
--- a/Engine/Plugins/Runtime/RigVM/Source/RigVMDeveloper/Public/RigVMModel/RigVMController.h
+++ b/Engine/Plugins/Runtime/RigVM/Source/RigVMDeveloper/Public/RigVMModel/RigVMController.h
@@ -1482,6 +1482,7 @@ private:
 
 public:
 	
+	UE_API FRigVMExternalVariable GetVariableByProperty(const FProperty* InExternalVariableProperty) const;
 	UE_API FRigVMExternalVariable GetVariableByName(const FName& InExternalVariableName, const bool bIncludeInputArguments = false) const;
 	UE_API FRigVMExternalVariable GetVariableByGuid(const FGuid& InExternalVariableGuid, const bool bIncludeInputArguments = false) const;
 
diff --git a/Engine/Plugins/Runtime/RigVM/Source/RigVMEditor/Private/Editor/RigVMEditor.cpp b/Engine/Plugins/Runtime/RigVM/Source/RigVMEditor/Private/Editor/RigVMEditor.cpp
--- a/Engine/Plugins/Runtime/RigVM/Source/RigVMEditor/Private/Editor/RigVMEditor.cpp
+++ b/Engine/Plugins/Runtime/RigVM/Source/RigVMEditor/Private/Editor/RigVMEditor.cpp
@@ -3822,7 +3822,7 @@ void FRigVMEditorBase::HandleVariableDroppedFromBlueprint(UObject* InSubject, FP
 	URigVMController* Controller = GetFocusedController();
 	check(Controller);
 
-	FRigVMExternalVariable ExternalVariable = Controller->GetVariableByName(InVariableToDrop->GetFName());
+	FRigVMExternalVariable ExternalVariable = Controller->GetVariableByProperty(InVariableToDrop);
 	ensure(ExternalVariable.GetProperty() == InVariableToDrop);
 	if (!ExternalVariable.IsValid(true /* allow null ptr */))
 	{
