More descriptive compiler error message w.r.t unused pins

#Synced-CL 51408661

diff --git a/Engine/Plugins/Runtime/RigVM/Source/RigVMDeveloper/Private/RigVMCompiler/RigVMCompiler.cpp b/Engine/Plugins/Runtime/RigVM/Source/RigVMDeveloper/Private/RigVMCompiler/RigVMCompiler.cpp
--- a/Engine/Plugins/Runtime/RigVM/Source/RigVMDeveloper/Private/RigVMCompiler/RigVMCompiler.cpp
+++ b/Engine/Plugins/Runtime/RigVM/Source/RigVMDeveloper/Private/RigVMCompiler/RigVMCompiler.cpp
@@ -1222,8 +1222,19 @@ bool URigVMCompiler::Compile(const FRigVMCompileSettings& InSettings, TArray<URi
 
 			if(!VisitedGraphController->RemoveUnusedOrphanedPins(ModelNode))
 			{
-				static const FString LinkedMessage = TEXT("Node @@ uses pins that no longer exist. Please rewire the links and re-compile.");
-				Settings.ASTSettings.Report(EMessageSeverity::Error, ModelNode, LinkedMessage);
+				static const FString LinkedMessage = TEXT("Node @@ uses [{0}] pins that no longer exist. Please rewire the links and re-compile.");
+
+				FString PinNames;
+				for (URigVMPin* Pin : ModelNode->GetOrphanedPins())
+				{
+					if (PinNames.Len())
+					{
+						PinNames.Append(TEXT(", "));
+					}
+					PinNames.Append(Pin->GetDisplayName().ToString());
+				}
+				
+				Settings.ASTSettings.Report(EMessageSeverity::Error, ModelNode, FString::Format(*LinkedMessage, {PinNames}));
 				bEncounteredGraphError = true;
 			}
 
