Control Rig: Added instance pools for classes in RigVMDetailsWrapperObject to avoid crashes when too many instances are created.
#rnx

#Synced-CL 51454554

diff --git a/Engine/Plugins/Runtime/RigVM/Source/RigVMEditor/Private/Editor/RigVMDetailsViewWrapperObject.cpp b/Engine/Plugins/Runtime/RigVM/Source/RigVMEditor/Private/Editor/RigVMDetailsViewWrapperObject.cpp
--- a/Engine/Plugins/Runtime/RigVM/Source/RigVMEditor/Private/Editor/RigVMDetailsViewWrapperObject.cpp
+++ b/Engine/Plugins/Runtime/RigVM/Source/RigVMEditor/Private/Editor/RigVMDetailsViewWrapperObject.cpp
@@ -20,6 +20,7 @@
 TMap<URigVMDetailsViewWrapperObject::FPerClassInfo, UClass*> URigVMDetailsViewWrapperObject::InfoToClass;
 TMap<UClass*, URigVMDetailsViewWrapperObject::FPerClassInfo> URigVMDetailsViewWrapperObject::ClassToInfo;
 TSet<UClass*> URigVMDetailsViewWrapperObject::OutdatedClassToRecreate;
+TMap<UClass*, TArray<URigVMDetailsViewWrapperObject*>> URigVMDetailsViewWrapperObject::InstancePool;
 
 static const TCHAR DiscardedWrapperClassTemplateName[] = TEXT("DiscardedWrapperClassTemplate");
 
@@ -193,10 +194,14 @@ URigVMDetailsViewWrapperObject* URigVMDetailsViewWrapperObject::MakeInstance(UCl
 	{
 		return nullptr;
 	}
+	
+	URigVMDetailsViewWrapperObject* Instance = GetOrCreateNewInstance(WrapperClass, InOuter);
 
-	URigVMDetailsViewWrapperObject* Instance = NewObject<URigVMDetailsViewWrapperObject>(InOuter, WrapperClass, NAME_None, RF_Public | RF_Transient | RF_TextExportTransient | RF_DuplicateTransient);
-	Instance->SetContent(InStructMemory, InStruct);
-	Instance->SetSubject(InSubject);
+	if (Instance)
+	{
+		Instance->SetContent(InStructMemory, InStruct);
+		Instance->SetSubject(InSubject);
+	}
 
 	return Instance;
 }
@@ -651,12 +656,16 @@ URigVMDetailsViewWrapperObject* URigVMDetailsViewWrapperObject::MakeInstance(UCl
 	{
 		return nullptr;
 	}
+	
+	URigVMDetailsViewWrapperObject* Instance = GetOrCreateNewInstance(WrapperClass, InOuter);
 
-	URigVMDetailsViewWrapperObject* Instance = NewObject<URigVMDetailsViewWrapperObject>(InOuter, WrapperClass, NAME_None, RF_Public | RF_Transient | RF_TextExportTransient | RF_DuplicateTransient);
-	Instance->SetContent(InSubject);
-	Instance->SetSubject(InSubject);
+	if (Instance)
+	{
+		Instance->SetContent(InSubject);
+		Instance->SetSubject(InSubject);
 
-	InSubject->GetGraph()->OnModified().AddUObject(Instance, &URigVMDetailsViewWrapperObject::HandleModifiedEvent);
+		InSubject->GetGraph()->OnModified().AddUObject(Instance, &URigVMDetailsViewWrapperObject::HandleModifiedEvent);
+	}
 	
 	return Instance;
 
@@ -667,6 +676,7 @@ void URigVMDetailsViewWrapperObject::MarkOutdatedClass(UClass* InClass)
 	if (InClass)
 	{
 		OutdatedClassToRecreate.Add(InClass);
+		ClearClassPool(InClass);
 	}
 }
 
@@ -690,6 +700,103 @@ bool URigVMDetailsViewWrapperObject::IsValidClass(UClass* InClass)
 	return true;
 }
 
+URigVMDetailsViewWrapperObject* URigVMDetailsViewWrapperObject::GetOrCreateNewInstance(const UClass* InWrapperObjectClass, UObject* InOuter)
+{
+	URigVMDetailsViewWrapperObject* Instance = nullptr;
+
+	// Try to get from pool first
+	if (TArray<URigVMDetailsViewWrapperObject*>* Pool = InstancePool.Find(InWrapperObjectClass))
+	{
+		// Filter out any invalid instances
+		Pool->RemoveAll([](URigVMDetailsViewWrapperObject* Candidate)
+		{
+			return !IsValid(Candidate);
+		});
+		
+		if (Pool->Num() > 0)
+		{
+			Instance = Pool->Pop();
+			
+			// Rename to new outer if needed
+			if (Instance->GetOuter() != InOuter)
+			{
+				Instance->Rename(nullptr, InOuter, REN_DoNotDirty | REN_DontCreateRedirectors | REN_NonTransactional);
+			}
+		}
+	}
+	
+	// Otherwise, create new
+	if (Instance == nullptr || !IsValid(Instance))
+	{
+		Instance = NewObject<URigVMDetailsViewWrapperObject>(InOuter, InWrapperObjectClass, NAME_None, RF_Public | RF_Transient | RF_TextExportTransient | RF_DuplicateTransient);
+	}
+	
+	return Instance;
+}
+
+void URigVMDetailsViewWrapperObject::ReturnToPool(URigVMDetailsViewWrapperObject* InInstance)
+{
+	if (!InInstance || !IsValid(InInstance))
+		return;
+
+	UClass* InstanceClass = InInstance->GetClass();
+
+	// Don't pool instances of outdated classes
+	if (OutdatedClassToRecreate.Contains(InstanceClass))
+	{
+		InInstance->RemoveFromRoot();
+		InInstance->MarkAsGarbage();
+		return;
+	}
+
+	TArray<URigVMDetailsViewWrapperObject*>& Pool = InstancePool.FindOrAdd(InstanceClass);
+
+	// Don't exceed max pool size
+	if (Pool.Num() >= MaxPoolSizePerClass)
+	{
+		InInstance->RemoveFromRoot();
+		InInstance->MarkAsGarbage();
+		return;
+	}
+
+	// Reset state and add to pool
+	InInstance->ResetForReuse();
+	Pool.Add(InInstance);
+}
+
+void URigVMDetailsViewWrapperObject::ClearClassPool(UClass* InClass)
+{
+	TArray<URigVMDetailsViewWrapperObject*> Pool;
+	if (InstancePool.RemoveAndCopyValue(InClass, Pool))
+	{
+		for (URigVMDetailsViewWrapperObject* Instance : Pool)
+		{
+			if (IsValid(Instance))
+			{
+				Instance->RemoveFromRoot();
+				Instance->MarkAsGarbage();
+			}
+		}
+	}
+}
+
+void URigVMDetailsViewWrapperObject::ClearAllClassPools()
+{
+	// During shutdown, we don't need to remove instances individually
+	if (!IsEngineExitRequested())
+	{
+		TArray<UClass*> Keys;
+		InstancePool.GetKeys(Keys);
+
+		for (UClass* Key : Keys)
+		{
+			ClearClassPool(Key);
+		}
+	}
+	
+	InstancePool.Empty();
+}
+
 FString URigVMDetailsViewWrapperObject::GetWrappedNodeNotation() const
 {
 	return ClassToInfo.FindChecked(GetClass()).Notation;
@@ -758,6 +865,46 @@ void URigVMDetailsViewWrapperObject::PostEditChangeChainProperty(FPropertyChange
 	WrappedPropertyChangedChainEvent.Broadcast(this, PropertyPath, PropertyChangedEvent);
 }
 
+void URigVMDetailsViewWrapperObject::ResetForReuse()
+{
+	// Unbind from graph events if this was a node wrapper
+	if (URigVMNode* Node = Cast<URigVMNode>(GetSubject()))
+	{
+		if (URigVMGraph* Graph = Node->GetGraph())
+		{
+			Graph->OnModified().RemoveAll(this);
+		}
+	}
+	
+	// Clear the subject reference
+	SubjectPtr.Reset();
+
+	// Clear the event delegate bindings
+	WrappedPropertyChangedChainEvent.Clear();
+
+	// Reset the value setting guard
+	bIsSettingValue = false;
+
+	// Reset all UPROPERTY values to defaults from CDO
+	UClass* MyClass = GetClass();
+	UObject* CDO = MyClass->GetDefaultObject();
+	for (TFieldIterator<FProperty> It(MyClass); It; ++It)
+	{
+		FProperty* Property = *It;
+		// Skip properties from parent class (URigVMDetailsViewWrapperObject itself)
+		if (Property->GetOwnerClass() != MyClass)
+			continue;
+
+		uint8* DestPtr = Property->ContainerPtrToValuePtr<uint8>(this);
+		const uint8* SrcPtr = Property->ContainerPtrToValuePtr<uint8>(CDO);
+		Property->CopyCompleteValue(DestPtr, SrcPtr);
+	}
+	
+	// Move to transient package so we don't hold references to the old outer
+	Rename(nullptr, GetTransientPackage(), REN_DoNotDirty | REN_DontCreateRedirectors | REN_NonTransactional);
+	AddToRoot();  // Ensure pooled instances won't be GC'd
+}
+
 void URigVMDetailsViewWrapperObject::CopyPropertiesForUnrelatedStructs(uint8* InTargetMemory,
 	const UStruct* InTargetStruct, const uint8* InSourceMemory, const UStruct* InSourceStruct)
 {
diff --git a/Engine/Plugins/Runtime/RigVM/Source/RigVMEditor/Private/Editor/RigVMEditor.cpp b/Engine/Plugins/Runtime/RigVM/Source/RigVMEditor/Private/Editor/RigVMEditor.cpp
--- a/Engine/Plugins/Runtime/RigVM/Source/RigVMEditor/Private/Editor/RigVMEditor.cpp
+++ b/Engine/Plugins/Runtime/RigVM/Source/RigVMEditor/Private/Editor/RigVMEditor.cpp
@@ -3572,10 +3572,8 @@ void FRigVMEditorBase::ClearDetailsViewWrapperObjects()
 	{
 		if(WrapperObjectPtr.IsValid())
 		{
-			URigVMDetailsViewWrapperObject* WrapperObject = WrapperObjectPtr.Get();
-			WrapperObject->RemoveFromRoot();
-			WrapperObject->Rename(nullptr, GetTransientPackage(), REN_DoNotDirty | REN_DontCreateRedirectors | REN_NonTransactional);
-			WrapperObject->MarkAsGarbage();
+			// Return to pool instead of destroying
+			URigVMDetailsViewWrapperObject::ReturnToPool(WrapperObjectPtr.Get());
 		}
 	}
 	WrapperObjects.Reset();
diff --git a/Engine/Plugins/Runtime/RigVM/Source/RigVMEditor/Private/RigVMEditorModule.cpp b/Engine/Plugins/Runtime/RigVM/Source/RigVMEditor/Private/RigVMEditorModule.cpp
--- a/Engine/Plugins/Runtime/RigVM/Source/RigVMEditor/Private/RigVMEditorModule.cpp
+++ b/Engine/Plugins/Runtime/RigVM/Source/RigVMEditor/Private/RigVMEditorModule.cpp
@@ -243,6 +243,9 @@ void FRigVMEditorModule::ShutdownModuleCommon()
 		BlueprintEditorModule.UnregisterGraphCustomization(SchemaCDO);
 #endif
 	}
+	
+	// Clear all instances in the class pools of the details wrapper objects
+	URigVMDetailsViewWrapperObject::ClearAllClassPools();
 
 	// Unregister to fixup newly created BPs
 	FKismetEditorUtilities::UnregisterAutoBlueprintNodeCreation(this);
diff --git a/Engine/Plugins/Runtime/RigVM/Source/RigVMEditor/Public/Editor/RigVMDetailsViewWrapperObject.h b/Engine/Plugins/Runtime/RigVM/Source/RigVMEditor/Public/Editor/RigVMDetailsViewWrapperObject.h
--- a/Engine/Plugins/Runtime/RigVM/Source/RigVMEditor/Public/Editor/RigVMDetailsViewWrapperObject.h
+++ b/Engine/Plugins/Runtime/RigVM/Source/RigVMEditor/Public/Editor/RigVMDetailsViewWrapperObject.h
@@ -33,6 +33,17 @@ public:
 	static UE_API void MarkOutdatedClass(UClass* InClass);
 	static UE_API bool IsValidClass(UClass* InClass);
 	
+	static UE_API URigVMDetailsViewWrapperObject* GetOrCreateNewInstance(const UClass* InWrapperObjectClass, UObject* InOuter);
+	
+	// Return an instance to the pool for reuse
+	static UE_API void ReturnToPool(URigVMDetailsViewWrapperObject* InInstance);
+
+	// Clear pooled instance of a certain class
+	static UE_API void ClearClassPool(UClass* InClass);
+	
+	// Clear all pooled instances (call on module shutdown)
+	static UE_API void ClearAllClassPools();
+	
 	UE_API FString GetWrappedNodeNotation() const;
 	
 	bool IsChildOf(const UStruct* InStruct) const
@@ -76,6 +87,9 @@ public:
 	UE_API virtual void PostEditChangeChainProperty(FPropertyChangedChainEvent& PropertyChangedEvent) override;
 
 private:
+	
+	// Reset instance state for reuse
+	void ResetForReuse();
 
 	static UE_API void CopyPropertiesForUnrelatedStructs(uint8* InTargetMemory, const UStruct* InTargetStruct, const uint8* InSourceMemory, const UStruct* InSourceStruct);
 	UE_API void HandleModifiedEvent(ERigVMGraphNotifType InNotifType, URigVMGraph* InGraph, UObject* InSubject);
@@ -120,6 +134,13 @@ private:
 	static UE_API TMap<FPerClassInfo, UClass*> InfoToClass;
 	static UE_API TMap<UClass*, FPerClassInfo> ClassToInfo;
 	static UE_API TSet<UClass*> OutdatedClassToRecreate;
+	
+	// Pool of available wrapper objects per class
+	static TMap<UClass*, TArray<URigVMDetailsViewWrapperObject*>> InstancePool;
+
+	// Maximum pool size per class (prevents unbounded growth)
+	static constexpr int32 MaxPoolSizePerClass = 32;
+	
 	bool bIsSettingValue;
 
 	UPROPERTY()
