Control Rig: Fix ensure when dropping variable in graph in legacy blueprints
#rnx

#Synced-CL 51501689

diff --git a/Engine/Plugins/Runtime/RigVM/Source/RigVMEditor/Private/Editor/RigVMEditor.cpp b/Engine/Plugins/Runtime/RigVM/Source/RigVMEditor/Private/Editor/RigVMEditor.cpp
--- a/Engine/Plugins/Runtime/RigVM/Source/RigVMEditor/Private/Editor/RigVMEditor.cpp
+++ b/Engine/Plugins/Runtime/RigVM/Source/RigVMEditor/Private/Editor/RigVMEditor.cpp
@@ -3821,7 +3821,13 @@ void FRigVMEditorBase::HandleVariableDroppedFromBlueprint(UObject* InSubject, FP
 	check(Controller);
 
 	FRigVMExternalVariable ExternalVariable = Controller->GetVariableByProperty(InVariableToDrop);
-	ensure(ExternalVariable.GetProperty() == InVariableToDrop);
+	if (!ExternalVariable.IsValid(true /* allow null ptr */))
+	{
+		// In some cases the variable property does not match because it comes from the SkeletalClass, not the GeneratedBlueprintClass (in the case of legacy)
+		// Let's fall back to looking by name
+		ExternalVariable = Controller->GetVariableByName(InVariableToDrop->GetFName());
+	}
+	
 	if (!ExternalVariable.IsValid(true /* allow null ptr */))
 	{
 		return;
