Ensure RigVM asset is loaded when the host loads

#Synced-CL 51602191

diff --git a/Engine/Plugins/Runtime/RigVM/Source/RigVM/Private/RigVMHost.cpp b/Engine/Plugins/Runtime/RigVM/Source/RigVM/Private/RigVMHost.cpp
--- a/Engine/Plugins/Runtime/RigVM/Source/RigVM/Private/RigVMHost.cpp
+++ b/Engine/Plugins/Runtime/RigVM/Source/RigVM/Private/RigVMHost.cpp
@@ -177,9 +177,11 @@ void URigVMHost::PostLoad()
 	{
 		if (VM != nullptr)
 		{
+			// Ensure the VM has been fully loaded as we depend on it
+			VM->ConditionalPostLoad();
+
 			if (HasAnyFlags(RF_ClassDefaultObject))
 			{
-				VM->ConditionalPostLoad();
 				InitializeCDOVM();
 			}
 
