Control Rig: Fix ensure creating CDO while async loading
#rnx

#Synced-CL 51621469

diff --git a/Engine/Plugins/Runtime/RigVM/Source/RigVMDeveloper/Private/RigVMBlueprintLegacy.cpp b/Engine/Plugins/Runtime/RigVM/Source/RigVMDeveloper/Private/RigVMBlueprintLegacy.cpp
--- a/Engine/Plugins/Runtime/RigVM/Source/RigVMDeveloper/Private/RigVMBlueprintLegacy.cpp
+++ b/Engine/Plugins/Runtime/RigVM/Source/RigVMDeveloper/Private/RigVMBlueprintLegacy.cpp
@@ -882,7 +882,8 @@ TArray<FRigVMGraphVariableDescription> URigVMBlueprint::GetAssetVariables() cons
 		{
 			if (const FProperty* Property = GeneratedClass->FindPropertyByName(NewVariable.Name))
 			{
-				if (const URigVMHost* CDO = Cast<URigVMHost>(GeneratedClass->GetDefaultObject(true)))
+				const bool bCreateCDO = !IsAsyncLoading();
+				if (const URigVMHost* CDO = Cast<URigVMHost>(GeneratedClass->GetDefaultObject(bCreateCDO))) 
 				{
 					if (const void* PropertyAddr = Property->ContainerPtrToValuePtr<void>(CDO))
 					{
