Control Rig: Fix changing variable type creating unbreakable link
#rnx

#Synced-CL 51621469

diff --git a/Engine/Plugins/Runtime/RigVM/Source/RigVMDeveloper/Private/RigVMModel/RigVMController.cpp b/Engine/Plugins/Runtime/RigVM/Source/RigVMDeveloper/Private/RigVMModel/RigVMController.cpp
--- a/Engine/Plugins/Runtime/RigVM/Source/RigVMDeveloper/Private/RigVMModel/RigVMController.cpp
+++ b/Engine/Plugins/Runtime/RigVM/Source/RigVMDeveloper/Private/RigVMModel/RigVMController.cpp
@@ -21969,7 +21969,7 @@ bool URigVMController::ChangePinType(URigVMPin* InPin, TRigVMTypeIndex InTypeInd
 		LinkedPaths = GetLinkedPaths(InPin, true, true);
 		if(!LinkedPaths.IsEmpty())
 		{
-			FastBreakLinkedPaths(LinkedPaths);
+			BreakLinkedPaths(LinkedPaths, bSetupUndoRedo, bBreakLinks);
 
 			const FString OrphanedName = FString::Printf(TEXT("%s%s"), URigVMPin::OrphanPinPrefix, *InPin->GetName());
 			if(InPin->GetNode()->FindPin(OrphanedName) == nullptr)
