Control Rig: Fix ensure when copy pasting variable from BP rig to runtime asset rig
#rnx

#Synced-CL 51621469

diff --git a/Engine/Plugins/Runtime/RigVM/Source/RigVMDeveloper/Private/RigVMEditorAsset.cpp b/Engine/Plugins/Runtime/RigVM/Source/RigVMDeveloper/Private/RigVMEditorAsset.cpp
--- a/Engine/Plugins/Runtime/RigVM/Source/RigVMDeveloper/Private/RigVMEditorAsset.cpp
+++ b/Engine/Plugins/Runtime/RigVM/Source/RigVMDeveloper/Private/RigVMEditorAsset.cpp
@@ -4819,24 +4819,48 @@ bool URigVMEditorAsset::OnPasteVariable(const FString& InText)
 {
 	if (URigVMRuntimeAsset* RuntimeAsset = GetRuntimeAsset())
 	{
-		if (!ensure(InText.StartsWith(TEXT("RigVMVar"), ESearchCase::CaseSensitive)))
+		bool bIsRigVMVar = InText.StartsWith(TEXT("RigVMVar"), ESearchCase::CaseSensitive);
+		bool bIsBPVar = !bIsRigVMVar && InText.StartsWith(TEXT("BPVar"), ESearchCase::CaseSensitive);
+
+		if (!ensure(bIsRigVMVar || bIsBPVar))
 		{
 			return false;
 		}
 
-		FRigVMGraphVariableDescription NewVar;
-		FStringOutputDevice Errors;
-		const TCHAR* Import = InText.GetCharArray().GetData() + FCString::Strlen(TEXT("RigVMVar"));
-		FRigVMGraphVariableDescription::StaticStruct()->ImportText(Import, &NewVar, nullptr, PPF_None, &Errors, FBPVariableDescription::StaticStruct()->GetName());
-		if (Errors.IsEmpty())
+		if (bIsRigVMVar)
+		{
+			FRigVMGraphVariableDescription NewVar;
+			FStringOutputDevice Errors;
+			const TCHAR* Import = InText.GetCharArray().GetData() + FCString::Strlen(TEXT("RigVMVar"));
+			FRigVMGraphVariableDescription::StaticStruct()->ImportText(Import, &NewVar, nullptr, PPF_None, &Errors, FBPVariableDescription::StaticStruct()->GetName());
+			if (Errors.IsEmpty())
+			{
+				FScopedTransaction Transaction(FText::Format(LOCTEXT("PasteVariable", "Paste Variable: {0}"), FText::FromName(NewVar.Name)));
+
+				NewVar.Name = FRigVMBlueprintUtils::FindUniqueVariableName(this, NewVar.Name.ToString());
+				FRigVMExternalVariable ExternalVariable = RigVMVariableUtils::ExternalVariableFromRigVMVariableDescription(NewVar);
+				FName Result = RuntimeAsset->AddHostMemberVariableFromExternal(ExternalVariable, NewVar.DefaultValue);
+
+				return !Result.IsNone();
+			}
+		}
+		
+		if (bIsBPVar)
 		{
-			FScopedTransaction Transaction(FText::Format(LOCTEXT("PasteVariable", "Paste Variable: {0}"), FText::FromName(NewVar.Name)));
+			FBPVariableDescription Description;
+			FStringOutputDevice Errors;
+			const TCHAR* Import = InText.GetCharArray().GetData() + FCString::Strlen(TEXT("BPVar"));
+			FBPVariableDescription::StaticStruct()->ImportText(Import, &Description, nullptr, PPF_None, &Errors, FBPVariableDescription::StaticStruct()->GetName());
+			if (Errors.IsEmpty())
+			{
+				FScopedTransaction Transaction(FText::Format(LOCTEXT("PasteVariable", "Paste Variable: {0}"), FText::FromName(Description.VarName)));
 
-			NewVar.Name = FRigVMBlueprintUtils::FindUniqueVariableName(this, NewVar.Name.ToString());
-			FRigVMExternalVariable ExternalVariable = RigVMVariableUtils::ExternalVariableFromRigVMVariableDescription(NewVar);
-			FName Result = RuntimeAsset->AddHostMemberVariableFromExternal(ExternalVariable, NewVar.DefaultValue);
+				FRigVMExternalVariable ExternalVariable = RigVMTypeUtils::ExternalVariableFromBPVariableDescription(Description);
+				ExternalVariable.SetName(FRigVMBlueprintUtils::FindUniqueVariableName(this, ExternalVariable.GetName().ToString()));
+				FName Result = RuntimeAsset->AddHostMemberVariableFromExternal(ExternalVariable, Description.DefaultValue);
 
-			return !Result.IsNone();
+				return !Result.IsNone();
+			}
 		}
 	}
 	return false;
