Control Rig: Fix promote to variable not maintaining default values
#rnx

#Synced-CL 51710020

diff --git a/Engine/Plugins/Runtime/RigVM/Source/RigVMDeveloper/Private/EdGraph/RigVMEdGraphSchema.cpp b/Engine/Plugins/Runtime/RigVM/Source/RigVMDeveloper/Private/EdGraph/RigVMEdGraphSchema.cpp
--- a/Engine/Plugins/Runtime/RigVM/Source/RigVMDeveloper/Private/EdGraph/RigVMEdGraphSchema.cpp
+++ b/Engine/Plugins/Runtime/RigVM/Source/RigVMDeveloper/Private/EdGraph/RigVMEdGraphSchema.cpp
@@ -342,101 +342,24 @@ UEdGraphNode* FRigVMEdGraphSchemaAction_PromoteToVariable::PerformAction(UEdGrap
 		bLocalVariable ?
 		LOCTEXT("GraphEd_PromoteToLocalVariable", "Promote Pin To Local Variable") :
 		LOCTEXT("GraphEd_PromoteToVariable", "Promote Pin To Variable"));
+	
 
+	bool bSuccess = false;
 	if(bLocalVariable)
 	{
-		const FRigVMGraphVariableDescription VariableDescription = Controller->AddLocalVariable(
-			*ModelPin->GetPinPath(),
-			ModelPin->GetCPPType(),
-			ModelPin->GetCPPTypeObject(),
-			ModelPin->GetDefaultValue(),
-			true,
-			true
-		);
-
-		VariableName = VariableDescription.Name;
+		bSuccess = Controller->PromotePinToLocalVariable(ModelPin->GetPinPath(), true, FDeprecateSlateVector2D(Location), true, true);
 	}
 	else
 	{
-		Blueprint->GetObject()->Modify();
-
-		FString DefaultValue = ModelPin->GetDefaultValue();
-		if(!DefaultValue.IsEmpty())
-		{
-			if(UScriptStruct* ScriptStruct = Cast<UScriptStruct>(ModelPin->GetCPPTypeObject()))
-			{
-				if(ScriptStruct == TBaseStructure<FVector2D>::Get())
-				{
-					FVector2D Value = FVector2D::ZeroVector;
-					ScriptStruct->ImportText(*DefaultValue, &Value, nullptr, PPF_None, nullptr, FString());
-					DefaultValue = Value.ToString();
-				}
-				if(ScriptStruct == TBaseStructure<FVector>::Get())
-				{
-					FVector Value = FVector::ZeroVector;
-					ScriptStruct->ImportText(*DefaultValue, &Value, nullptr, PPF_None, nullptr, FString());
-					DefaultValue = Value.ToString();
-				}
-				if(ScriptStruct == TBaseStructure<FQuat>::Get())
-				{
-					FQuat Value = FQuat::Identity;
-					ScriptStruct->ImportText(*DefaultValue, &Value, nullptr, PPF_None, nullptr, FString());
-					DefaultValue = Value.ToString();
-				}
-				if(ScriptStruct == TBaseStructure<FRotator>::Get())
-				{
-					FRotator Value = FRotator::ZeroRotator;
-					ScriptStruct->ImportText(*DefaultValue, &Value, nullptr, PPF_None, nullptr, FString());
-					DefaultValue = Value.ToString();
-				}
-				if(ScriptStruct == TBaseStructure<FTransform>::Get())
-				{
-					FTransform Value = FTransform::Identity;
-					ScriptStruct->ImportText(*DefaultValue, &Value, nullptr, PPF_None, nullptr, FString());
-					DefaultValue = Value.ToString();
-				}
-			}
-		}
-
-		FRigVMExternalVariable ExternalVariable = FRigVMExternalVariable::Make(
-			FGuid::NewGuid(),
-			FromPin->GetFName(),
-			ModelPin->GetCPPType(),
-			ModelPin->GetCPPTypeObject(),
-			false,
-			false);
-		
-		VariableName = Blueprint->AddHostMemberVariableFromExternal(
-			ExternalVariable,
-			DefaultValue
-		);
+		bSuccess = Controller->PromotePinToVariable(ModelPin->GetPinPath(), true, FDeprecateSlateVector2D(Location), true, true);
 	}
 
-	if(!VariableName.IsNone())
+	if(bSuccess)
 	{
-		const URigVMNode* ModelNode = Controller->AddVariableNode(
-			VariableName,
-			ModelPin->GetCPPType(),
-			ModelPin->GetCPPTypeObject(),
-			FromPin->Direction == EGPD_Input,
-			ModelPin->GetDefaultValue(),
-			FDeprecateSlateVector2D(Location),
-			FString(),
-			true,
-			true
-		);
-
-		if(ModelNode)
+		TArray<URigVMPin*> SourcePins = ModelPin->GetLinkedSourcePins();
+		if (!SourcePins.IsEmpty())
 		{
-			if(FromPin->Direction == EGPD_Input)
-			{
-				Controller->AddLink(ModelNode->FindPin(TEXT("Value")), ModelPin, true);
-			}
-			else
-			{
-				Controller->AddLink(ModelPin, ModelNode->FindPin(TEXT("Value")), true);
-			}
-			return RigGraph->FindNodeForModelNodeName(ModelNode->GetFName());
+			return RigGraph->FindNodeForModelNodeName(SourcePins[0]->GetNode()->GetFName());
 		}
 	}
 	
diff --git a/Engine/Plugins/Runtime/RigVM/Source/RigVMDeveloper/Private/RigVMModel/RigVMController.cpp b/Engine/Plugins/Runtime/RigVM/Source/RigVMDeveloper/Private/RigVMModel/RigVMController.cpp
--- a/Engine/Plugins/Runtime/RigVM/Source/RigVMDeveloper/Private/RigVMModel/RigVMController.cpp
+++ b/Engine/Plugins/Runtime/RigVM/Source/RigVMDeveloper/Private/RigVMModel/RigVMController.cpp
@@ -12696,7 +12696,7 @@ bool URigVMController::PromotePinToVariable(const FString& InPinPath, bool bCrea
 		return false;
 	}
 
-	const bool bSuccess = PromotePinToVariable(Pin, bCreateVariableNode, InNodePosition, bSetupUndoRedo);
+	const bool bSuccess = PromotePinToVariable(Pin, false, bCreateVariableNode, InNodePosition, bSetupUndoRedo);
 	if (bSuccess && bPrintPythonCommand)
 	{
 		const FString GraphName = GetSchema()->GetSanitizedGraphName(GetGraph()->GetGraphName());
@@ -12712,7 +12712,45 @@ bool URigVMController::PromotePinToVariable(const FString& InPinPath, bool bCrea
 	return bSuccess;
 }
 
-bool URigVMController::PromotePinToVariable(URigVMPin* InPin, bool bCreateVariableNode, const FVector2D& InNodePosition, bool bSetupUndoRedo)
+bool URigVMController::PromotePinToLocalVariable(const FString& InPinPath, bool bCreateVariableNode, const FVector2D& InNodePosition, bool bSetupUndoRedo, bool bPrintPythonCommand)
+{
+	if (!IsValidGraph())
+	{
+		return false;
+	}
+
+	if (!bIsTransacting && !IsGraphEditable())
+	{
+		return false;
+	}
+
+	URigVMGraph* Graph = GetGraph();
+	check(Graph);
+
+	URigVMPin* Pin = Graph->FindPin(InPinPath);
+	if (Pin == nullptr)
+	{
+		ReportErrorf(TEXT("Cannot find pin '%s'."), *InPinPath);
+		return false;
+	}
+
+	const bool bSuccess = PromotePinToVariable(Pin, true, bCreateVariableNode, InNodePosition, bSetupUndoRedo);
+	if (bSuccess && bPrintPythonCommand)
+	{
+		const FString GraphName = GetSchema()->GetSanitizedGraphName(GetGraph()->GetGraphName());
+		
+		RigVMPythonUtils::Print(GetSchema()->GetGraphOuterName(GetGraph()), 
+			FString::Printf(TEXT("asset.get_controller_by_name('%s').promote_pin_to_local_variable('%s', %s, %s)"),
+			*GraphName,
+			*GetSchema()->GetSanitizedPinPath(InPinPath),
+			(bCreateVariableNode) ? TEXT("True") : TEXT("False"),
+			*RigVMPythonUtils::Vector2DToPythonString(InNodePosition)));
+	}
+	
+	return bSuccess;
+}
+
+bool URigVMController::PromotePinToVariable(URigVMPin* InPin, bool bCreateLocalVariable, bool bCreateVariableNode, const FVector2D& InNodePosition, bool bSetupUndoRedo)
 {
 	if (!bIsTransacting && !IsGraphEditable())
 	{
@@ -12750,16 +12788,25 @@ bool URigVMController::PromotePinToVariable(URigVMPin* InPin, bool bCreateVariab
 	}
 	else
 	{
-		if (!UnitNodeCreatedContext.GetCreateExternalVariableDelegate().IsBound())
+		FName VariableName;
+		VariableForPin = InPin->ToExternalVariable();
+		if (bCreateLocalVariable)
 		{
-			return false;
+			const FRigVMGraphVariableDescription VariableDescription = AddLocalVariable(VariableForPin.GetName(),VariableForPin.GetExtendedCPPType().ToString(), VariableForPin.GetCPPTypeObject(), InPin->GetDefaultValue());
+			VariableName = VariableDescription.Name;
 		}
-
-		VariableForPin = InPin->ToExternalVariable();
-		FName VariableName = UnitNodeCreatedContext.GetCreateExternalVariableDelegate().Execute(VariableForPin, InPin->GetDefaultValue());
-		if (VariableName.IsNone())
+		else
 		{
-			return false;
+			if (!UnitNodeCreatedContext.GetCreateExternalVariableDelegate().IsBound())
+			{
+				return false;
+			}
+
+			VariableName = UnitNodeCreatedContext.GetCreateExternalVariableDelegate().Execute(VariableForPin, InPin->GetDefaultValue());
+			if (VariableName.IsNone())
+			{
+				return false;
+			}
 		}
 
 		VariableForPin = GetVariableByName(VariableName);
diff --git a/Engine/Plugins/Runtime/RigVM/Source/RigVMDeveloper/Public/RigVMModel/RigVMController.h b/Engine/Plugins/Runtime/RigVM/Source/RigVMDeveloper/Public/RigVMModel/RigVMController.h
--- a/Engine/Plugins/Runtime/RigVM/Source/RigVMDeveloper/Public/RigVMModel/RigVMController.h
+++ b/Engine/Plugins/Runtime/RigVM/Source/RigVMDeveloper/Public/RigVMModel/RigVMController.h
@@ -1031,6 +1031,10 @@ public:
 	// Promotes a pin to a variable
 	UFUNCTION(BlueprintCallable, Category = RigVMController)
 	UE_API bool PromotePinToVariable(const FString& InPinPath, bool bCreateVariableNode, const FVector2D& InNodePosition = FVector2D::ZeroVector, bool bSetupUndoRedo = true, bool bPrintPythonCommand = false);
+	
+	// Promotes a pin to a local variable
+	UFUNCTION(BlueprintCallable, Category = RigVMController)
+	UE_API bool PromotePinToLocalVariable(const FString& InPinPath, bool bCreateVariableNode, const FVector2D& InNodePosition = FVector2D::ZeroVector, bool bSetupUndoRedo = true, bool bPrintPythonCommand = false);
 
 	// Adds a link to the graph.
 	// This causes a LinkAdded modified event.
@@ -1361,7 +1365,7 @@ private:
 	UE_API bool BindPinToVariable(URigVMPin* InPin, const FString& InNewBoundVariablePath, bool bSetupUndoRedo, const FString& InVariableNodeName = FString());
 	UE_API bool UnbindPinFromVariable(URigVMPin* InPin, bool bSetupUndoRedo);
 	UE_API bool MakeBindingsFromVariableNode(URigVMVariableNode* InNode, bool bSetupUndoRedo);
-	UE_API bool PromotePinToVariable(URigVMPin* InPin, bool bCreateVariableNode, const FVector2D& InNodePosition, bool bSetupUndoRedo);
+	UE_API bool PromotePinToVariable(URigVMPin* InPin, bool bCreateLocalVariable, bool bCreateVariableNode, const FVector2D& InNodePosition, bool bSetupUndoRedo);
 	UE_API URigVMInjectionInfo* InjectNodeIntoPin(const FString& InPinPath, bool bAsInput, const FName& InInputPinName, const FName& InOutputPinName, bool bSetupUndoRedo = true);
 	UE_API URigVMInjectionInfo* InjectNodeIntoPin(URigVMPin* InPin, bool bAsInput, const FName& InInputPinName, const FName& InOutputPinName, bool bSetupUndoRedo = true);
 	UE_API URigVMNode* EjectNodeFromPin(URigVMPin* InPin, bool bSetupUndoRedo = true, bool bPrintPythonCommands = false);
