Control Rig: Implement reordering of variable in BPI assets
#rnx

#Synced-CL 51710142

diff --git a/Engine/Plugins/Runtime/RigVM/Source/RigVM/Private/RigVMCore/RigVMVariableDescription.cpp b/Engine/Plugins/Runtime/RigVM/Source/RigVM/Private/RigVMCore/RigVMVariableDescription.cpp
--- a/Engine/Plugins/Runtime/RigVM/Source/RigVM/Private/RigVMCore/RigVMVariableDescription.cpp
+++ b/Engine/Plugins/Runtime/RigVM/Source/RigVM/Private/RigVMCore/RigVMVariableDescription.cpp
@@ -119,6 +119,10 @@ FRigVMGraphVariableDescription RigVMVariableUtils::VariableDescriptionFromProper
 	{
 		Variable.bPrivate = InPropertyDesc.GetMetaData(FBlueprintMetadata::MD_Private) == TEXT("true");
 	}
+	if(InPropertyDesc.HasMetaData(FBlueprintMetadata::MD_FunctionCategory))
+	{
+		Variable.Category = FText::FromString(InPropertyDesc.GetMetaData(FBlueprintMetadata::MD_FunctionCategory));
+	}
 #endif
 	if (InPropertyDesc.PropertyFlags & CPF_Interp)
 	{
@@ -147,7 +151,7 @@ FRigVMPropertyDescription RigVMVariableUtils::PropertyDescriptionFromVariableDes
 #if WITH_EDITOR
 	TMap<FName, FString> MetaData;
 	MetaData.Add(FBlueprintMetadata::MD_Tooltip, InVariable.Tooltip.ToString());
-	MetaData.Add(TEXT("Category"), InVariable.Category.ToString());
+	MetaData.Add(FBlueprintMetadata::MD_FunctionCategory, InVariable.Category.ToString());
 	if (InVariable.bExposedOnSpawn)
 	{
 		MetaData.Add(FBlueprintMetadata::MD_ExposeOnSpawn, TEXT("true"));
diff --git a/Engine/Plugins/Runtime/RigVM/Source/RigVM/Private/RigVMRuntimeAsset.cpp b/Engine/Plugins/Runtime/RigVM/Source/RigVM/Private/RigVMRuntimeAsset.cpp
--- a/Engine/Plugins/Runtime/RigVM/Source/RigVM/Private/RigVMRuntimeAsset.cpp
+++ b/Engine/Plugins/Runtime/RigVM/Source/RigVM/Private/RigVMRuntimeAsset.cpp
@@ -193,7 +193,7 @@ FRigVMGraphVariableDescription URigVMRuntimeAsset::FindAssetVariable(const FName
 #endif
 		{
 			FRigVMGraphVariableDescription Result = RigVMVariableUtils::VariableDescriptionFromPropertyDesc(Desc);
-			Result.DefaultValue = Variables.GetDataAsStringByName(Desc.Name);
+			Result.DefaultValue = Variables.GetDataAsStringByNameSafe(Desc.Name);
 			return Result;
 		}
 	}
@@ -201,6 +201,58 @@ FRigVMGraphVariableDescription URigVMRuntimeAsset::FindAssetVariable(const FName
 	return FRigVMGraphVariableDescription();
 }
 
+int32 URigVMRuntimeAsset::GetVariableIndex(const FName& InName) const
+{
+	const UPropertyBag* PropertyBag = Variables.GetPropertyBagStruct();
+	TConstArrayView<FPropertyBagPropertyDesc> PropertyDescs = PropertyBag->GetPropertyDescs();
+
+	for (int32 PropertyIndex = 0; PropertyIndex < PropertyDescs.Num(); PropertyIndex++)
+	{
+		const FPropertyBagPropertyDesc& Desc = PropertyDescs[PropertyIndex];
+#if WITH_EDITOR
+		if (Desc.GetMetaData(FRigVMPropertyDescription::MD_DisplayName).Equals(InName.ToString()))
+#else
+		if (Desc.Name.IsEqual(InName))
+#endif
+		{
+			return PropertyIndex;
+		}
+	}
+
+	return INDEX_NONE;
+}
+
+int32 URigVMRuntimeAsset::GetVariableIndex(const FGuid& InGuid) const
+{
+	const UPropertyBag* PropertyBag = Variables.GetPropertyBagStruct();
+
+	// Find the variable by GUID
+	const FPropertyBagPropertyDesc* FoundDesc = PropertyBag->FindPropertyDescByID(InGuid);
+	if (!FoundDesc)
+	{
+		return INDEX_NONE;
+	}
+
+	// Return the cached index if valid
+	const int32 CachedIndex = FoundDesc->GetCachedIndex();
+	if (CachedIndex != INDEX_NONE)
+	{
+		return CachedIndex;
+	}
+
+	// Fallback: find the index by iterating through all descriptors
+	TConstArrayView<FPropertyBagPropertyDesc> PropertyDescs = PropertyBag->GetPropertyDescs();
+	for (int32 PropertyIndex = 0; PropertyIndex < PropertyDescs.Num(); PropertyIndex++)
+	{
+		if (&PropertyDescs[PropertyIndex] == FoundDesc)
+		{
+			return PropertyIndex;
+		}
+	}
+
+	return INDEX_NONE;
+}
+
 const TArray<UAssetUserData*>* URigVMRuntimeAsset::GetAssetUserDataArray() const
 {
 #if WITH_EDITOR
@@ -263,6 +315,9 @@ bool URigVMRuntimeAsset::RenameMemberVariable(const FName& InOldName, const FNam
 	FRigVMGraphVariableDescription Variable = FindAssetVariable(InOldName);
 	if (!Variable.Name.IsNone())
 	{
+		// Store original index before removal
+		const int32 OriginalIndex = GetVariableIndex(InOldName);
+
 		FRigVMAssetVariablesChangeScope VariablesChangeScope(this, Variable.Name);
 		Modify();
 		
@@ -272,6 +327,13 @@ bool URigVMRuntimeAsset::RenameMemberVariable(const FName& InOldName, const FNam
 			Variable.Name = InNewName;
 			FRigVMPropertyDescription NewDesc = RigVMVariableUtils::PropertyDescriptionFromVariableDescription(Variable, true);
 			Variables.AddProperties({NewDesc});
+
+			// Restore original index
+			if (OriginalIndex != INDEX_NONE)
+			{
+				SetVariableIndex(Variable.Name, OriginalIndex);
+			}
+
 			return true;
 		}
 	}
@@ -284,6 +346,9 @@ bool URigVMRuntimeAsset::ChangeMemberVariableType(const FName& InName, const FSt
 	FRigVMGraphVariableDescription Variable = FindAssetVariable(InName);
 	if (!Variable.Name.IsNone())
 	{
+		// Store original index before removal
+		const int32 OriginalIndex = GetVariableIndex(InName);
+
 		FRigVMAssetVariablesChangeScope VariablesChangeScope(this, Variable.Name);
 		Modify();
 		RemoveMemberVariable(InName);
@@ -299,11 +364,89 @@ bool URigVMRuntimeAsset::ChangeMemberVariableType(const FName& InName, const FSt
 
 		Variables.AddProperties({Description});
 
+		// Restore original index
+		if (OriginalIndex != INDEX_NONE)
+		{
+			SetVariableIndex(Variable.Name, OriginalIndex);
+		}
+
 		return true;
 	}
 	return false;
 }
 
+bool URigVMRuntimeAsset::SetVariableIndex(const FName& InName, int32 NewIndex)
+{
+	// Find the variable by name to get its GUID
+	FRigVMGraphVariableDescription Variable = FindAssetVariable(InName);
+	if (!Variable.Name.IsNone())
+	{
+		return SetVariableIndex(Variable.Guid, NewIndex);
+	}
+	return false;
+}
+
+bool URigVMRuntimeAsset::SetVariableIndex(const FGuid& InVariableGuid, int32 NewIndex)
+{
+	// Find the source property descriptor by GUID
+	const UPropertyBag* PropertyBag = Variables.GetPropertyBagStruct();
+	const FPropertyBagPropertyDesc* SourceDesc = PropertyBag->FindPropertyDescByID(InVariableGuid);
+	if (!SourceDesc)
+	{
+		return false;
+	}
+
+	// Get current index for comparison
+	const int32 OldIndex = SourceDesc->GetCachedIndex();
+	if (OldIndex == INDEX_NONE)
+	{
+		return false; // Invalid cached index
+	}
+
+	// No change needed
+	if (OldIndex == NewIndex)
+	{
+		return true;
+	}
+
+	// Get target property descriptor at NewIndex (no array copy needed!)
+	const FPropertyBagPropertyDesc* TargetDesc = PropertyBag->FindPropertyDescByIndex(NewIndex);
+	if (!TargetDesc)
+	{
+		return false; // Invalid index
+	}
+
+	const FName SourceName = SourceDesc->Name;
+	const FName TargetName = TargetDesc->Name;
+
+#if WITH_EDITOR
+	// Get the display name for the change scope
+	const FName VariableDisplayName = *SourceDesc->GetMetaData(FRigVMPropertyDescription::MD_DisplayName);
+#else
+	const FName VariableDisplayName = SourceName;
+#endif
+
+	// Modify for undo/redo
+	Modify();
+
+	// Scope for variables change notification
+	FRigVMAssetVariablesChangeScope VariablesChangeScope(this, VariableDisplayName);
+
+	// Use the property bag's built-in reorder functionality with property names
+	// When moving down (OldIndex < NewIndex), insert after the target
+	// When moving up (OldIndex > NewIndex), insert before the target
+	const bool bInsertBefore = (OldIndex > NewIndex);
+	const EPropertyBagAlterationResult Result = Variables.ReorderProperty(SourceName, TargetName, bInsertBefore);
+	Variables.Refresh();
+
+	if (Result != EPropertyBagAlterationResult::Success)
+	{
+		return false;
+	}
+
+	return true;
+}
+
 FText URigVMRuntimeAsset::GetVariableTooltip(const FName& InName) const
 {
 	return FText::FromString(GetVariableMetadataValue(InName, FBlueprintMetadata::MD_Tooltip));
@@ -338,6 +481,9 @@ bool URigVMRuntimeAsset::SetVariableMetadataValue(const FName& InName, const FNa
 	FRigVMGraphVariableDescription Variable = FindAssetVariable(InName);
 	if (!Variable.Name.IsNone())
 	{
+		// Store original index before removal
+		const int32 OriginalIndex = GetVariableIndex(InName);
+
 		FRigVMAssetVariablesChangeScope VariablesChangeScope(this, Variable.Name);
 		RemoveMemberVariable(InName);
 
@@ -345,6 +491,13 @@ bool URigVMRuntimeAsset::SetVariableMetadataValue(const FName& InName, const FNa
 		Description.MetaData.FindOrAdd(InKey) = InValue;
 
 		Variables.AddProperties({Description});
+
+		// Restore original index
+		if (OriginalIndex != INDEX_NONE)
+		{
+			SetVariableIndex(Variable.Name, OriginalIndex);
+		}
+
 		return true;
 	}
 	return false;
@@ -355,6 +508,9 @@ bool URigVMRuntimeAsset::RemoveVariableMetadataValue(const FName& InName, const
 	FRigVMGraphVariableDescription Variable = FindAssetVariable(InName);
 	if (!Variable.Name.IsNone())
 	{
+		// Store original index before removal
+		const int32 OriginalIndex = GetVariableIndex(InName);
+
 		FRigVMAssetVariablesChangeScope VariablesChangeScope(this, Variable.Name);
 		RemoveMemberVariable(InName);
 
@@ -362,6 +518,13 @@ bool URigVMRuntimeAsset::RemoveVariableMetadataValue(const FName& InName, const
 		Description.MetaData.Remove(InKey);
 
 		Variables.AddProperties({Description});
+
+		// Restore original index
+		if (OriginalIndex != INDEX_NONE)
+		{
+			SetVariableIndex(Variable.Name, OriginalIndex);
+		}
+
 		return true;
 	}
 	return false;
@@ -372,13 +535,22 @@ bool URigVMRuntimeAsset::SetVariableExposeOnSpawn(const FName& InName, const boo
 	FRigVMGraphVariableDescription Variable = FindAssetVariable(InName);
 	if (!Variable.Name.IsNone())
 	{
+		// Store original index before removal
+		const int32 OriginalIndex = GetVariableIndex(InName);
+
 		FRigVMAssetVariablesChangeScope VariablesChangeScope(this, Variable.Name);
 		RemoveMemberVariable(InName);
 		Variable.bExposedOnSpawn = bInExposeOnSpawn;
 
 		FRigVMPropertyDescription Description = RigVMVariableUtils::PropertyDescriptionFromVariableDescription(Variable, true);
 		Variables.AddProperties({Description});
-		
+
+		// Restore original index
+		if (OriginalIndex != INDEX_NONE)
+		{
+			SetVariableIndex(Variable.Name, OriginalIndex);
+		}
+
 		return true;
 	}
 	return false;
@@ -389,13 +561,22 @@ bool URigVMRuntimeAsset::SetVariableExposeToCinematics(const FName& InName, cons
 	FRigVMGraphVariableDescription Variable = FindAssetVariable(InName);
 	if (!Variable.Name.IsNone())
 	{
+		// Store original index before removal
+		const int32 OriginalIndex = GetVariableIndex(InName);
+
 		FRigVMAssetVariablesChangeScope VariablesChangeScope(this, Variable.Name);
 		RemoveMemberVariable(InName);
 		Variable.bExposeToCinematics = bInExposeToCinematics;
 
 		FRigVMPropertyDescription Description = RigVMVariableUtils::PropertyDescriptionFromVariableDescription(Variable, true);
 		Variables.AddProperties({Description});
-		
+
+		// Restore original index
+		if (OriginalIndex != INDEX_NONE)
+		{
+			SetVariableIndex(Variable.Name, OriginalIndex);
+		}
+
 		return true;
 	}
 	return false;
@@ -406,13 +587,22 @@ bool URigVMRuntimeAsset::SetVariablePrivate(const FName& InName, const bool bInP
 	FRigVMGraphVariableDescription Variable = FindAssetVariable(InName);
 	if (!Variable.Name.IsNone())
 	{
+		// Store original index before removal
+		const int32 OriginalIndex = GetVariableIndex(InName);
+
 		FRigVMAssetVariablesChangeScope VariablesChangeScope(this, Variable.Name);
 		RemoveMemberVariable(InName);
 		Variable.bPrivate = bInPrivate;
 
 		FRigVMPropertyDescription Description = RigVMVariableUtils::PropertyDescriptionFromVariableDescription(Variable, true);
 		Variables.AddProperties({Description});
-		
+
+		// Restore original index
+		if (OriginalIndex != INDEX_NONE)
+		{
+			SetVariableIndex(Variable.Name, OriginalIndex);
+		}
+
 		return true;
 	}
 	return false;
@@ -423,13 +613,22 @@ bool URigVMRuntimeAsset::SetVariablePublic(const FName& InName, const bool bIsPu
 	FRigVMGraphVariableDescription Variable = FindAssetVariable(InName);
 	if (!Variable.Name.IsNone())
 	{
+		// Store original index before removal
+		const int32 OriginalIndex = GetVariableIndex(InName);
+
 		FRigVMAssetVariablesChangeScope VariablesChangeScope(this, Variable.Name);
 		RemoveMemberVariable(InName);
 		Variable.bPublic = bIsPublic;
 
 		FRigVMPropertyDescription Description = RigVMVariableUtils::PropertyDescriptionFromVariableDescription(Variable, true);
 		Variables.AddProperties({Description});
-		
+
+		// Restore original index
+		if (OriginalIndex != INDEX_NONE)
+		{
+			SetVariableIndex(Variable.Name, OriginalIndex);
+		}
+
 		return true;
 	}
 	return false;
@@ -464,6 +663,7 @@ FName URigVMRuntimeAsset::AddHostMemberVariableFromExternal(FRigVMExternalVariab
 	RigVMTypeUtils::CPPTypeFromExternalVariable(InVariableToCreate, CPPType, &CPPTypeObject);
 
 	FRigVMPropertyDescription Description(*VariableName, CPPType, CPPTypeObject, InDefaultValue, true);
+	Description.Guid = InVariableToCreate.GetGuid();
 	Description.PropertyFlags = CPF_Edit | CPF_BlueprintVisible | CPF_DisableEditOnInstance | CPF_Interp;
 	if (InVariableToCreate.IsPublic())
 	{
diff --git a/Engine/Plugins/Runtime/RigVM/Source/RigVM/Public/RigVMCore/RigVMVariableDescription.h b/Engine/Plugins/Runtime/RigVM/Source/RigVM/Public/RigVMCore/RigVMVariableDescription.h
--- a/Engine/Plugins/Runtime/RigVM/Source/RigVM/Public/RigVMCore/RigVMVariableDescription.h
+++ b/Engine/Plugins/Runtime/RigVM/Source/RigVM/Public/RigVMCore/RigVMVariableDescription.h
@@ -74,6 +74,7 @@ public:
 	UPROPERTY(EditAnywhere, BlueprintReadOnly, Category = RigVMGraphVariableDescription)
 	bool bPublic = false;
 
+	// This property should soon become deprecated, since it does not make sense for blueprint independent assets
 	// Is this variable private
 	UPROPERTY(EditAnywhere, BlueprintReadOnly, Category = RigVMGraphVariableDescription)
 	bool bPrivate = true;
diff --git a/Engine/Plugins/Runtime/RigVM/Source/RigVM/Public/RigVMRuntimeAsset.h b/Engine/Plugins/Runtime/RigVM/Source/RigVM/Public/RigVMRuntimeAsset.h
--- a/Engine/Plugins/Runtime/RigVM/Source/RigVM/Public/RigVMRuntimeAsset.h
+++ b/Engine/Plugins/Runtime/RigVM/Source/RigVM/Public/RigVMRuntimeAsset.h
@@ -124,6 +124,8 @@ protected:
 	UFUNCTION(BlueprintCallable, Category = "VM")
 	UE_API virtual TArray<FRigVMGraphVariableDescription> GetAssetVariables() const;
 	UE_API virtual FRigVMGraphVariableDescription FindAssetVariable(const FName& InName) const;
+	UE_API int32 GetVariableIndex(const FName& InName) const;
+	UE_API int32 GetVariableIndex(const FGuid& InGuid) const;
 #if WITH_EDITOR
 	FOnRigVMPreVariablesChanged OnPreVariablesChangedDelegate;
 	FOnRigVMPostVariablesChanged OnPostVariablesChangedDelegate;
@@ -136,6 +138,8 @@ protected:
 	UE_API virtual bool BulkRemoveMemberVariables(const TArray<FName>& InNames);
 	UE_API virtual bool RenameMemberVariable(const FName& InOldName, const FName& InNewName);
 	UE_API virtual bool ChangeMemberVariableType(const FName& InName, const FString& InCPPType, bool bIsPublic, bool bIsReadOnly, FString InDefaultValue);
+	UE_API bool SetVariableIndex(const FName& InName, int32 NewIndex);
+	UE_API virtual bool SetVariableIndex(const FGuid& InVariableGuid, int32 NewIndex);
 	UE_API virtual FText GetVariableTooltip(const FName& InName) const;
 	UE_API virtual bool SetVariableTooltip(const FName& InName, const FText& InTooltip);
 	UE_API virtual FString GetVariableCategory(const FName& InName) const;
diff --git a/Engine/Plugins/Runtime/RigVM/Source/RigVMDeveloper/Private/RigVMBlueprintLegacy.cpp b/Engine/Plugins/Runtime/RigVM/Source/RigVMDeveloper/Private/RigVMBlueprintLegacy.cpp
--- a/Engine/Plugins/Runtime/RigVM/Source/RigVMDeveloper/Private/RigVMBlueprintLegacy.cpp
+++ b/Engine/Plugins/Runtime/RigVM/Source/RigVMDeveloper/Private/RigVMBlueprintLegacy.cpp
@@ -1238,6 +1238,67 @@ bool URigVMBlueprint::ChangeMemberVariableType(const FName& InName, const FEdGra
 	return true;
 }
 
+bool URigVMBlueprint::SetVariableIndex(const FName& InName, int32 NewIndex)
+{
+	// Find the variable index by name
+	const FGuid VariableGuid = FBlueprintEditorUtils::FindMemberVariableGuidByName(this, InName);
+	if (!VariableGuid.IsValid())
+	{
+		return false;
+	}
+
+	return SetVariableIndex(VariableGuid, NewIndex);
+}
+
+bool URigVMBlueprint::SetVariableIndex(const FGuid& InVariableGuid, int32 NewIndex)
+{
+	// Find the variable by GUID
+	int32 OldIndex = INDEX_NONE;
+	for (int32 Index = 0; Index < NewVariables.Num(); Index++)
+	{
+		if (NewVariables[Index].VarGuid == InVariableGuid)
+		{
+			OldIndex = Index;
+			break;
+		}
+	}
+
+	// Variable not found
+	if (OldIndex == INDEX_NONE)
+	{
+		return false;
+	}
+
+	// Validate new index
+	if (NewIndex < 0 || NewIndex >= NewVariables.Num())
+	{
+		return false;
+	}
+
+	// No change needed
+	if (OldIndex == NewIndex)
+	{
+		return true;
+	}
+
+	// Modify for undo/redo
+	Modify();
+
+	// Reorder the array
+	FBPVariableDescription VariableToMove = NewVariables[OldIndex];
+	NewVariables.RemoveAt(OldIndex);
+	NewVariables.Insert(VariableToMove, NewIndex);
+
+	// Broadcast property change events
+	FPropertyChangedEvent PropertyChangedEvent(FindFProperty<FProperty>(GetClass(), GET_MEMBER_NAME_CHECKED(UBlueprint, NewVariables)), EPropertyChangeType::ArrayMove);
+	FCoreUObjectDelegates::OnObjectPropertyChanged.Broadcast(this, PropertyChangedEvent);
+
+	// Mark as modifice
+	FBlueprintEditorUtils::MarkBlueprintAsModified(this);
+
+	return true;
+}
+
 FName URigVMBlueprint::FindHostMemberVariableUniqueName(TSharedPtr<FKismetNameValidator> InNameValidator, const FString& InBaseName)
 {
 	FString BaseName = InBaseName;
diff --git a/Engine/Plugins/Runtime/RigVM/Source/RigVMDeveloper/Private/RigVMEditorAsset.cpp b/Engine/Plugins/Runtime/RigVM/Source/RigVMDeveloper/Private/RigVMEditorAsset.cpp
--- a/Engine/Plugins/Runtime/RigVM/Source/RigVMDeveloper/Private/RigVMEditorAsset.cpp
+++ b/Engine/Plugins/Runtime/RigVM/Source/RigVMDeveloper/Private/RigVMEditorAsset.cpp
@@ -4718,6 +4718,24 @@ bool URigVMEditorAsset::ChangeMemberVariableType(const FName& InName, const FStr
 	return false;
 }
 
+bool URigVMEditorAsset::SetVariableIndex(const FName& InName, int32 NewIndex)
+{
+	if (URigVMRuntimeAsset* RuntimeAsset = GetRuntimeAsset())
+	{
+		return RuntimeAsset->SetVariableIndex(InName, NewIndex);
+	}
+	return false;
+}
+
+bool URigVMEditorAsset::SetVariableIndex(const FGuid& InVariableGuid, int32 NewIndex)
+{
+	if (URigVMRuntimeAsset* RuntimeAsset = GetRuntimeAsset())
+	{
+		return RuntimeAsset->SetVariableIndex(InVariableGuid, NewIndex);
+	}
+	return false;
+}
+
 FText URigVMEditorAsset::GetVariableTooltip(const FName& InName) const
 {
 	if (URigVMRuntimeAsset* RuntimeAsset = GetRuntimeAsset())
diff --git a/Engine/Plugins/Runtime/RigVM/Source/RigVMDeveloper/Private/RigVMModel/RigVMController.cpp b/Engine/Plugins/Runtime/RigVM/Source/RigVMDeveloper/Private/RigVMModel/RigVMController.cpp
--- a/Engine/Plugins/Runtime/RigVM/Source/RigVMDeveloper/Private/RigVMModel/RigVMController.cpp
+++ b/Engine/Plugins/Runtime/RigVM/Source/RigVMDeveloper/Private/RigVMModel/RigVMController.cpp
@@ -16373,6 +16373,126 @@ bool URigVMController::SetLocalVariableDefaultValue(const FName& InVariableName,
 	return true;
 }
 
+bool URigVMController::SetLocalVariableIndex(const FName& InVariableName, int32 InNewIndex, bool bSetupUndoRedo, bool bPrintPythonCommand)
+{
+	if (!IsValidGraph())
+	{
+		return false;
+	}
+
+	if (!bIsTransacting && !IsGraphEditable())
+	{
+		return false;
+	}
+
+	URigVMGraph* Graph = GetGraph();
+	check(Graph);
+
+	// Look up variable by name
+	TArray<FRigVMGraphVariableDescription>& LocalVariables = Graph->LocalVariables;
+	const FRigVMGraphVariableDescription* Variable = LocalVariables.FindByPredicate(
+		[InVariableName](const FRigVMGraphVariableDescription& Var) {
+			return Var.Name == InVariableName;
+		});
+
+	if (!Variable)
+	{
+		ReportErrorf(TEXT("Cannot find local variable '%s'."), *InVariableName.ToString());
+		return false;
+	}
+
+	// Call GUID-based version
+	return SetLocalVariableIndex(Variable->Guid, InNewIndex, bSetupUndoRedo, bPrintPythonCommand);
+}
+
+bool URigVMController::SetLocalVariableIndex(const FGuid& InVariableGuid, int32 InNewIndex, bool bSetupUndoRedo, bool bPrintPythonCommand)
+{
+	if (!IsValidGraph())
+	{
+		return false;
+	}
+
+	if (!bIsTransacting && !IsGraphEditable())
+	{
+		return false;
+	}
+
+	URigVMGraph* Graph = GetGraph();
+	check(Graph);
+
+	TArray<FRigVMGraphVariableDescription>& LocalVariables = Graph->LocalVariables;
+
+	// Validate new index is within bounds
+	if (InNewIndex < 0 || InNewIndex >= LocalVariables.Num())
+	{
+		ReportErrorf(TEXT("Invalid new variable index '%d'. Valid range is 0 to %d."), InNewIndex, LocalVariables.Num() - 1);
+		return false;
+	}
+
+	// Find variable by GUID and get current index
+	int32 OldIndex = INDEX_NONE;
+	for (int32 i = 0; i < LocalVariables.Num(); ++i)
+	{
+		if (LocalVariables[i].Guid == InVariableGuid)
+		{
+			OldIndex = i;
+			break;
+		}
+	}
+
+	if (OldIndex == INDEX_NONE)
+	{
+		ReportError(TEXT("Cannot find local variable by GUID."));
+		return false;
+	}
+
+	// Nothing to do if already at target index
+	if (OldIndex == InNewIndex)
+	{
+		return true;
+	}
+
+	FRigVMControllerCompileBracketScope CompileScope(this);
+
+	// Setup undo/redo action
+	if (bSetupUndoRedo)
+	{
+		FRigVMSetLocalVariableIndexAction Action(this, InVariableGuid, OldIndex, InNewIndex);
+		Action.SetTitle(FString::Printf(TEXT("Reorder Local Variable '%s'"), *LocalVariables[OldIndex].Name.ToString()));
+		GetActionStack()->AddAction(Action);
+	}
+
+	// Perform the reorder: remove from old position, insert at new position
+	FRigVMGraphVariableDescription Variable = LocalVariables[OldIndex];
+	LocalVariables.RemoveAt(OldIndex);
+	LocalVariables.Insert(Variable, InNewIndex);
+
+	// Notify listeners
+	if (!bSuspendNotifications)
+	{
+		(void)MarkPackageDirty(Graph);
+
+		if (URigVMLibraryNode* LibraryNode = Cast<URigVMLibraryNode>(Graph->GetOuter()))
+		{
+			Notify(ERigVMGraphNotifType::LocalVariableIndexChanged, LibraryNode);
+		}
+	}
+
+	// Print Python command
+	if (bPrintPythonCommand)
+	{
+		const FString GraphName = GetSchema()->GetSanitizedGraphName(GetGraph()->GetGraphName());
+
+		RigVMPythonUtils::Print(GetSchema()->GetGraphOuterName(GetGraph()),
+			FString::Printf(TEXT("asset.get_controller_by_name('%s').set_local_variable_index('%s', %d)"),
+				*GraphName,
+				*GetSchema()->GetSanitizedVariableName(Variable.Name.ToString()),
+				InNewIndex));
+	}
+
+	return true;
+}
+
 URigVMUserWorkflowOptions* URigVMController::MakeOptionsForWorkflow(UObject* InSubject, const FRigVMUserWorkflow& InWorkflow)
 {
 	URigVMUserWorkflowOptions* Options = nullptr;
diff --git a/Engine/Plugins/Runtime/RigVM/Source/RigVMDeveloper/Private/RigVMModel/RigVMControllerActions.cpp b/Engine/Plugins/Runtime/RigVM/Source/RigVMDeveloper/Private/RigVMModel/RigVMControllerActions.cpp
--- a/Engine/Plugins/Runtime/RigVM/Source/RigVMDeveloper/Private/RigVMModel/RigVMControllerActions.cpp
+++ b/Engine/Plugins/Runtime/RigVM/Source/RigVMDeveloper/Private/RigVMModel/RigVMControllerActions.cpp
@@ -2258,6 +2258,48 @@ bool FRigVMChangeLocalVariableDefaultValueAction::Redo()
 	return FRigVMBaseAction::Redo();
 }
 
+FRigVMSetLocalVariableIndexAction::FRigVMSetLocalVariableIndexAction()
+: FRigVMBaseAction(nullptr)
+, VariableGuid()
+, OldIndex(INDEX_NONE)
+, NewIndex(INDEX_NONE)
+{
+}
+
+FRigVMSetLocalVariableIndexAction::FRigVMSetLocalVariableIndexAction(
+	URigVMController* InController,
+	const FGuid& InVariableGuid,
+	int32 InOldIndex,
+	int32 InNewIndex)
+: FRigVMBaseAction(InController)
+, VariableGuid(InVariableGuid)
+, OldIndex(InOldIndex)
+, NewIndex(InNewIndex)
+{
+}
+
+bool FRigVMSetLocalVariableIndexAction::Undo()
+{
+	if (!FRigVMBaseAction::Undo())
+	{
+		return false;
+	}
+	return GetController()->SetLocalVariableIndex(VariableGuid, OldIndex, false);
+}
+
+bool FRigVMSetLocalVariableIndexAction::Redo()
+{
+	if(!CanUndoRedo())
+	{
+		return false;
+	}
+	if (!GetController()->SetLocalVariableIndex(VariableGuid, NewIndex, false))
+	{
+		return false;
+	}
+	return FRigVMBaseAction::Redo();
+}
+
 FRigVMPromoteNodeAction::FRigVMPromoteNodeAction()
 : FRigVMBaseAction(nullptr)
 , LibraryNodePath()
diff --git a/Engine/Plugins/Runtime/RigVM/Source/RigVMDeveloper/Public/RigVMBlueprintLegacy.h b/Engine/Plugins/Runtime/RigVM/Source/RigVMDeveloper/Public/RigVMBlueprintLegacy.h
--- a/Engine/Plugins/Runtime/RigVM/Source/RigVMDeveloper/Public/RigVMBlueprintLegacy.h
+++ b/Engine/Plugins/Runtime/RigVM/Source/RigVMDeveloper/Public/RigVMBlueprintLegacy.h
@@ -442,6 +442,10 @@ public:
 	UE_API virtual bool ChangeMemberVariableType(const FName& InName, const FString& InCPPType, bool bIsPublic = false, bool bIsReadOnly = false, FString InDefaultValue = TEXT("")) override;
 	UE_API virtual bool ChangeMemberVariableType(const FName& InName, const FEdGraphPinType& InType) override;
 
+	UFUNCTION(BlueprintCallable, Category = "Variables")
+	UE_API virtual bool SetVariableIndex(const FName& InName, int32 NewIndex) override;
+	UE_API virtual bool SetVariableIndex(const FGuid& InVariableGuid, int32 NewIndex);
+
 	UFUNCTION(BlueprintPure, Category = "Variants", meta = (DisplayName = "GetAssetVariant", ScriptName = "GetAssetVariant"))
 	virtual FRigVMVariant GetAssetVariantBP() const { return AssetVariant; }
 
diff --git a/Engine/Plugins/Runtime/RigVM/Source/RigVMDeveloper/Public/RigVMEditorAsset.h b/Engine/Plugins/Runtime/RigVM/Source/RigVMDeveloper/Public/RigVMEditorAsset.h
--- a/Engine/Plugins/Runtime/RigVM/Source/RigVMDeveloper/Public/RigVMEditorAsset.h
+++ b/Engine/Plugins/Runtime/RigVM/Source/RigVMDeveloper/Public/RigVMEditorAsset.h
@@ -212,6 +212,7 @@ public:
 	virtual bool RenameMemberVariable(const FName& InOldName, const FName& InNewName) = 0;
 	virtual bool ChangeMemberVariableType(const FName& InName, const FString& InCPPType, bool bIsPublic = false, bool bIsReadOnly = false, FString InDefaultValue = TEXT("")) = 0;
 	virtual bool ChangeMemberVariableType(const FName& InName, const FEdGraphPinType& InType) = 0;
+	virtual bool SetVariableIndex(const FName& InName, int32 NewIndex) = 0;
 	virtual FText GetVariableTooltip(const FName& InName) const = 0;
 	virtual FString GetVariableCategory(const FName& InName) = 0;
 	virtual FString GetVariableMetadataValue(const FName& InName, const FName& InKey) = 0;
@@ -851,6 +852,9 @@ class URigVMEditorAsset : public UObject, public IRigVMEditorAssetInterface, pub
 	UE_API virtual bool ChangeMemberVariableType(const FName& InName, const FEdGraphPinType& InType) override;
 	UFUNCTION(BlueprintCallable, Category = "VM")
 	UE_API virtual bool ChangeMemberVariableType(const FName& InName, const FString& InCPPType, bool bIsPublic = false, bool bIsReadOnly = false, FString InDefaultValue = TEXT("")) override;
+	UFUNCTION(BlueprintCallable, Category = "VM")
+	UE_API virtual bool SetVariableIndex(const FName& InName, int32 NewIndex) override;
+	UE_API virtual bool SetVariableIndex(const FGuid& InVariableGuid, int32 NewIndex);
 	UE_API virtual FText GetVariableTooltip(const FName& InName) const override;
 	UE_API virtual FString GetVariableCategory(const FName& InName) override;
 	UE_API virtual FString GetVariableMetadataValue(const FName& InName, const FName& InKey) override;
diff --git a/Engine/Plugins/Runtime/RigVM/Source/RigVMDeveloper/Public/RigVMModel/RigVMController.h b/Engine/Plugins/Runtime/RigVM/Source/RigVMDeveloper/Public/RigVMModel/RigVMController.h
--- a/Engine/Plugins/Runtime/RigVM/Source/RigVMDeveloper/Public/RigVMModel/RigVMController.h
+++ b/Engine/Plugins/Runtime/RigVM/Source/RigVMDeveloper/Public/RigVMModel/RigVMController.h
@@ -1181,6 +1181,13 @@ public:
 	UFUNCTION(BlueprintCallable, Category = RigVMController)
 	UE_API bool SetLocalVariableDefaultValue(const FName& InVariableName, const FString& InDefaultValue, bool bSetupUndoRedo = true, bool bPrintPythonCommand = false);
 
+	// Reorder a local variable to a new index (UFUNCTION wrapper - converts name to GUID)
+	UFUNCTION(BlueprintCallable, Category = RigVMController)
+	UE_API bool SetLocalVariableIndex(const FName& InVariableName, int32 InNewIndex, bool bSetupUndoRedo = true, bool bPrintPythonCommand = false);
+
+	// Reorder a local variable to a new index (primary GUID-based implementation)
+	UE_API bool SetLocalVariableIndex(const FGuid& InVariableGuid, int32 InNewIndex, bool bSetupUndoRedo = true, bool bPrintPythonCommand = false);
+
 	// creates the options struct for a given workflow
 	UFUNCTION(BlueprintCallable, Category = RigVMController)
 	UE_API URigVMUserWorkflowOptions* MakeOptionsForWorkflow(UObject* InSubject, const FRigVMUserWorkflow& InWorkflow);
diff --git a/Engine/Plugins/Runtime/RigVM/Source/RigVMDeveloper/Public/RigVMModel/RigVMControllerActions.h b/Engine/Plugins/Runtime/RigVM/Source/RigVMDeveloper/Public/RigVMModel/RigVMControllerActions.h
--- a/Engine/Plugins/Runtime/RigVM/Source/RigVMDeveloper/Public/RigVMModel/RigVMControllerActions.h
+++ b/Engine/Plugins/Runtime/RigVM/Source/RigVMDeveloper/Public/RigVMModel/RigVMControllerActions.h
@@ -1545,6 +1545,32 @@ public:
 	FString DefaultValue;
 };
 
+/**
+ * An action to reorder a local variable to a new index
+ */
+USTRUCT()
+struct FRigVMSetLocalVariableIndexAction : public FRigVMBaseAction
+{
+	GENERATED_BODY()
+
+public:
+	FRigVMSetLocalVariableIndexAction();
+	FRigVMSetLocalVariableIndexAction(URigVMController* InController, const FGuid& InVariableGuid, int32 InOldIndex, int32 InNewIndex);
+	virtual ~FRigVMSetLocalVariableIndexAction() = default;
+	virtual UScriptStruct* GetScriptStruct() const override { return FRigVMSetLocalVariableIndexAction::StaticStruct(); }
+	virtual bool Undo() override;
+	virtual bool Redo() override;
+
+	UPROPERTY()
+	FGuid VariableGuid;
+
+	UPROPERTY()
+	int32 OldIndex;
+
+	UPROPERTY()
+	int32 NewIndex;
+};
+
 /**
  * An action to promote a function to collapse node or vice versa
  */
diff --git a/Engine/Plugins/Runtime/RigVM/Source/RigVMDeveloper/Public/RigVMModel/RigVMNotifications.h b/Engine/Plugins/Runtime/RigVM/Source/RigVMDeveloper/Public/RigVMModel/RigVMNotifications.h
--- a/Engine/Plugins/Runtime/RigVM/Source/RigVMDeveloper/Public/RigVMModel/RigVMNotifications.h
+++ b/Engine/Plugins/Runtime/RigVM/Source/RigVMDeveloper/Public/RigVMModel/RigVMNotifications.h
@@ -66,6 +66,7 @@ enum class ERigVMGraphNotifType : uint8
 	LocalVariableRemoved, // A local variable has been removed from a function (Subject == URigVMLibraryNode)
 	LocalVariableRenamed, // A local variable has been renamed inside a function (Subject == URigVMLibraryNode)
 	LocalVariableTypeChanged, // A local variable has been renamed inside a function (Subject == URigVMLibraryNode)
+	LocalVariableIndexChanged, // A local variable's index has changed (Subject == URigVMLibraryNode)
 	Invalid // The max for this enum (used for guarding)
 };
 
diff --git a/Engine/Plugins/Runtime/RigVM/Source/RigVMEditor/Private/Editor/RigVMNewEditor.cpp b/Engine/Plugins/Runtime/RigVM/Source/RigVMEditor/Private/Editor/RigVMNewEditor.cpp
--- a/Engine/Plugins/Runtime/RigVM/Source/RigVMEditor/Private/Editor/RigVMNewEditor.cpp
+++ b/Engine/Plugins/Runtime/RigVM/Source/RigVMEditor/Private/Editor/RigVMNewEditor.cpp
@@ -218,6 +218,13 @@ void FRigVMNewEditor::HandlePreVariablesChanged(FName InName)
 
 void FRigVMNewEditor::HandlePostVariablesChanged(FName InName)
 {
+	// Refresh Graph Explorer to show updated variable state
+	if (GraphExplorerWidget.IsValid())
+	{
+		GraphExplorerWidget->Refresh();
+	}
+
+	// Refresh details panel
 	if (Inspector)
 	{
 		SelectVariable(ElementInDetailsPanel);
diff --git a/Engine/Plugins/Runtime/RigVM/Source/RigVMEditor/Private/Editor/RigVMVariableDetailCustomization.cpp b/Engine/Plugins/Runtime/RigVM/Source/RigVMEditor/Private/Editor/RigVMVariableDetailCustomization.cpp
--- a/Engine/Plugins/Runtime/RigVM/Source/RigVMEditor/Private/Editor/RigVMVariableDetailCustomization.cpp
+++ b/Engine/Plugins/Runtime/RigVM/Source/RigVMEditor/Private/Editor/RigVMVariableDetailCustomization.cpp
@@ -250,22 +250,22 @@ void FRigVMVariableDetailCustomization::CustomizeDetails(IDetailLayoutBuilder& D
 
 
 
-		TSharedPtr<SToolTip> PrivateTooltip = IDocumentation::Get()->CreateToolTip(LOCTEXT("VariablePrivate_Tooltip", "Should this variable be private (derived blueprints cannot modify it)?"), NULL, "", TEXT("Private"));
+		TSharedPtr<SToolTip> PublicTooltip = IDocumentation::Get()->CreateToolTip(LOCTEXT("VariablePublic_Tooltip", "Should this variable be public?"), NULL, "", TEXT("Public"));
 
-		Category.AddCustomRow(LOCTEXT("VariablePrivate", "Private"))
+		Category.AddCustomRow(LOCTEXT("VariablePublic", "Public"))
 		.NameContent()
 		[
 			SNew(STextBlock)
-			.ToolTip(PrivateTooltip)
-			.Text( LOCTEXT("VariablePrivate", "Private") )
+			.ToolTip(PublicTooltip)
+			.Text( LOCTEXT("VariablePublic", "Public") )
 			.Font( DetailFontInfo )
 		]
 		.ValueContent()
 		[
 			SNew(SCheckBox)
-			.IsChecked( this, &FRigVMVariableDetailCustomization::OnGetPrivateCheckboxState )
-			.OnCheckStateChanged( this, &FRigVMVariableDetailCustomization::OnPrivateChanged )
-			.ToolTip(PrivateTooltip)
+			.IsChecked( this, &FRigVMVariableDetailCustomization::OnGetPublicCheckboxState )
+			.OnCheckStateChanged( this, &FRigVMVariableDetailCustomization::OnPublicChanged )
+			.ToolTip(PublicTooltip)
 		];
 
 		TSharedPtr<SToolTip> ExposeToCinematicsTooltip = IDocumentation::Get()->CreateToolTip(LOCTEXT("VariableExposeToCinematics_Tooltip", "Should this variable be exposed for Sequencer to modify?"), NULL, "", TEXT("ExposeToCinematics"));
@@ -999,17 +999,17 @@ void FRigVMVariableDetailCustomization::OnExposedToSpawnChanged(ECheckBoxState I
 	}
 }
 
-ECheckBoxState FRigVMVariableDetailCustomization::OnGetPrivateCheckboxState() const
+ECheckBoxState FRigVMVariableDetailCustomization::OnGetPublicCheckboxState() const
 {
 	const FProperty* Property = CachedVariableProperty.Get();
 	if (Property)
 	{
-		return (Property && Property->GetBoolMetaData(FBlueprintMetadata::MD_Private) != false) ? ECheckBoxState::Checked : ECheckBoxState::Unchecked;
+		return (Property && (Property->PropertyFlags & CPF_DisableEditOnInstance) == 0) ? ECheckBoxState::Checked : ECheckBoxState::Unchecked;
 	}
 	return ECheckBoxState::Unchecked;
 }
 
-void FRigVMVariableDetailCustomization::OnPrivateChanged(ECheckBoxState InNewState)
+void FRigVMVariableDetailCustomization::OnPublicChanged(ECheckBoxState InNewState)
 {
 	TStrongObjectPtr<UObject> Blueprint = BlueprintPtr.GetWeakObjectPtr().Pin();
 	if (!Blueprint)
@@ -1020,8 +1020,8 @@ void FRigVMVariableDetailCustomization::OnPrivateChanged(ECheckBoxState InNewSta
 	
 	if (CachedVariableName != NAME_None)
 	{
-		const bool bPrivate = (InNewState == ECheckBoxState::Checked);
-		Interface->SetVariablePrivate(CachedVariableName, bPrivate);
+		const bool bPublic = (InNewState == ECheckBoxState::Checked);
+		Interface->SetVariablePublic(CachedVariableName, bPublic);
 	}
 }
 
diff --git a/Engine/Plugins/Runtime/RigVM/Source/RigVMEditor/Private/Widgets/SRigVMEditorGraphExplorer.cpp b/Engine/Plugins/Runtime/RigVM/Source/RigVMEditor/Private/Widgets/SRigVMEditorGraphExplorer.cpp
--- a/Engine/Plugins/Runtime/RigVM/Source/RigVMEditor/Private/Widgets/SRigVMEditorGraphExplorer.cpp
+++ b/Engine/Plugins/Runtime/RigVM/Source/RigVMEditor/Private/Widgets/SRigVMEditorGraphExplorer.cpp
@@ -519,8 +519,8 @@ void SRigVMEditorGraphExplorer::CreateWidgets()
 	Delegates.OnToggleVariablePublic = FRigVMGraphExplorer_OnToggleVariablePublic::CreateSP(this, &SRigVMEditorGraphExplorer::OnToggleVariablePublic);
 	Delegates.OnIsFunctionFocused = FRigVMGraphExplorer_OnIsFunctionFocused::CreateSP(this, &SRigVMEditorGraphExplorer::IsFunctionFocused);
 	Delegates.OnGetCustomPinFilters = FRigVMGraphExplorer_OnGetCustomPinFilters::CreateSP(this, &SRigVMEditorGraphExplorer::GetCustomPinFilters);
-	
-	Delegates.OnSelectionChanged = FRigVMGraphExplorer_OnSelectionChanged::CreateSP(this, &SRigVMEditorGraphExplorer::HandleSelectionChanged);
+	Delegates.OnReorderVariable = FRigVMGraphExplorer_OnReorderVariable::CreateSP(this, &SRigVMEditorGraphExplorer::OnReorderVariable);
+	Delegates.OnReorderLocalVariable = FRigVMGraphExplorer_OnReorderVariable::CreateSP(this, &SRigVMEditorGraphExplorer::OnReorderLocalVariable);
 	
 	SAssignNew(TreeView, SRigVMEditorGraphExplorerTreeView, SharedThis(this))
 		.RigTreeDelegates(Delegates)
@@ -1051,6 +1051,8 @@ bool SRigVMEditorGraphExplorer::OnSetFunctionCategory(const FString& InFunctionP
 		return false;
 	}
 
+	const FScopedTransaction Transaction(LOCTEXT("SetFunctionCategory", "Set Function Category"));
+
 	const URigVMEdGraphSchema* Schema = GetDefault<URigVMEdGraphSchema>();
 	if (URigVMFunctionLibrary* Library = Editor->GetRigVMAssetInterface()->GetLocalFunctionLibrary())
 	{
@@ -1124,8 +1126,8 @@ bool SRigVMEditorGraphExplorer::OnSetVariableCategory(const FString& InVariable,
 		return false;
 	}
 
+	const FScopedTransaction Transaction(LOCTEXT("SetVariableCategory", "Set Variable Category"));
 	Editor->GetRigVMAssetInterface()->SetVariableCategory(*InVariable, InCategory);
-	bNeedsRefresh = true;
 	return true;
 }
 
@@ -2143,6 +2145,73 @@ bool SRigVMEditorGraphExplorer::OnRenameVariable(const FRigVMExplorerElementKey&
 	return true;
 }
 
+bool SRigVMEditorGraphExplorer::OnReorderVariable(const FRigVMExplorerElementKey& InVariable, int32 NewIndex)
+{
+	TSharedPtr<IRigVMEditor> Editor = RigVMEditor.Pin();
+	if (!Editor)
+	{
+		return false;
+	}
+
+	FRigVMEditorAssetInterfacePtr Blueprint = Editor->GetRigVMAssetInterface();
+
+	const FScopedTransaction Transaction(LOCTEXT("ReorderVariable", "Reorder Variable"));
+
+	// Call asset interface to reorder using name (will be converted to GUID internally)
+	bool bSuccess = Blueprint->SetVariableIndex(*InVariable.Name, NewIndex);
+
+	if (bSuccess)
+	{
+		Refresh();
+	}
+
+	return bSuccess;
+}
+
+bool SRigVMEditorGraphExplorer::OnReorderLocalVariable(const FRigVMExplorerElementKey& InVariable, int32 NewIndex)
+{
+	TSharedPtr<IRigVMEditor> Editor = RigVMEditor.Pin();
+	if (!Editor)
+	{
+		return false;
+	}
+
+	FRigVMEditorAssetInterfacePtr Blueprint = Editor->GetRigVMAssetInterface();
+
+	// Find local variable in the focused model
+	URigVMGraph* Graph = Blueprint->GetFocusedModel();
+	if (!Graph)
+	{
+		return false;
+	}
+
+	URigVMController* Controller = Blueprint->GetRigVMClient()->GetController(Graph);
+	if (!Controller)
+	{
+		return false;
+	}
+
+	// Find the local variable by name
+	const TArray<FRigVMGraphVariableDescription>& LocalVars = Graph->GetLocalVariables();
+	const FRigVMGraphVariableDescription* LocalVar = LocalVars.FindByPredicate(
+		[&InVariable](const FRigVMGraphVariableDescription& Var) {
+			return Var.Name.ToString() == InVariable.Name;
+		});
+
+	if (!LocalVar)
+	{
+		return false;
+	}
+
+	const FScopedTransaction Transaction(LOCTEXT("ReorderLocalVariable", "Reorder Local Variable"));
+	bool bSuccess = Controller->SetLocalVariableIndex(LocalVar->Guid, NewIndex, true, true);
+	if (bSuccess)
+	{
+		Refresh();
+	}
+	return bSuccess;
+}
+
 bool SRigVMEditorGraphExplorer::OnCanRenameVariable(const FRigVMExplorerElementKey& InOldKey, const FString& InNewName, FText& OutErrorMessage) const
 {
 	if (InNewName.Equals(InOldKey.Name, ESearchCase::CaseSensitive))
@@ -2270,10 +2339,6 @@ bool SRigVMEditorGraphExplorer::CanRemoveUnusedVariables() const
 	return true;
 }
 
-void SRigVMEditorGraphExplorer::HandleSelectionChanged(TSharedPtr<FRigVMEditorGraphExplorerTreeElement> Selection, ESelectInfo::Type SelectInfo)
-{
-	TreeView->RefreshTreeView();
-}
 
 
 #undef LOCTEXT_NAMESPACE
diff --git a/Engine/Plugins/Runtime/RigVM/Source/RigVMEditor/Private/Widgets/SRigVMEditorGraphExplorerTreeView.cpp b/Engine/Plugins/Runtime/RigVM/Source/RigVMEditor/Private/Widgets/SRigVMEditorGraphExplorerTreeView.cpp
--- a/Engine/Plugins/Runtime/RigVM/Source/RigVMEditor/Private/Widgets/SRigVMEditorGraphExplorerTreeView.cpp
+++ b/Engine/Plugins/Runtime/RigVM/Source/RigVMEditor/Private/Widgets/SRigVMEditorGraphExplorerTreeView.cpp
@@ -98,6 +98,9 @@ void SRigVMEditorGraphExplorerItem::Construct(const FArguments& InArgs,
 		STableRow<TSharedPtr<FRigVMEditorGraphExplorerTreeElement>>::FArguments()
 		.ShowWires(false)
 		.Padding(FMargin(0.0f, 2.0f, .0f, 0.0f))
+		.OnDragDetected(Delegates.OnDragDetected)
+		.OnCanAcceptDrop(this, &SRigVMEditorGraphExplorerItem::OnCanAcceptDrop)
+		.OnAcceptDrop(this, &SRigVMEditorGraphExplorerItem::OnAcceptDrop)
 		.Content()
 		[
 			SNew(SHorizontalBox)
@@ -122,6 +125,8 @@ void SRigVMEditorGraphExplorerItem::Construct(const FArguments& InArgs,
 			STableRow<TSharedPtr<FRigVMEditorGraphExplorerTreeElement>>::FArguments()
 			.ShowWires(false)
 			.OnDragDetected(Delegates.OnDragDetected)
+			.OnCanAcceptDrop(this, &SRigVMEditorGraphExplorerItem::OnCanAcceptDrop)
+			.OnAcceptDrop(this, &SRigVMEditorGraphExplorerItem::OnAcceptDrop)
 			.Content()
 			[
 				SNew(SHorizontalBox)
@@ -478,6 +483,356 @@ FReply SRigVMEditorGraphExplorerItem::OnAddButtonClickedOnSection(TSharedRef<FRi
 	return FReply::Unhandled();
 }
 
+TOptional<EItemDropZone> SRigVMEditorGraphExplorerItem::OnCanAcceptDrop(const FDragDropEvent& DragDropEvent, EItemDropZone DropZone, TSharedPtr<FRigVMEditorGraphExplorerTreeElement> TargetItem)
+{
+	static const TOptional<EItemDropZone> InvalidDropZone;
+
+	TSharedPtr<FRigVMGraphExplorerDragDropOp> DragDropOp = DragDropEvent.GetOperationAs<FRigVMGraphExplorerDragDropOp>();
+	if (!DragDropOp.IsValid())
+	{
+		return InvalidDropZone;
+	}
+
+	const FRigVMExplorerElementKey& DraggedElement = DragDropOp->GetElement();
+	const ERigVMExplorerElementType::Type DraggedType = DraggedElement.Type;
+	const ERigVMExplorerElementType::Type TargetType = TargetItem->Key.Type;
+
+	// Only Variable, LocalVariable, VariableCategory, and FunctionCategory can be dropped
+	if (DraggedType != ERigVMExplorerElementType::Variable &&
+		DraggedType != ERigVMExplorerElementType::LocalVariable &&
+		DraggedType != ERigVMExplorerElementType::VariableCategory &&
+		DraggedType != ERigVMExplorerElementType::FunctionCategory)
+	{
+		return InvalidDropZone;
+	}
+
+	// Cannot drop on self
+	if (DraggedElement.Name == TargetItem->Key.Name && DraggedType == TargetType)
+	{
+		return InvalidDropZone;
+	}
+
+	// Check same section - items can only drop within their own section
+	bool bSameSection = false;
+	if (DraggedType == TargetType)
+	{
+		// Same type always same section
+		bSameSection = true;
+	}
+	else if (TargetType == ERigVMExplorerElementType::VariableCategory)
+	{
+		// Need to verify the category is in the correct section
+		// Walk up the parent hierarchy to find which section this category belongs to
+		TSharedPtr<SRigVMEditorGraphExplorerTreeView> TreeView = WeakTreeView.Pin();
+		if (TreeView.IsValid())
+		{
+			FRigVMExplorerElementKey CurrentKey = TargetItem->Key;
+			while (const FRigVMExplorerElementKey* ParentKey = TreeView->ParentMap.Find(CurrentKey))
+			{
+				if (ParentKey->Type == ERigVMExplorerElementType::Section)
+				{
+					// Found the section - check if it matches the dragged type
+					if ((DraggedType == ERigVMExplorerElementType::Variable || DraggedType == ERigVMExplorerElementType::VariableCategory) &&
+						ParentKey->Name == VariablesSectionName.ToString())
+					{
+						bSameSection = true;
+					}
+					else if ((DraggedType == ERigVMExplorerElementType::LocalVariable) &&
+						ParentKey->Name == LocalVariablesSectionName.ToString())
+					{
+						bSameSection = true;
+					}
+					break;
+				}
+				CurrentKey = *ParentKey;
+			}
+		}
+	}
+	else if (DraggedType == ERigVMExplorerElementType::VariableCategory || DraggedType == ERigVMExplorerElementType::FunctionCategory)
+	{
+		// When dragging a category, allow dropping onto same category type or compatible variable type
+		if (TargetType == ERigVMExplorerElementType::Variable && DraggedType == ERigVMExplorerElementType::VariableCategory)
+		{
+			bSameSection = true;
+		}
+		else if (TargetType == ERigVMExplorerElementType::LocalVariable && DraggedType == ERigVMExplorerElementType::VariableCategory)
+		{
+			// Need to check if the category is in the local variables section
+			TSharedPtr<SRigVMEditorGraphExplorerTreeView> TreeView = WeakTreeView.Pin();
+			if (TreeView.IsValid())
+			{
+				FRigVMExplorerElementKey CurrentKey = DraggedElement;
+				while (const FRigVMExplorerElementKey* ParentKey = TreeView->ParentMap.Find(CurrentKey))
+				{
+					if (ParentKey->Type == ERigVMExplorerElementType::Section)
+					{
+						if (ParentKey->Name == LocalVariablesSectionName.ToString())
+						{
+							bSameSection = true;
+						}
+						break;
+					}
+					CurrentKey = *ParentKey;
+				}
+			}
+		}
+	}
+
+	if (!bSameSection)
+	{
+		return InvalidDropZone;
+	}
+
+	// Allow onto category to nest categories (category onto category)
+	// Must check this BEFORE "variable onto category" to ensure cycle detection runs
+	// Must be same category type: VariableCategory onto VariableCategory, or FunctionCategory onto FunctionCategory
+	if (DraggedType == TargetType &&
+		(DraggedType == ERigVMExplorerElementType::VariableCategory || DraggedType == ERigVMExplorerElementType::FunctionCategory) &&
+		DropZone == EItemDropZone::OntoItem)
+	{
+		// Prevent dropping a category onto itself or its own descendants (would create a cycle)
+		// Allow dropping a descendant onto its ancestor (promotes/un-nests the category)
+		FString DraggedCategoryPath = DraggedElement.Name.TrimStartAndEnd();
+		FString TargetCategoryPath = TargetItem->Key.Name.TrimStartAndEnd();
+
+		// Don't allow dropping onto self
+		if (DraggedCategoryPath.Equals(TargetCategoryPath, ESearchCase::IgnoreCase))
+		{
+			return InvalidDropZone;
+		}
+
+		// Don't allow dropping an ancestor onto its own descendant (e.g., dropping "aaa" onto "aaa|bbb|ccc")
+		// This would create a cycle where "aaa" contains itself
+		FString DraggedWithSeparator = DraggedCategoryPath + TEXT("|");
+		if (TargetCategoryPath.StartsWith(DraggedWithSeparator, ESearchCase::IgnoreCase))
+		{
+			return InvalidDropZone;
+		}
+
+		// ALLOW dropping a descendant onto its own ancestor (e.g., dropping "aaa|bbb|ccc" onto "aaa")
+		// This promotes/un-nests the category: "aaa|bbb|ccc" becomes "aaa|ccc"
+		// No cycle is created because we're moving up the hierarchy
+
+		return EItemDropZone::OntoItem;
+	}
+
+	// Allow onto category to change category (variable onto category)
+	// Check this AFTER category-onto-category to ensure only variables (not categories) are accepted
+	if ((DraggedType == ERigVMExplorerElementType::Variable || DraggedType == ERigVMExplorerElementType::LocalVariable) &&
+		TargetType == ERigVMExplorerElementType::VariableCategory &&
+		DropZone == EItemDropZone::OntoItem)
+	{
+		return EItemDropZone::OntoItem;
+	}
+
+	// Allow above/below same variable type for reordering
+	if (TargetType == DraggedType &&
+		(DropZone == EItemDropZone::AboveItem || DropZone == EItemDropZone::BelowItem))
+	{
+		return DropZone;
+	}
+
+	return InvalidDropZone;
+}
+
+FReply SRigVMEditorGraphExplorerItem::OnAcceptDrop(const FDragDropEvent& DragDropEvent, EItemDropZone DropZone, TSharedPtr<FRigVMEditorGraphExplorerTreeElement> TargetItem)
+{
+	if (!OnCanAcceptDrop(DragDropEvent, DropZone, TargetItem).IsSet())
+	{
+		return FReply::Unhandled();
+	}
+
+	TSharedPtr<FRigVMGraphExplorerDragDropOp> DragDropOp = DragDropEvent.GetOperationAs<FRigVMGraphExplorerDragDropOp>();
+	if (!DragDropOp.IsValid())
+	{
+		return FReply::Unhandled();
+	}
+
+	const FRigVMExplorerElementKey& DraggedElement = DragDropOp->GetElement();
+	const ERigVMExplorerElementType::Type DraggedType = DraggedElement.Type;
+	const ERigVMExplorerElementType::Type TargetType = TargetItem->Key.Type;
+
+	TSharedPtr<SRigVMEditorGraphExplorerTreeView> TreeView = WeakTreeView.Pin();
+	if (!TreeView.IsValid())
+	{
+		return FReply::Unhandled();
+	}
+
+	// Handle category nesting (category onto category)
+	if ((DraggedType == ERigVMExplorerElementType::VariableCategory || DraggedType == ERigVMExplorerElementType::FunctionCategory) &&
+		TargetType == DraggedType &&
+		DropZone == EItemDropZone::OntoItem)
+	{
+
+		FString DraggedCategoryPath = DraggedElement.Name;
+		FString TargetCategoryPath = TargetItem->Key.Name;
+
+		// Extract the leaf category name (last part after the last pipe)
+		int32 LastPipeIndex;
+		FString LeafCategoryName = DraggedCategoryPath;
+		if (DraggedCategoryPath.FindLastChar('|', LastPipeIndex))
+		{
+			LeafCategoryName = DraggedCategoryPath.RightChop(LastPipeIndex + 1);
+		}
+
+		// New nested category path: TargetPath|LeafName
+		FString NewCategoryPath = TargetCategoryPath + TEXT("|") + LeafCategoryName;
+
+		// Find all items that belong to the dragged category (or its subcategories) and update them
+		if (DraggedType == ERigVMExplorerElementType::VariableCategory)
+		{
+			TArray<FRigVMGraphVariableDescription> Variables = Delegates.GetVariables();
+			for (const FRigVMGraphVariableDescription& Var : Variables)
+			{
+				FString VarCategory = Delegates.GetVariableCategory(Var.Name.ToString());
+
+				// Check if this variable belongs to the dragged category or any of its subcategories
+				if (VarCategory == DraggedCategoryPath || VarCategory.StartsWith(DraggedCategoryPath + TEXT("|")))
+				{
+					// Replace the dragged category path with the new nested path
+					FString UpdatedCategory = NewCategoryPath;
+					if (VarCategory.Len() > DraggedCategoryPath.Len())
+					{
+						// Preserve subcategory structure: "B|C" -> "A|B|C"
+						UpdatedCategory = NewCategoryPath + VarCategory.RightChop(DraggedCategoryPath.Len());
+					}
+					Delegates.SetVariableCategory(Var.Name.ToString(), UpdatedCategory);
+				}
+			}
+
+			// Also handle local variables
+			TArray<FRigVMGraphVariableDescription> LocalVariables = Delegates.GetLocalVariables();
+			for (const FRigVMGraphVariableDescription& Var : LocalVariables)
+			{
+				FString VarCategory = Delegates.GetVariableCategory(Var.Name.ToString());
+
+				if (VarCategory == DraggedCategoryPath || VarCategory.StartsWith(DraggedCategoryPath + TEXT("|")))
+				{
+					FString UpdatedCategory = NewCategoryPath;
+					if (VarCategory.Len() > DraggedCategoryPath.Len())
+					{
+						UpdatedCategory = NewCategoryPath + VarCategory.RightChop(DraggedCategoryPath.Len());
+					}
+					Delegates.SetVariableCategory(Var.Name.ToString(), UpdatedCategory);
+				}
+			}
+		}
+		else if (DraggedType == ERigVMExplorerElementType::FunctionCategory)
+		{
+			TArray<URigVMLibraryNode*> Functions = Delegates.GetFunctions();
+			for (URigVMLibraryNode* Function : Functions)
+			{
+				FString FunctionCategory = Delegates.GetFunctionCategory(Function->GetPathName());
+
+				if (FunctionCategory == DraggedCategoryPath || FunctionCategory.StartsWith(DraggedCategoryPath + TEXT("|")))
+				{
+					FString UpdatedCategory = NewCategoryPath;
+					if (FunctionCategory.Len() > DraggedCategoryPath.Len())
+					{
+						UpdatedCategory = NewCategoryPath + FunctionCategory.RightChop(DraggedCategoryPath.Len());
+					}
+					Delegates.SetFunctionCategory(Function->GetPathName(), UpdatedCategory);
+				}
+			}
+		}
+
+		TreeView->RefreshTreeView();
+		return FReply::Handled();
+	}
+
+	// Handle category assignment (variable onto category)
+	if (TargetType == ERigVMExplorerElementType::VariableCategory && DropZone == EItemDropZone::OntoItem)
+	{
+
+		// Use the full category path from the category key
+		FString CategoryPath = TargetItem->Key.Name;
+
+		Delegates.SetVariableCategory(DraggedElement.Name, CategoryPath);
+		TreeView->RefreshTreeView();
+		return FReply::Handled();
+	}
+
+	// Handle reordering within same variable type
+	if (TargetType == DraggedType &&
+		(DropZone == EItemDropZone::AboveItem || DropZone == EItemDropZone::BelowItem))
+	{
+		// Get all variables of this type to find source and target indices
+		TArray<FRigVMGraphVariableDescription> Variables;
+		if (DraggedType == ERigVMExplorerElementType::Variable)
+		{
+			Variables = Delegates.GetVariables();
+		}
+		else if (DraggedType == ERigVMExplorerElementType::LocalVariable)
+		{
+			Variables = Delegates.GetLocalVariables();
+		}
+
+		// Find source index (where the dragged item currently is)
+		int32 SourceIndex = INDEX_NONE;
+		for (int32 i = 0; i < Variables.Num(); ++i)
+		{
+			if (Variables[i].Name.ToString() == DraggedElement.Name)
+			{
+				SourceIndex = i;
+				break;
+			}
+		}
+
+		// Find target index (where we want to drop)
+		int32 TargetIndex = INDEX_NONE;
+		for (int32 i = 0; i < Variables.Num(); ++i)
+		{
+			if (Variables[i].Name.ToString() == TargetItem->Key.Name)
+			{
+				TargetIndex = i;
+				// If dropping below, insert after the target
+				if (DropZone == EItemDropZone::BelowItem)
+				{
+					TargetIndex++;
+				}
+				break;
+			}
+		}
+
+		if (SourceIndex != INDEX_NONE && TargetIndex != INDEX_NONE)
+		{
+			// When dropping onto a variable, adopt that variable's category
+			// This allows both categorizing (if target has category) and uncategorizing (if target has no category)
+			FString TargetCategory = Delegates.GetVariableCategory(TargetItem->Key.Name);
+			FString CurrentCategory = Delegates.GetVariableCategory(DraggedElement.Name);
+			if (TargetCategory != CurrentCategory)
+			{
+				Delegates.SetVariableCategory(DraggedElement.Name, TargetCategory);
+			}
+
+			// Adjust target index if we're moving down
+			// When we remove the source item, all items after it shift down by one
+			if (SourceIndex < TargetIndex)
+			{
+				TargetIndex--;
+			}
+
+			bool bSuccess = false;
+			if (DraggedType == ERigVMExplorerElementType::Variable)
+			{
+				bSuccess = Delegates.OnReorderVariable.Execute(DraggedElement, TargetIndex);
+			}
+			else if (DraggedType == ERigVMExplorerElementType::LocalVariable)
+			{
+				bSuccess = Delegates.OnReorderLocalVariable.Execute(DraggedElement, TargetIndex);
+			}
+
+			if (bSuccess)
+			{
+				TreeView->RefreshTreeView();
+				return FReply::Handled();
+			}
+		}
+	}
+
+	return FReply::Unhandled();
+}
+
 void SRigVMEditorGraphExplorerItem::UpdateIsUsedInGraph()
 {
 	using namespace UE::RigVMEditor;
diff --git a/Engine/Plugins/Runtime/RigVM/Source/RigVMEditor/Public/Editor/RigVMVariableDetailCustomization.h b/Engine/Plugins/Runtime/RigVM/Source/RigVMEditor/Public/Editor/RigVMVariableDetailCustomization.h
--- a/Engine/Plugins/Runtime/RigVM/Source/RigVMEditor/Public/Editor/RigVMVariableDetailCustomization.h
+++ b/Engine/Plugins/Runtime/RigVM/Source/RigVMEditor/Public/Editor/RigVMVariableDetailCustomization.h
@@ -72,8 +72,8 @@ private:
 	UE_API ECheckBoxState OnGetExposedToSpawnCheckboxState() const;
 	UE_API void OnExposedToSpawnChanged(ECheckBoxState InNewState);
 
-	UE_API ECheckBoxState OnGetPrivateCheckboxState() const;
-	UE_API void OnPrivateChanged(ECheckBoxState InNewState);
+	UE_API ECheckBoxState OnGetPublicCheckboxState() const;
+	UE_API void OnPublicChanged(ECheckBoxState InNewState);
 	
 	UE_API ECheckBoxState OnGetExposedToCinematicsCheckboxState() const;
 	UE_API void OnExposedToCinematicsChanged(ECheckBoxState InNewState);
diff --git a/Engine/Plugins/Runtime/RigVM/Source/RigVMEditor/Public/Widgets/SRigVMEditorGraphExplorer.h b/Engine/Plugins/Runtime/RigVM/Source/RigVMEditor/Public/Widgets/SRigVMEditorGraphExplorer.h
--- a/Engine/Plugins/Runtime/RigVM/Source/RigVMEditor/Public/Widgets/SRigVMEditorGraphExplorer.h
+++ b/Engine/Plugins/Runtime/RigVM/Source/RigVMEditor/Public/Widgets/SRigVMEditorGraphExplorer.h
@@ -120,7 +120,7 @@ public:
 	void HandleGraphModifiedEvent(ERigVMGraphNotifType InNotifType, URigVMGraph* InGraph, UObject* InSubject);
 
 private:
-	bool bNeedsRefresh;
+	bool bNeedsRefresh = true;
 
 	FEdGraphPinType LastPinType;
 	
@@ -215,6 +215,8 @@ private:
 	bool OnCanRenameFunction(const FString& InOldPath, const FString& InNewPath, FText& OutErrorMessage) const;
 	bool OnRenameVariable(const FRigVMExplorerElementKey& InOldKey, const FString& InNewName);
 	bool OnCanRenameVariable(const FRigVMExplorerElementKey& InOldKey, const FString& InNewName, FText& OutErrorMessage) const;
+	bool OnReorderVariable(const FRigVMExplorerElementKey& InVariable, int32 NewIndex);
+	bool OnReorderLocalVariable(const FRigVMExplorerElementKey& InVariable, int32 NewIndex);
 
 	void OnFindReferences(const bool bSearchAllBlueprints);
 	bool CanFindReferences() const;
@@ -224,8 +226,6 @@ private:
 	void OnRemoveUnusedVariables();
 	bool CanRemoveUnusedVariables() const;
 
-	void HandleSelectionChanged(TSharedPtr<FRigVMEditorGraphExplorerTreeElement> Selection, ESelectInfo::Type SelectInfo);
-
 	TWeakPtr<IRigVMEditor> RigVMEditor;
 
 	TSharedPtr<SRigVMEditorGraphExplorerTreeView> TreeView;
diff --git a/Engine/Plugins/Runtime/RigVM/Source/RigVMEditor/Public/Widgets/SRigVMEditorGraphExplorerTreeView.h b/Engine/Plugins/Runtime/RigVM/Source/RigVMEditor/Public/Widgets/SRigVMEditorGraphExplorerTreeView.h
--- a/Engine/Plugins/Runtime/RigVM/Source/RigVMEditor/Public/Widgets/SRigVMEditorGraphExplorerTreeView.h
+++ b/Engine/Plugins/Runtime/RigVM/Source/RigVMEditor/Public/Widgets/SRigVMEditorGraphExplorerTreeView.h
@@ -105,6 +105,7 @@ DECLARE_DELEGATE_RetVal_TwoParams(bool, FRigVMGraphExplorer_OnSetVariablePinType
 DECLARE_DELEGATE_RetVal_OneParam(bool, FRigVMGraphExplorer_OnIsVariablePublic, const FString&);
 DECLARE_DELEGATE_RetVal_OneParam(bool, FRigVMGraphExplorer_OnToggleVariablePublic, const FString&);
 DECLARE_DELEGATE_RetVal(TArray<TSharedPtr<IPinTypeSelectorFilter>>, FRigVMGraphExplorer_OnGetCustomPinFilters);
+DECLARE_DELEGATE_RetVal_TwoParams(bool, FRigVMGraphExplorer_OnReorderVariable, const FRigVMExplorerElementKey&, int32 /* new index */);
 
 typedef STreeView<TSharedPtr<FRigVMEditorGraphExplorerTreeElement>>::FOnSelectionChanged FRigVMGraphExplorer_OnSelectionChanged;
 
@@ -150,7 +151,9 @@ public:
 	FRigVMGraphExplorer_OnIsVariablePublic OnIsVariablePublic;
 	FRigVMGraphExplorer_OnToggleVariablePublic OnToggleVariablePublic;
 	FRigVMGraphExplorer_OnGetCustomPinFilters OnGetCustomPinFilters;
-	
+	FRigVMGraphExplorer_OnReorderVariable OnReorderVariable;
+	FRigVMGraphExplorer_OnReorderVariable OnReorderLocalVariable;
+
 	FRigVMGraphExplorer_OnSelectionChanged OnSelectionChanged;
 	FRigVMGraphExplorer_OnRequestContextMenu OnRequestContextMenu;
 	FOnDragDetected OnDragDetected;
@@ -544,7 +547,11 @@ public:
 	bool OnVerifyNameChanged(const FText& InText, FText& OutErrorMessage);
 	void OnNameCommitted(const FText& InText, ETextCommit::Type InCommitType) const;
 	FReply OnAddButtonClickedOnSection(TSharedRef<FRigVMEditorGraphExplorerTreeElement> InElement);
-	
+
+	// Drag and drop support for reordering
+	virtual TOptional<EItemDropZone> OnCanAcceptDrop(const FDragDropEvent& DragDropEvent, EItemDropZone DropZone, TSharedPtr<FRigVMEditorGraphExplorerTreeElement> TargetItem);
+	virtual FReply OnAcceptDrop(const FDragDropEvent& DragDropEvent, EItemDropZone DropZone, TSharedPtr<FRigVMEditorGraphExplorerTreeElement> TargetItem);
+
 private:
 	/** Updates OptionalIsUsedInGraph */
 	void UpdateIsUsedInGraph();
