Add a comment to the RigVM drawing code to clarify capsule orientation

#Synced-CL 51744685

diff --git a/Engine/Plugins/Runtime/RigVM/Source/RigVM/Public/RigVMCore/RigVMDrawInterface.h b/Engine/Plugins/Runtime/RigVM/Source/RigVM/Public/RigVMCore/RigVMDrawInterface.h
--- a/Engine/Plugins/Runtime/RigVM/Source/RigVM/Public/RigVMCore/RigVMDrawInterface.h
+++ b/Engine/Plugins/Runtime/RigVM/Source/RigVM/Public/RigVMCore/RigVMDrawInterface.h
@@ -26,6 +26,7 @@ public:
 	UE_API void DrawSphere(const FTransform& WorldOffset, const FTransform& Transform, float Radius, const FLinearColor& Color, float Thickness = 0.f, int32 Detail = 12, ESceneDepthPriorityGroup DepthPriority = ESceneDepthPriorityGroup::SDPG_Foreground, float Lifetime = -1.f);
 	// Hemispheres are drawn so that they extend up along +ve Z
 	UE_API void DrawHemisphere(const FTransform& WorldOffset, const FTransform& Transform, float Radius, const FLinearColor& Color, float Thickness = 0.f, int32 Detail = 12, ESceneDepthPriorityGroup DepthPriority = ESceneDepthPriorityGroup::SDPG_Foreground, float Lifetime = -1.f);
+	// Capsules are drawn so that they extend along Z
 	UE_API void DrawCapsule(const FTransform& WorldOffset, const FTransform& Transform, float Radius, float Length, const FLinearColor& Color, float Thickness = 0.f, int32 Detail = 12, ESceneDepthPriorityGroup DepthPriority = ESceneDepthPriorityGroup::SDPG_Foreground, float Lifetime = -1.f);
 	UE_API void DrawAxes(const FTransform& WorldOffset, const FTransform& Transform, float Size, float Thickness = 0.f, ESceneDepthPriorityGroup DepthPriority = ESceneDepthPriorityGroup::SDPG_Foreground, float Lifetime = -1.f);
 	UE_API void DrawAxes(const FTransform& WorldOffset, const TArrayView<const FTransform>& Transforms, float Size, float Thickness = 0.f, ESceneDepthPriorityGroup DepthPriority = ESceneDepthPriorityGroup::SDPG_Foreground, float Lifetime = -1.f);
