RigVM: Disable fuzzy search by default

#rnx

#Synced-CL 51334386

diff --git a/Engine/Plugins/Runtime/RigVM/Source/RigVM/Private/RigVMSettings.cpp b/Engine/Plugins/Runtime/RigVM/Source/RigVM/Private/RigVMSettings.cpp
--- a/Engine/Plugins/Runtime/RigVM/Source/RigVM/Private/RigVMSettings.cpp
+++ b/Engine/Plugins/Runtime/RigVM/Source/RigVM/Private/RigVMSettings.cpp
@@ -10,7 +10,7 @@ URigVMEditorSettings::URigVMEditorSettings(const FObjectInitializer& ObjectIniti
     : Super(ObjectInitializer)
 #if WITH_EDITORONLY_DATA
 	, bEnableContextMenuTimeSlicing(true)
-	, bEnableFuzzySearchInActionMenu(true)
+	, bEnableFuzzySearchInActionMenu(false)
 	, MinimumFuzzyScore(25.f)
 #endif
 {
