[RigVM] Delete variable references on ref pin removal.
Prevent UAF from crashing on ref pin removal.
Fix function reference nodes do not update variable references as functions add variables, leading to compile failures

#tests unit, PIE, editor, RigVM, asset reload
#Synced-CL 51775323

diff --git a/Engine/Plugins/Experimental/UAF/UAF/Source/UAFUncookedOnly/Internal/AnimNextRigVMAssetEditorData.h b/Engine/Plugins/Experimental/UAF/UAF/Source/UAFUncookedOnly/Internal/AnimNextRigVMAssetEditorData.h
--- a/Engine/Plugins/Experimental/UAF/UAF/Source/UAFUncookedOnly/Internal/AnimNextRigVMAssetEditorData.h
+++ b/Engine/Plugins/Experimental/UAF/UAF/Source/UAFUncookedOnly/Internal/AnimNextRigVMAssetEditorData.h
@@ -610,6 +610,9 @@ protected:
 	// Refresh the 'external' models for the RigVM client to reference
 	UE_API void RefreshExternalModels();
 
+	// Update referenced function impl graph nodes to avoid compile issues. Doesn't recompile, only updates graph nodes.
+	UE_API void RefreshFunctionImplementationGraphs();
+
 	// Clear the error info for all EdGraphNodes
 	UE_API void ClearErrorInfoForAllEdGraphs();
 
diff --git a/Engine/Plugins/Experimental/UAF/UAF/Source/UAFUncookedOnly/Private/AnimNextEdGraphSchema.cpp b/Engine/Plugins/Experimental/UAF/UAF/Source/UAFUncookedOnly/Private/AnimNextEdGraphSchema.cpp
--- a/Engine/Plugins/Experimental/UAF/UAF/Source/UAFUncookedOnly/Private/AnimNextEdGraphSchema.cpp
+++ b/Engine/Plugins/Experimental/UAF/UAF/Source/UAFUncookedOnly/Private/AnimNextEdGraphSchema.cpp
@@ -17,7 +17,7 @@ bool UAnimNextEdGraphSchema::ShouldHidePinOnNode(const URigVMPin* InModelPin) co
 	{
 		if (const URigVMPin* RootPin = InModelPin->GetRootPin())
 		{
-			if (RootPin->IsDefinedAsInputVariable() && !RootPin->GetBoundVariableName().IsEmpty())
+			if (RootPin->IsDefinedAsInputVariable() && RootPin->GetTypedOuter<UUAFRigVMAsset>())
 			{
 				return true;
 			}
diff --git a/Engine/Plugins/Experimental/UAF/UAF/Source/UAFUncookedOnly/Private/AnimNextRigVMAssetEditorData.cpp b/Engine/Plugins/Experimental/UAF/UAF/Source/UAFUncookedOnly/Private/AnimNextRigVMAssetEditorData.cpp
--- a/Engine/Plugins/Experimental/UAF/UAF/Source/UAFUncookedOnly/Private/AnimNextRigVMAssetEditorData.cpp
+++ b/Engine/Plugins/Experimental/UAF/UAF/Source/UAFUncookedOnly/Private/AnimNextRigVMAssetEditorData.cpp
@@ -1590,6 +1590,9 @@ void UUAFRigVMAssetEditorData::HandleModifiedEvent(ERigVMGraphNotifType InNotifT
 		}
 	case ERigVMGraphNotifType::PinAdded:
 		{
+			// Rebind function variable pins. Needed when a function impl has new referenced vars as it's new var arg pin will be unbound.
+			RefreshFunctionImplementationGraphs();
+
 			if (URigVMPin* Pin = Cast<URigVMPin>(InSubject))
 			{
 				if (Pin->IsTraitPin())
@@ -3049,6 +3052,65 @@ void UUAFRigVMAssetEditorData::RefreshExternalModels()
 	}
 }
 
+
+void UUAFRigVMAssetEditorData::RefreshFunctionImplementationGraphs()
+{
+	FRigVMClient* VMClient = GetRigVMClient();
+	TArray<URigVMGraph*> AllModels = VMClient->GetAllModels(false, false);
+
+	for (URigVMGraph* Model : AllModels)
+	{
+		URigVMController* Controller = RigVMClient.GetOrCreateController(Model);
+		if (!Controller)
+		{
+			continue;
+		}
+
+		const bool bSetupUndoRedo = false;
+
+		// Capture current nodes in separate array as binding can change node count.
+		// No new function nodes should be introduced, so this is safe
+		TArray<URigVMNode*> Nodes = Model->GetNodes();
+		for (URigVMNode* Node : Nodes)
+		{
+			if (URigVMFunctionReferenceNode* ReferenceNode = Cast<URigVMFunctionReferenceNode>(Node))
+			{
+				// Remove outdated pins & related var nodes
+				{
+					// Pins will change as we unbind, cache.
+					TArray<URigVMPin*> OrphanPins = ReferenceNode->GetOrphanedPins();
+					TArray<URigVMNode*> OrphanVars;
+					OrphanVars.Reserve(OrphanPins.Num());
+
+					for (URigVMPin* OrphanPin : OrphanPins)
+					{
+						if (!OrphanPin->IsDefinedAsInputVariable())
+						{
+							continue;
+						}
+						OrphanVars.Add(OrphanPin->GetBoundVariableNode());
+						Controller->UnbindPinFromVariable(OrphanPin->GetPinPath(), bSetupUndoRedo);
+					}
+					Controller->RemoveNodes(OrphanVars, bSetupUndoRedo);
+					Controller->RemoveUnusedOrphanedPins(ReferenceNode);
+				}
+
+				// Bind new pins to variable nodes
+				{
+					for (URigVMPin* Pin : ReferenceNode->GetPins())
+					{
+						if (!Pin->IsDefinedAsInputVariable())
+						{
+							continue;
+						}
+						Controller->BindPinToVariable(Pin->GetPinPath(), Pin->GetName(), bSetupUndoRedo);
+					}
+				}
+			}
+		}
+	}
+}
+
 void UUAFRigVMAssetEditorData::AddReferencedObjects(UObject* InThis, FReferenceCollector& Collector)
 {
 	Super::AddReferencedObjects(InThis, Collector);
@@ -3364,6 +3426,8 @@ void UUAFRigVMAssetEditorData::ConditionalPatchFunctionsOnLoad()
 		TMap<URigVMLibraryNode*, FRigVMGraphFunctionHeader> OldHeaders;
 		TArray<FName> BackwardsCompatiblePublicFunctions;
 		GetRigVMClient()->PatchFunctionsOnLoad(this, BackwardsCompatiblePublicFunctions, OldHeaders);
+
+		RefreshFunctionImplementationGraphs();
 	}
 }
 
diff --git a/Engine/Plugins/Runtime/RigVM/Source/RigVMDeveloper/Private/RigVMModel/RigVMController.cpp b/Engine/Plugins/Runtime/RigVM/Source/RigVMDeveloper/Private/RigVMModel/RigVMController.cpp
--- a/Engine/Plugins/Runtime/RigVM/Source/RigVMDeveloper/Private/RigVMModel/RigVMController.cpp
+++ b/Engine/Plugins/Runtime/RigVM/Source/RigVMDeveloper/Private/RigVMModel/RigVMController.cpp
@@ -7157,7 +7157,11 @@ bool URigVMController::RemoveNodes(TArray<URigVMNode*> InNodes, bool bSetupUndoR
 
 			if (!EjectNodeFromPin(Pin->GetPinPath(), bSetupUndoRedo))
 			{
-				GetActionStack()->CancelAction(Action);
+				if (bSetupUndoRedo)
+				{
+					GetActionStack()->CancelAction(Action);
+				}
+					
 				return false;
 			}
 			
@@ -12146,6 +12150,11 @@ bool URigVMController::RemovePin(URigVMPin* InPinToRemove, bool bSetupUndoRedo,
 	check(!(InPinToRemove->IsRootPin() && InPinToRemove->IsTraitPin()));
 
 	FRigVMControllerCompileBracketScope CompileScope(this);
+	if (InPinToRemove->IsDefinedAsInputVariable())
+	{
+		UnbindPinFromVariable(InPinToRemove, bSetupUndoRedo);
+	}
+
 	if (bSetupUndoRedo || bForceBreakLinks)
 	{
 		BreakAllLinks(InPinToRemove, true, bSetupUndoRedo);
@@ -13925,15 +13934,14 @@ bool URigVMController::RemoveExposedPin(const FName& InPinName, bool bSetupUndoR
 				{
 					if (VariableNode->GetVariableGuid() == Guid)
 					{
-						NodesToRemove.Add(VariableNode);
+						NodesToRemove.AddUnique(VariableNode);
 					}
 				}
-				else
+
+				// Guid accuracy relies on PatchPinGuidInCollapseNodes, but some RigVM users don't patch. So allow for variable name fallback.
+				if (VariableNode->GetVariableName() == InPinName)
 				{
-					if (VariableNode->GetVariableName() == InPinName)
-					{
-						NodesToRemove.Add(VariableNode);
-					}
+					NodesToRemove.AddUnique(VariableNode);
 				}
 			}
 		}
@@ -24920,11 +24928,36 @@ FRigVMClientPatchResult URigVMController::PatchPinGuidInCollapseNodes()
 		{
 			if (URigVMCollapseNode* CollapseNode = Cast<URigVMCollapseNode>(Node))
 			{
-				if (CollapseNode->PinNameToGuid.IsEmpty())
+				for (URigVMPin* Pin : CollapseNode->GetPins())
 				{
-					for (URigVMPin* Pin : CollapseNode->GetPins())
+					if (FGuid* PinGuid = CollapseNode->PinNameToGuid.Find(Pin->GetFName()))
 					{
-						CollapseNode->PinNameToGuid.FindOrAdd(Pin->GetFName()) = FGuid::NewDeterministicGuid(Pin->GetPathName());
+						if (Pin->IsDefinedAsInputVariable())
+						{
+							FRigVMExternalVariable ExternalVariable = GetVariableByName(Pin->GetFName(), true);
+
+							if (ExternalVariable.GetGuid() != *PinGuid)
+							{
+								*PinGuid = ExternalVariable.GetGuid();
+								Result.bChangedContent = true;
+							}
+						}
+						else
+						{
+							// Path based GUID should always be correct. Do nothing
+						}
+					}
+					else
+					{
+						if (Pin->IsDefinedAsInputVariable())
+						{
+							FRigVMExternalVariable ExternalVariable = GetVariableByName(Pin->GetFName(), true);
+							CollapseNode->PinNameToGuid.FindOrAdd(Pin->GetFName()) = ExternalVariable.GetGuid();
+						}
+						else
+						{
+							CollapseNode->PinNameToGuid.FindOrAdd(Pin->GetFName()) = FGuid::NewDeterministicGuid(Pin->GetPathName());
+						}
 						Result.bChangedContent = true;
 					}
 				}
