Control Rig: Pure functions
#rnx

#Synced-CL 51775216

diff --git a/Engine/Plugins/Runtime/RigVM/Source/RigVM/Public/RigVMCore/RigVMStruct.h b/Engine/Plugins/Runtime/RigVM/Source/RigVM/Public/RigVMCore/RigVMStruct.h
--- a/Engine/Plugins/Runtime/RigVM/Source/RigVM/Public/RigVMCore/RigVMStruct.h
+++ b/Engine/Plugins/Runtime/RigVM/Source/RigVM/Public/RigVMCore/RigVMStruct.h
@@ -285,6 +285,7 @@ public:
 	static inline const FLazyName ExpandPinByDefaultMetaName = FLazyName(TEXT("ExpandByDefault"));
 	static inline const FLazyName DefaultArraySizeMetaName = FLazyName(TEXT("DefaultArraySize"));
 	static inline const FLazyName VaryingMetaName = FLazyName(TEXT("Varying"));
+	static inline const FLazyName PureMetaName = FLazyName(TEXT("Pure")); // The node will not modify the context
 	static inline const FLazyName SingletonMetaName = FLazyName(TEXT("Singleton"));
 	static inline const FLazyName SliceContextMetaName = FLazyName(TEXT("SliceContext"));
 	static inline const FLazyName ExecuteName = FLazyName(TEXT("Execute"));
diff --git a/Engine/Plugins/Runtime/RigVM/Source/RigVM/Public/RigVMFunctions/RigVMDispatch_Array.h b/Engine/Plugins/Runtime/RigVM/Source/RigVM/Public/RigVMFunctions/RigVMDispatch_Array.h
--- a/Engine/Plugins/Runtime/RigVM/Source/RigVM/Public/RigVMFunctions/RigVMDispatch_Array.h
+++ b/Engine/Plugins/Runtime/RigVM/Source/RigVM/Public/RigVMFunctions/RigVMDispatch_Array.h
@@ -76,7 +76,7 @@ protected:
 	static RIGVM_API void Execute(FRigVMExtendedExecuteContext& InContext, FRigVMMemoryHandleArray Handles, FRigVMPredicateBranchArray Predicates);
 };
 
-USTRUCT(meta=(DisplayName = "Reset"))
+USTRUCT(meta=(DisplayName = "Reset", Pure))
 struct FRigVMDispatch_ArrayReset : public FRigVMDispatch_ArrayBaseMutable
 {
 	GENERATED_BODY()
@@ -126,7 +126,7 @@ protected:
 	static RIGVM_API void Execute(FRigVMExtendedExecuteContext& InContext, FRigVMMemoryHandleArray Handles, FRigVMPredicateBranchArray Predicates);
 };
 
-USTRUCT(meta=(DisplayName = "Set Num"))
+USTRUCT(meta=(DisplayName = "Set Num", Pure))
 struct FRigVMDispatch_ArraySetNum : public FRigVMDispatch_ArrayBaseMutable
 {
 	GENERATED_BODY()
@@ -175,7 +175,7 @@ protected:
 	static RIGVM_API void Execute(FRigVMExtendedExecuteContext& InContext, FRigVMMemoryHandleArray Handles, FRigVMPredicateBranchArray Predicates);
 };
 
-USTRUCT(meta=(DisplayName = "Init"))
+USTRUCT(meta=(DisplayName = "Init", Pure))
 struct FRigVMDispatch_ArrayInit : public FRigVMDispatch_ArrayBaseMutable
 {
 	GENERATED_BODY()
@@ -224,7 +224,7 @@ protected:
 	static RIGVM_API void Execute(FRigVMExtendedExecuteContext& InContext, FRigVMMemoryHandleArray Handles, FRigVMPredicateBranchArray Predicates);
 };
 
-USTRUCT(meta=(DisplayName = "Set At"))
+USTRUCT(meta=(DisplayName = "Set At", Pure))
 struct FRigVMDispatch_ArraySetAtIndex : public FRigVMDispatch_ArrayBaseMutable
 {
 	GENERATED_BODY()
@@ -249,7 +249,7 @@ protected:
 	static RIGVM_API void Execute(FRigVMExtendedExecuteContext& InContext, FRigVMMemoryHandleArray Handles, FRigVMPredicateBranchArray Predicates);
 };
 
-USTRUCT(meta=(DisplayName = "Add"))
+USTRUCT(meta=(DisplayName = "Add", Pure))
 struct FRigVMDispatch_ArrayAdd : public FRigVMDispatch_ArraySetAtIndex
 {
 	GENERATED_BODY()
@@ -273,7 +273,7 @@ protected:
 	static RIGVM_API void Execute(FRigVMExtendedExecuteContext& InContext, FRigVMMemoryHandleArray Handles, FRigVMPredicateBranchArray Predicates);
 };
 
-USTRUCT(meta=(DisplayName = "Insert"))
+USTRUCT(meta=(DisplayName = "Insert", Pure))
 struct FRigVMDispatch_ArrayInsert : public FRigVMDispatch_ArraySetAtIndex
 {
 	GENERATED_BODY()
@@ -295,7 +295,7 @@ protected:
 	static RIGVM_API void Execute(FRigVMExtendedExecuteContext& InContext, FRigVMMemoryHandleArray Handles, FRigVMPredicateBranchArray Predicates);
 };
 
-USTRUCT(meta=(DisplayName = "Remove"))
+USTRUCT(meta=(DisplayName = "Remove", Pure))
 struct FRigVMDispatch_ArrayRemove : public FRigVMDispatch_ArrayBaseMutable
 {
 	GENERATED_BODY()
@@ -320,7 +320,7 @@ protected:
 	static RIGVM_API void Execute(FRigVMExtendedExecuteContext& InContext, FRigVMMemoryHandleArray Handles, FRigVMPredicateBranchArray Predicates);
 };
 
-USTRUCT(meta=(DisplayName = "Reverse"))
+USTRUCT(meta=(DisplayName = "Reverse", Pure))
 struct FRigVMDispatch_ArrayReverse : public FRigVMDispatch_ArrayReset
 {
 	GENERATED_BODY()
@@ -368,7 +368,7 @@ protected:
 	static RIGVM_API void Execute(FRigVMExtendedExecuteContext& InContext, FRigVMMemoryHandleArray Handles, FRigVMPredicateBranchArray Predicates);
 };
 
-USTRUCT(meta=(DisplayName = "Append"))
+USTRUCT(meta=(DisplayName = "Append", Pure))
 struct FRigVMDispatch_ArrayAppend : public FRigVMDispatch_ArrayBaseMutable
 {
 	GENERATED_BODY()
@@ -418,7 +418,7 @@ protected:
 	static RIGVM_API void Execute(FRigVMExtendedExecuteContext& InContext, FRigVMMemoryHandleArray Handles, FRigVMPredicateBranchArray Predicates);
 };
 
-USTRUCT(meta=(DisplayName = "Union"))
+USTRUCT(meta=(DisplayName = "Union", Pure))
 struct FRigVMDispatch_ArrayUnion : public FRigVMDispatch_ArrayAppend
 {
 	GENERATED_BODY()
@@ -487,7 +487,7 @@ protected:
 	static RIGVM_API void Execute(FRigVMExtendedExecuteContext& InContext, FRigVMMemoryHandleArray Handles, FRigVMPredicateBranchArray Predicates);
 };
 
-USTRUCT(meta=(DisplayName = "For Each", Icon="EditorStyle|GraphEditor.Macro.ForEach_16x"))
+USTRUCT(meta=(DisplayName = "For Each", Icon="EditorStyle|GraphEditor.Macro.ForEach_16x", Pure))
 struct FRigVMDispatch_ArrayIterator : public FRigVMDispatch_ArrayBaseMutable
 {
 	GENERATED_BODY()
diff --git a/Engine/Plugins/Runtime/RigVM/Source/RigVMDeveloper/Private/RigVMCompiler/RigVMCompiler.cpp b/Engine/Plugins/Runtime/RigVM/Source/RigVMDeveloper/Private/RigVMCompiler/RigVMCompiler.cpp
--- a/Engine/Plugins/Runtime/RigVM/Source/RigVMDeveloper/Private/RigVMCompiler/RigVMCompiler.cpp
+++ b/Engine/Plugins/Runtime/RigVM/Source/RigVMDeveloper/Private/RigVMCompiler/RigVMCompiler.cpp
@@ -1484,6 +1484,98 @@ bool URigVMCompiler::Compile(const FRigVMCompileSettings& InSettings, TArray<URi
 				}
 			}
 
+			// Pure function validation
+			if (URigVMCollapseNode* OwningFunction = Cast<URigVMCollapseNode>(ModelNode->FindFunctionForNode()))
+			{
+				if (OwningFunction->IsPure())
+				{
+					if (!ModelNode->IsA<URigVMFunctionInterfaceNode>())
+					{
+						// Reject mutable nodes without Pure metadata
+						// Mutable nodes have execute pins and modify the execution context.
+						// They need explicit Pure metadata to be allowed in pure functions.
+						// Note: Varying/Constant metadata is about whether the output will change on each call,
+						// not about side effects, so it is not relevant here.
+						if (ModelNode->IsMutable() && !ModelNode->IsPure())
+						{
+							static const FString PureViolationMessage = TEXT("Mutable node @@ cannot be used in pure function. Mutable nodes must have 'Pure' metadata to be allowed in pure functions.");
+							Settings.ASTSettings.Report(EMessageSeverity::Error, ModelNode, PureViolationMessage);
+							bEncounteredGraphError = true;
+						}
+
+						// Reject calls to mutable (non-pure) functions
+						if (const URigVMFunctionReferenceNode* FuncRefNode = Cast<URigVMFunctionReferenceNode>(ModelNode))
+						{
+							if (FuncRefNode->GetReferencedFunctionHeader().IsMutable())
+							{
+								static const FString MutableCallMessage = TEXT("Pure function cannot call mutable function @@.");
+								Settings.ASTSettings.Report(EMessageSeverity::Error, ModelNode, MutableCallMessage);
+								bEncounteredGraphError = true;
+							}
+						}
+
+						// Reject setting non-local variables (only local variable setters allowed)
+						if (const URigVMVariableNode* VarNode = Cast<URigVMVariableNode>(ModelNode))
+						{
+							if (!VarNode->IsGetter() && !VarNode->IsLocalVariable())
+							{
+								static const FString NonLocalVarMessage = TEXT("Pure function cannot set non-local variable @@.");
+								Settings.ASTSettings.Report(EMessageSeverity::Error, ModelNode, NonLocalVarMessage);
+								bEncounteredGraphError = true;
+							}
+						}
+
+						// Reject non-local sources connected to IO pins.
+						// Only local variable getters can connect to IO pins in pure functions.
+						for (const URigVMPin* Pin : ModelNode->GetPins())
+						{
+							if (Pin->GetDirection() != ERigVMPinDirection::IO)
+							{
+								continue;
+							}
+
+							for (const URigVMPin* SourcePin : Pin->GetLinkedSourcePins())
+							{
+								if (!SourcePin)
+								{
+									continue;
+								}
+								
+								if (SourcePin->IsExecuteContext())
+								{
+									continue;
+								}
+
+								// Local variable getters are allowed, otherwise reject
+								bool bIsAllowed = true;
+								if (const URigVMVariableNode* SourceVarNode = Cast<URigVMVariableNode>(SourcePin->GetNode()))
+								{
+									if (!SourceVarNode->IsLocalVariable())
+									{
+										bIsAllowed = false;
+									}
+								}
+								
+								// No connections from entry node
+								if (SourcePin->GetNode()->IsA<URigVMFunctionEntryNode>())
+								{
+									bIsAllowed = false;
+								}
+
+								if (!bIsAllowed)
+								{
+									const FString IOPinMessage = FString::Printf(
+										TEXT("Only local variables can connect to IO pin '%s' in pure function. Non-local data cannot be modified."),
+										*Pin->GetPinPath());
+									Settings.ASTSettings.Report(EMessageSeverity::Error, ModelNode, IOPinMessage);
+									bEncounteredGraphError = true;
+								}
+							}
+						}
+					}
+				}
+			}
+
 			auto ReportIncompatibleExecuteContextString = [&] (const FString InExecuteContextName)
 			{
 				static constexpr TCHAR Format[] = TEXT("ExecuteContext '%s' on node '%s' is not compatible with '%s' provided by the compiler settings."); 
diff --git a/Engine/Plugins/Runtime/RigVM/Source/RigVMDeveloper/Private/RigVMModel/Nodes/RigVMCollapseNode.cpp b/Engine/Plugins/Runtime/RigVM/Source/RigVMDeveloper/Private/RigVMModel/Nodes/RigVMCollapseNode.cpp
--- a/Engine/Plugins/Runtime/RigVM/Source/RigVMDeveloper/Private/RigVMModel/Nodes/RigVMCollapseNode.cpp
+++ b/Engine/Plugins/Runtime/RigVM/Source/RigVMDeveloper/Private/RigVMModel/Nodes/RigVMCollapseNode.cpp
@@ -26,6 +26,15 @@ URigVMFunctionLibrary* URigVMCollapseNode::GetLibrary() const
 	return Cast<URigVMFunctionLibrary>(GetOuter());
 }
 
+bool URigVMCollapseNode::IsPure() const
+{
+	if (IsGraphFunctionDefinition())
+	{
+		return FindExecutePin() == nullptr;
+	}
+	return false;
+}
+
 TArray<FRigVMExternalVariable> URigVMCollapseNode::GetExternalVariables() const
 {
 	TArray<FRigVMExternalVariable> ExternalVariables;
diff --git a/Engine/Plugins/Runtime/RigVM/Source/RigVMDeveloper/Private/RigVMModel/Nodes/RigVMDispatchNode.cpp b/Engine/Plugins/Runtime/RigVM/Source/RigVMDeveloper/Private/RigVMModel/Nodes/RigVMDispatchNode.cpp
--- a/Engine/Plugins/Runtime/RigVM/Source/RigVMDeveloper/Private/RigVMModel/Nodes/RigVMDispatchNode.cpp
+++ b/Engine/Plugins/Runtime/RigVM/Source/RigVMDeveloper/Private/RigVMModel/Nodes/RigVMDispatchNode.cpp
@@ -63,6 +63,25 @@ bool URigVMDispatchNode::IsDefinedAsVarying() const
 	return false;
 }
 
+bool URigVMDispatchNode::IsPure() const
+{
+	if (const UScriptStruct* Struct = GetFactoryStruct())
+	{
+		// Pure metadata explicitly marks dispatch nodes as pure, overriding hidden pin checks.
+		// This allows nodes with internal state (e.g., loop counters in control flow)
+		// to be marked pure as long as they have no external side effects.
+		if (Struct->HasMetaData(FRigVMStruct::PureMetaName))
+		{
+			return true;
+		}
+	}
+
+	// For nodes without Pure metadata, fall back to base class logic which checks:
+	// 1. Not mutable (no execute context pins)
+	// 2. No hidden pins (no internal state)
+	return Super::IsPure();
+}
+
 const TArray<FName>& URigVMDispatchNode::GetControlFlowBlocks() const
 {
 	if (const FRigVMDispatchFactory* Factory = GetFactory())
diff --git a/Engine/Plugins/Runtime/RigVM/Source/RigVMDeveloper/Private/RigVMModel/Nodes/RigVMUnitNode.cpp b/Engine/Plugins/Runtime/RigVM/Source/RigVMDeveloper/Private/RigVMModel/Nodes/RigVMUnitNode.cpp
--- a/Engine/Plugins/Runtime/RigVM/Source/RigVMDeveloper/Private/RigVMModel/Nodes/RigVMUnitNode.cpp
+++ b/Engine/Plugins/Runtime/RigVM/Source/RigVMDeveloper/Private/RigVMModel/Nodes/RigVMUnitNode.cpp
@@ -88,6 +88,25 @@ bool URigVMUnitNode::IsDefinedAsVarying() const
 	return false;
 }
 
+bool URigVMUnitNode::IsPure() const
+{
+	if (UScriptStruct* Struct = GetScriptStruct())
+	{
+		// Pure metadata explicitly marks nodes as pure, overriding hidden pin checks.
+		// This allows nodes with internal state (e.g., loop counters in control flow)
+		// to be marked pure as long as they have no external side effects.
+		if (Struct->HasMetaData(FRigVMStruct::PureMetaName))
+		{
+			return true;
+		}
+	}
+
+	// For nodes without Pure metadata, fall back to base class logic which checks:
+	// 1. Not mutable (no execute context pins)
+	// 2. No hidden pins (no internal state)
+	return Super::IsPure();
+}
+
 FName URigVMUnitNode::GetEventName() const
 {
 	if(!CachedEventName.IsBound())
diff --git a/Engine/Plugins/Runtime/RigVM/Source/RigVMDeveloper/Private/RigVMModel/Nodes/RigVMVariableNode.cpp b/Engine/Plugins/Runtime/RigVM/Source/RigVMDeveloper/Private/RigVMModel/Nodes/RigVMVariableNode.cpp
--- a/Engine/Plugins/Runtime/RigVM/Source/RigVMDeveloper/Private/RigVMModel/Nodes/RigVMVariableNode.cpp
+++ b/Engine/Plugins/Runtime/RigVM/Source/RigVMDeveloper/Private/RigVMModel/Nodes/RigVMVariableNode.cpp
@@ -141,6 +141,22 @@ FRigVMGraphVariableDescription URigVMVariableNode::GetVariableDescription() cons
 	return Variable;
 }
 
+bool URigVMVariableNode::IsPure() const
+{
+	if (IsGetter())
+	{
+		return true;
+	}
+	
+	// If its a local variable return true
+	if (IsLocalVariable())
+	{
+		return true;
+	}
+	
+	return false;
+}
+
 URigVMPin* URigVMVariableNode::GetVariableNamePin() const
 {
 	return FindPin(VariableName);
diff --git a/Engine/Plugins/Runtime/RigVM/Source/RigVMDeveloper/Private/RigVMModel/RigVMController.cpp b/Engine/Plugins/Runtime/RigVM/Source/RigVMDeveloper/Private/RigVMModel/RigVMController.cpp
--- a/Engine/Plugins/Runtime/RigVM/Source/RigVMDeveloper/Private/RigVMModel/RigVMController.cpp
+++ b/Engine/Plugins/Runtime/RigVM/Source/RigVMDeveloper/Private/RigVMModel/RigVMController.cpp
@@ -302,6 +302,7 @@ URigVMController* URigVMController::GetControllerForGraph(const URigVMGraph* InG
 			if(bIsRunningUnitTest)
 			{
 				Controller->bIsRunningUnitTest = true;
+				Controller->bReportWarningsAndErrors = bReportWarningsAndErrors;
 				if(!Controller->GetExternalVariablesDelegate.IsBound() && !InGraph->IsRootGraph())
 				{
 					// get the controller for the root graph
@@ -6683,7 +6684,7 @@ URigVMFunctionReferenceNode* URigVMController::PromoteCollapseNodeToFunctionRefe
 			{
 				for(const URigVMPin* Pin : InCollapseNode->GetPins())
 				{
-					DefinitionController->AddExposedPin(Pin->GetFName(), Pin->GetDirection(), Pin->GetCPPType(), (Pin->GetCPPTypeObject() ? *Pin->GetCPPTypeObject()->GetPathName() : TEXT("")), Pin->GetDefaultValue(), false);
+					DefinitionController->AddExposedPin_Internal(Pin->GetFName(), Pin->GetDirection(), Pin->GetCPPType(), (Pin->GetCPPTypeObject() ? *Pin->GetCPPTypeObject()->GetPathName() : TEXT("")), Pin->GetDefaultValue(), false, /*bPrintPythonCommand=*/ false, /*bIsInputVariable=*/ false, /*bAllowExecutePinManagement=*/ true);
 				}
 			}
 		}
@@ -13617,6 +13618,11 @@ bool URigVMController::BreakAllLinksRecursive(URigVMPin* Pin, bool bAsInput, boo
 }
 
 FName URigVMController::AddExposedPin(const FName& InPinName, ERigVMPinDirection InDirection, const FString& InCPPType, const FName& InCPPTypeObjectPath, const FString& InDefaultValue, bool bSetupUndoRedo, bool bPrintPythonCommand, bool bIsInputVariable)
+{
+	return AddExposedPin_Internal(InPinName, InDirection, InCPPType, InCPPTypeObjectPath, InDefaultValue, bSetupUndoRedo, bPrintPythonCommand, bIsInputVariable, /*bAllowExecutePinManagement=*/ false);
+}
+
+FName URigVMController::AddExposedPin_Internal(const FName& InPinName, ERigVMPinDirection InDirection, const FString& InCPPType, const FName& InCPPTypeObjectPath, const FString& InDefaultValue, bool bSetupUndoRedo, bool bPrintPythonCommand, bool bIsInputVariable, bool bAllowExecutePinManagement)
 {
 	if (!IsValidGraph())
 	{
@@ -13683,6 +13689,18 @@ FName URigVMController::AddExposedPin(const FName& InPinName, ERigVMPinDirection
 		{
 			if(CPPTypeStruct->IsChildOf(FRigVMExecutePin::StaticStruct()))
 			{
+				if (!bAllowExecutePinManagement)
+				{
+					if (const URigVMCollapseNode* CollapseNode = Cast<URigVMCollapseNode>(LibraryNode))
+					{
+						if (CollapseNode->IsGraphFunctionDefinition())
+						{
+							ReportError(TEXT("Execute pins cannot be manually added to functions. Use SetFunctionIsPure() to change mutability."));
+							return NAME_None;
+						}
+					}
+				}
+
 				for(const URigVMPin* ExistingPin : LibraryNode->Pins)
 				{
 					if(ExistingPin->IsExecuteContext())
@@ -13829,6 +13847,11 @@ FName URigVMController::AddExposedPin(const FName& InPinName, ERigVMPinDirection
 }
 
 bool URigVMController::RemoveExposedPin(const FName& InPinName, bool bSetupUndoRedo, bool bPrintPythonCommand)
+{
+	return RemoveExposedPin_Internal(InPinName, bSetupUndoRedo, bPrintPythonCommand, /*bAllowExecutePinManagement=*/ false);
+}
+
+bool URigVMController::RemoveExposedPin_Internal(const FName& InPinName, bool bSetupUndoRedo, bool bPrintPythonCommand, bool bAllowExecutePinManagement)
 {
 	if (!IsValidGraph())
 	{
@@ -13864,6 +13887,18 @@ bool URigVMController::RemoveExposedPin(const FName& InPinName, bool bSetupUndoR
 		return false;
 	}
 
+	if (Pin->IsExecuteContext() && !bAllowExecutePinManagement)
+	{
+		if (const URigVMCollapseNode* CollapseNode = Cast<URigVMCollapseNode>(LibraryNode))
+		{
+			if (CollapseNode->IsGraphFunctionDefinition())
+			{
+				ReportError(TEXT("Execute pins cannot be manually removed from functions. Use SetFunctionIsPure() to change mutability."));
+				return false;
+			}
+		}
+	}
+
 	if (URigVMAggregateNode* AggregateNode = Cast<URigVMAggregateNode>(LibraryNode))
 	{
 		if (AggregateNode->IsInputAggregate())
@@ -14181,6 +14216,11 @@ bool URigVMController::RenameExposedPin(const FName& InOldPinName, const FName&
 }
 
 bool URigVMController::ChangeExposedPinType(const FName& InPinName, const FString& InCPPType, const FName& InCPPTypeObjectPath, bool& bSetupUndoRedo, bool bSetupOrphanPins, bool bPrintPythonCommand)
+{
+	return ChangeExposedPinType_Internal(InPinName, InCPPType, InCPPTypeObjectPath, bSetupUndoRedo, bSetupOrphanPins, bPrintPythonCommand, /*bAllowExecutePinManagement=*/ false);
+}
+
+bool URigVMController::ChangeExposedPinType_Internal(const FName& InPinName, const FString& InCPPType, const FName& InCPPTypeObjectPath, bool& bSetupUndoRedo, bool bSetupOrphanPins, bool bPrintPythonCommand, bool bAllowExecutePinManagement)
 {
 	if (!IsValidGraph())
 	{
@@ -14216,6 +14256,22 @@ bool URigVMController::ChangeExposedPinType(const FName& InPinName, const FStrin
 		return false;
 	}
 
+	// Guard: cannot change existing execute pin type on function definitions
+	if (!bAllowExecutePinManagement)
+	{
+		if (Pin->IsExecuteContext())
+		{
+			if (const URigVMCollapseNode* CollapseNode = Cast<URigVMCollapseNode>(LibraryNode))
+			{
+				if (CollapseNode->IsGraphFunctionDefinition())
+				{
+					ReportError(TEXT("Cannot change execute pin type on functions. Use SetFunctionIsPure() to change mutability."));
+					return false;
+				}
+			}
+		}
+	}
+
 	// We do not allow unresolving exposed pins
 	if (InCPPType == RigVMTypeUtils::GetWildCardCPPType())
 	{
@@ -14234,6 +14290,18 @@ bool URigVMController::ChangeExposedPinType(const FName& InPinName, const FStrin
 			{
 				if(CPPTypeStruct->IsChildOf(FRigVMExecutePin::StaticStruct()))
 				{
+					if (!bAllowExecutePinManagement)
+					{
+						if (const URigVMCollapseNode* CollapseNode = Cast<URigVMCollapseNode>(LibraryNode))
+						{
+							if (CollapseNode->IsGraphFunctionDefinition())
+							{
+								ReportError(TEXT("Cannot change pin type to execute context on functions. Use SetFunctionIsPure() to change mutability."));
+								return false;
+							}
+						}
+					}
+
 					for(URigVMPin* ExistingPin : LibraryNode->Pins)
 					{
 						if(ExistingPin != Pin)
@@ -15077,12 +15145,15 @@ URigVMLibraryNode* URigVMController::AddFunctionToLibrary(const FName& InFunctio
 
 		if(URigVMController* CollapseController = GetControllerForGraph(CollapseNode->GetContainedGraph()))
 		{
-			CollapseController->AddExposedPin(FRigVMStruct::ExecuteContextName,
+			CollapseController->AddExposedPin_Internal(FRigVMStruct::ExecuteContextName,
 				ERigVMPinDirection::IO,
 				FString::Printf(TEXT("F%s"), *ExecuteContextStruct->GetName()),
 				*ExecuteContextStruct->GetPathName(),
 				FString(),
-				false);
+				false,
+				/*bPrintPythonCommand=*/ false,
+				/*bIsInputVariable=*/ false,
+				/*bAllowExecutePinManagement=*/ true);
 		}
 	}
 
@@ -15346,6 +15417,103 @@ bool URigVMController::IsFunctionPublic(const FName& InFunctionName)
 	return false;
 }
 
+bool URigVMController::SetFunctionIsPure(const FName& InFunctionName, bool bIsPure, bool bSetupUndoRedo, bool bPrintPythonCommand)
+{
+	if (!IsValidGraph())
+	{
+		return false;
+	}
+
+	if (!bIsTransacting && !IsGraphEditable())
+	{
+		return false;
+	}
+
+	URigVMGraph* Graph = GetGraph();
+	check(Graph);
+
+	if (!Graph->IsA<URigVMFunctionLibrary>())
+	{
+		ReportError(TEXT("Can only change function purity from function library graphs."));
+		return false;
+	}
+
+	URigVMNode* Node = Graph->FindNode(InFunctionName.ToString());
+	if (!Node)
+	{
+		ReportErrorf(TEXT("Could not find function called '%s'."), *InFunctionName.ToString());
+		return false;
+	}
+
+	URigVMLibraryNode* FunctionNode = Cast<URigVMLibraryNode>(Node);
+	if (FunctionNode == nullptr)
+	{
+		ReportErrorf(TEXT("Node '%s' is not a function."), *InFunctionName.ToString());
+		return false;
+	}
+
+	// Check current state
+	const bool bCurrentlyHasExecutePin = FunctionNode->FindExecutePin() != nullptr;
+	const bool bWantExecutePin = !bIsPure;
+
+	// Already in desired state?
+	if (bCurrentlyHasExecutePin == bWantExecutePin)
+	{
+		return true;
+	}
+
+	// Get the controller for the function's contained graph to modify exposed pins
+	URigVMController* FunctionController = GetControllerForGraph(FunctionNode->GetContainedGraph());
+	if (FunctionController == nullptr)
+	{
+		ReportErrorf(TEXT("Could not get controller for function '%s'."), *InFunctionName.ToString());
+		return false;
+	}
+
+	FRigVMControllerCompileBracketScope CompileScope(FunctionController);
+
+	if (bIsPure)
+	{
+		// Remove the execute pin to make the function pure
+		if (!FunctionController->RemoveExposedPin_Internal(FRigVMStruct::ExecuteContextName, bSetupUndoRedo, bPrintPythonCommand, /*bAllowExecutePinManagement=*/ true))
+		{
+			ReportError(TEXT("Failed to remove execute pin from function."));
+			return false;
+		}
+	}
+	else
+	{
+		// Add the execute pin to make the function mutable
+		const UScriptStruct* ExecuteContextStruct = FRigVMExecuteContext::StaticStruct();
+		const FName PinName = FunctionController->AddExposedPin_Internal(
+			FRigVMStruct::ExecuteContextName,
+			ERigVMPinDirection::IO,
+			FString::Printf(TEXT("F%s"), *ExecuteContextStruct->GetName()),
+			*ExecuteContextStruct->GetPathName(),
+			FString(),
+			bSetupUndoRedo,
+			bPrintPythonCommand,
+			/*bIsInputVariable=*/ false,
+			/*bAllowExecutePinManagement=*/ true);
+
+		if (PinName.IsNone())
+		{
+			ReportError(TEXT("Failed to add execute pin to function."));
+			return false;
+		}
+	}
+
+	if (bPrintPythonCommand)
+	{
+		RigVMPythonUtils::Print(GetSchema()->GetGraphOuterName(GetGraph()),
+			FString::Printf(TEXT("library_controller.set_function_is_pure('%s', %s)"),
+				*GetSchema()->GetSanitizedNodeName(InFunctionName.ToString()),
+				(bIsPure) ? TEXT("True") : TEXT("False")));
+	}
+
+	return true;
+}
+
 URigVMLibraryNode* URigVMController::CreateFunctionVariant(const FName& InFunctionName, const FName& InVariantName, bool bSetupUndoRedo, bool bPrintPythonCommand)
 {
 	URigVMLibraryNode* Result = nullptr;
@@ -19515,11 +19683,30 @@ bool URigVMController::GenerateNewPinInfos(const FRigVMRegistry& Registry, URigV
 			TArray<URigVMPin*> SortedLibraryPins;
 
 			// add execute pins first
+			bool bHasExecutePin = false;
 			for (URigVMPin* LibraryPin : LibraryNode->GetPins())
 			{
 				if (LibraryPin->IsExecuteContext())
 				{
 					SortedLibraryPins.Add(LibraryPin);
+					bHasExecutePin = true;
+				}
+			}
+
+			// For function definitions (pure functions), entry/return nodes need execute pins
+			// for internal sequencing even if the external interface has no execute pin.
+			if (!bHasExecutePin)
+			{
+				if (URigVMCollapseNode* OuterCollapseNode = Cast<URigVMCollapseNode>(LibraryNode))
+				{
+					if (OuterCollapseNode->IsGraphFunctionDefinition())
+					{
+						// Add an execute pin for internal sequencing
+						const TRigVMTypeIndex ExecuteTypeIndex = RigVMTypeUtils::TypeIndex::Execute;
+						ERigVMPinDirection Direction = bIsEntryNode ? ERigVMPinDirection::Output : ERigVMPinDirection::Input;
+
+						(void)NewPinInfos.AddPin(this, INDEX_NONE, FRigVMStruct::ExecuteContextName, Direction, ExecuteTypeIndex, FString(), ERigVMPinDefaultValueType::Unset, nullptr, &PreviousPinInfos, false);
+					}
 				}
 			}
 
@@ -24661,8 +24848,8 @@ FRigVMClientPatchResult URigVMController::PatchFunctionsWithInvalidReturnPaths()
 			const FName ExecutePinName = EntryExecutePin->GetFName();
 			const FString ExecuteCPPType = EntryExecutePin->GetCPPType();
 			const FName ExecuteCPPTypeObjectPath = *EntryExecutePin->GetCPPTypeObject()->GetPathName();
-			RemoveExposedPin(ExecutePinName, false);
-			AddExposedPin(ExecutePinName, ERigVMPinDirection::IO, ExecuteCPPType, ExecuteCPPTypeObjectPath, FString(), false);
+			RemoveExposedPin_Internal(ExecutePinName, false, /*bPrintPythonCommand=*/ false, /*bAllowExecutePinManagement=*/ true);
+			AddExposedPin_Internal(ExecutePinName, ERigVMPinDirection::IO, ExecuteCPPType, ExecuteCPPTypeObjectPath, FString(), false, /*bPrintPythonCommand=*/ false, /*bIsInputVariable=*/ false, /*bAllowExecutePinManagement=*/ true);
 			EntryExecutePin = FindExecutePin(EntryNode);
 			check(EntryExecutePin);
 			URigVMPin* ReturnExecutePin = FindExecutePin(ReturnNode);
@@ -24693,8 +24880,8 @@ FRigVMClientPatchResult URigVMController::PatchFunctionsWithInvalidReturnPaths()
 			const FName ExecutePinName = ReturnExecutePin->GetFName();
 			const FString ExecuteCPPType = ReturnExecutePin->GetCPPType();
 			const FName ExecuteCPPTypeObjectPath = *ReturnExecutePin->GetCPPTypeObject()->GetPathName();
-			RemoveExposedPin(ExecutePinName, false);
-			AddExposedPin(ExecutePinName, ERigVMPinDirection::IO, ExecuteCPPType, ExecuteCPPTypeObjectPath, FString(), false);
+			RemoveExposedPin_Internal(ExecutePinName, false, /*bPrintPythonCommand=*/ false, /*bAllowExecutePinManagement=*/ true);
+			AddExposedPin_Internal(ExecutePinName, ERigVMPinDirection::IO, ExecuteCPPType, ExecuteCPPTypeObjectPath, FString(), false, /*bPrintPythonCommand=*/ false, /*bIsInputVariable=*/ false, /*bAllowExecutePinManagement=*/ true);
 		}
 
 		URigVMPin* EntryExecutePin = FindExecutePin(EntryNode);
diff --git a/Engine/Plugins/Runtime/RigVM/Source/RigVMDeveloper/Private/RigVMModel/RigVMPin.cpp b/Engine/Plugins/Runtime/RigVM/Source/RigVMDeveloper/Private/RigVMModel/RigVMPin.cpp
--- a/Engine/Plugins/Runtime/RigVM/Source/RigVMDeveloper/Private/RigVMModel/RigVMPin.cpp
+++ b/Engine/Plugins/Runtime/RigVM/Source/RigVMDeveloper/Private/RigVMModel/RigVMPin.cpp
@@ -2830,7 +2830,7 @@ bool URigVMPin::CanLink(const URigVMPin* InSourcePin, const URigVMPin* InTargetP
 				if (!bNodeCanLinkAnywhere)
 				{
 					// pure / immutable nodes can be connected to any input in any order.
-					// since a new link is going to change the abstract syntax tree 
+					// since a new link is going to change the abstract syntax tree
 					if (!SourceNodes[SourceNodeIndex]->IsMutable())
 					{
 						bNodeCanLinkAnywhere = true;
@@ -2866,6 +2866,47 @@ bool URigVMPin::CanLink(const URigVMPin* InSourcePin, const URigVMPin* InTargetP
 		}
 	}
 
+	// In pure function graphs, certain sources cannot connect to IO pins.
+	// Specifically: non-local variables and input parameters are rejected, as they would
+	// allow the pure function to mutate external state. Local variables, constants, and
+	// other function outputs are allowed since they are either function-scoped or immutable.
+	if (const URigVMCollapseNode* OwningFunction = Cast<URigVMCollapseNode>(SourceNode->FindFunctionForNode()))
+	{
+		if (OwningFunction->IsPure() && !InSourcePin->IsExecuteContext() && !InTargetPin->IsExecuteContext())
+		{
+			if (InTargetPin->GetDirection() == ERigVMPinDirection::IO)
+			{
+				// Local variable getters are allowed — they're function-scoped
+				if (const URigVMVariableNode* VarNode = Cast<URigVMVariableNode>(InSourcePin->GetNode()))
+				{
+					if (VarNode->IsGetter() && !VarNode->IsLocalVariable())
+					{
+						if (OutFailureReason)
+						{
+							*OutFailureReason = FString::Printf(
+								TEXT("Only local variables can connect to IO pin '%s' in a pure function. Non-local data cannot be modified."),
+								*InTargetPin->GetPinPath());
+						}
+				
+						return false;
+					}
+				}
+
+				// Entry node pins (input parameters) connected to IO pins would allow mutation of caller data
+				if (SourceNode->IsA<URigVMFunctionEntryNode>())
+				{
+					if (OutFailureReason)
+					{
+						*OutFailureReason = FString::Printf(
+							TEXT("Input parameter '%s' cannot connect to IO pin '%s' in a pure function. This would allow modification of the input."),
+							*InSourcePin->GetName(), *InTargetPin->GetPinPath());
+					}
+					return false;
+				}
+			}
+		}
+	}
+
 	return true;
 }
 
diff --git a/Engine/Plugins/Runtime/RigVM/Source/RigVMDeveloper/Public/RigVMModel/Nodes/RigVMCollapseNode.h b/Engine/Plugins/Runtime/RigVM/Source/RigVMDeveloper/Public/RigVMModel/Nodes/RigVMCollapseNode.h
--- a/Engine/Plugins/Runtime/RigVM/Source/RigVMDeveloper/Public/RigVMModel/Nodes/RigVMCollapseNode.h
+++ b/Engine/Plugins/Runtime/RigVM/Source/RigVMDeveloper/Public/RigVMModel/Nodes/RigVMCollapseNode.h
@@ -32,6 +32,7 @@ public:
 	bool IsGraphFunctionDefinition() const { return GetLibrary() != nullptr;}
 	virtual URigVMGraph* GetContainedGraph() const override { return ContainedGraph; }
 	UE_API virtual TArray<FRigVMExternalVariable> GetExternalVariables() const override;
+	UE_API virtual bool IsPure() const override;
 
 	UE_API FString GetEditorSubGraphName() const;
 
diff --git a/Engine/Plugins/Runtime/RigVM/Source/RigVMDeveloper/Public/RigVMModel/Nodes/RigVMDispatchNode.h b/Engine/Plugins/Runtime/RigVM/Source/RigVMDeveloper/Public/RigVMModel/Nodes/RigVMDispatchNode.h
--- a/Engine/Plugins/Runtime/RigVM/Source/RigVMDeveloper/Public/RigVMModel/Nodes/RigVMDispatchNode.h
+++ b/Engine/Plugins/Runtime/RigVM/Source/RigVMDeveloper/Public/RigVMModel/Nodes/RigVMDispatchNode.h
@@ -42,6 +42,7 @@ public:
 	UE_API virtual FLinearColor GetNodeColor() const override;
 	UE_API virtual bool IsDefinedAsConstant() const override;
 	UE_API virtual bool IsDefinedAsVarying() const override;
+	UE_API virtual bool IsPure() const override;
 	UE_API virtual const TArray<FName>& GetControlFlowBlocks() const override;
 	UE_API virtual const bool IsControlFlowBlockSliced(const FName& InBlockName) const override;
 	UE_API virtual TArray<URigVMPin*> GetAggregateInputs() const override;
diff --git a/Engine/Plugins/Runtime/RigVM/Source/RigVMDeveloper/Public/RigVMModel/Nodes/RigVMRerouteNode.h b/Engine/Plugins/Runtime/RigVM/Source/RigVMDeveloper/Public/RigVMModel/Nodes/RigVMRerouteNode.h
--- a/Engine/Plugins/Runtime/RigVM/Source/RigVMDeveloper/Public/RigVMModel/Nodes/RigVMRerouteNode.h
+++ b/Engine/Plugins/Runtime/RigVM/Source/RigVMDeveloper/Public/RigVMModel/Nodes/RigVMRerouteNode.h
@@ -33,6 +33,8 @@ public:
 
 	// Has no source connections
 	UE_API bool IsLiteral() const;
+	
+	virtual bool IsPure() const override { return true; }
 
 private:
 
diff --git a/Engine/Plugins/Runtime/RigVM/Source/RigVMDeveloper/Public/RigVMModel/Nodes/RigVMUnitNode.h b/Engine/Plugins/Runtime/RigVM/Source/RigVMDeveloper/Public/RigVMModel/Nodes/RigVMUnitNode.h
--- a/Engine/Plugins/Runtime/RigVM/Source/RigVMDeveloper/Public/RigVMModel/Nodes/RigVMUnitNode.h
+++ b/Engine/Plugins/Runtime/RigVM/Source/RigVMDeveloper/Public/RigVMModel/Nodes/RigVMUnitNode.h
@@ -33,6 +33,7 @@ public:
 	UE_API virtual FText GetToolTipText() const override;
 	UE_API virtual bool IsDefinedAsConstant() const override;
 	UE_API virtual bool IsDefinedAsVarying() const override;
+	UE_API virtual bool IsPure() const override;
 	UE_API virtual FName GetEventName() const override;
 	UE_API virtual bool CanOnlyExistOnce() const override;
 	UE_API virtual const TArray<FName>& GetControlFlowBlocks() const override;
diff --git a/Engine/Plugins/Runtime/RigVM/Source/RigVMDeveloper/Public/RigVMModel/Nodes/RigVMVariableNode.h b/Engine/Plugins/Runtime/RigVM/Source/RigVMDeveloper/Public/RigVMModel/Nodes/RigVMVariableNode.h
--- a/Engine/Plugins/Runtime/RigVM/Source/RigVMDeveloper/Public/RigVMModel/Nodes/RigVMVariableNode.h
+++ b/Engine/Plugins/Runtime/RigVM/Source/RigVMDeveloper/Public/RigVMModel/Nodes/RigVMVariableNode.h
@@ -70,7 +70,9 @@ public:
 
 	// Override of node title
 	virtual FLinearColor GetNodeColor() const override { return FLinearColor::Blue; }
+	
 	virtual bool IsDefinedAsVarying() const override { return true; }
+	UE_API virtual bool IsPure() const override;
 
 	// Get the variable name pin
 	UE_API URigVMPin* GetVariableNamePin() const;
diff --git a/Engine/Plugins/Runtime/RigVM/Source/RigVMDeveloper/Public/RigVMModel/RigVMController.h b/Engine/Plugins/Runtime/RigVM/Source/RigVMDeveloper/Public/RigVMModel/RigVMController.h
--- a/Engine/Plugins/Runtime/RigVM/Source/RigVMDeveloper/Public/RigVMModel/RigVMController.h
+++ b/Engine/Plugins/Runtime/RigVM/Source/RigVMDeveloper/Public/RigVMModel/RigVMController.h
@@ -1128,6 +1128,12 @@ public:
 	UFUNCTION(BlueprintCallable, Category = RigVMController)
 	UE_API bool IsFunctionPublic(const FName& InFunctionName);
 
+	// Sets whether a function is pure (no execute pins) or mutable (has execute pins).
+	// Call this on the function library controller.
+	// Returns true if the function's purity was successfully changed.
+	UFUNCTION(BlueprintCallable, Category = RigVMController)
+	UE_API bool SetFunctionIsPure(const FName& InFunctionName, bool bIsPure, bool bSetupUndoRedo = true, bool bPrintPythonCommand = false);
+
 	// Creates a variant of a function given the name of an existing function variant
 	UFUNCTION(BlueprintCallable, Category = RigVMController)
 	UE_API URigVMLibraryNode* CreateFunctionVariant(const FName& InFunctionName, const FName& InVariantName = NAME_None, bool bSetupUndoRedo = true, bool bPrintPythonCommand = false);
@@ -1348,6 +1354,10 @@ public:
 
 private:
 
+	UE_API FName AddExposedPin_Internal(const FName& InPinName, ERigVMPinDirection InDirection, const FString& InCPPType, const FName& InCPPTypeObjectPath, const FString& InDefaultValue, bool bSetupUndoRedo, bool bPrintPythonCommand, bool bIsInputVariable, bool bAllowExecutePinManagement);
+	UE_API bool RemoveExposedPin_Internal(const FName& InPinName, bool bSetupUndoRedo, bool bPrintPythonCommand, bool bAllowExecutePinManagement);
+	UE_API bool ChangeExposedPinType_Internal(const FName& InPinName, const FString& InCPPType, const FName& InCPPTypeObjectPath, bool& bSetupUndoRedo, bool bSetupOrphanPins, bool bPrintPythonCommand, bool bAllowExecutePinManagement);
+
 	UPROPERTY(BlueprintReadOnly, Category = RigVMController, meta = (ScriptName = "ModifiedEvent", AllowPrivateAccess = "true"))
 	FRigVMGraphModifiedDynamicEvent ModifiedEventDynamic;
 
diff --git a/Engine/Plugins/Runtime/RigVM/Source/RigVMDeveloper/Public/RigVMModel/RigVMNode.h b/Engine/Plugins/Runtime/RigVM/Source/RigVMDeveloper/Public/RigVMModel/RigVMNode.h
--- a/Engine/Plugins/Runtime/RigVM/Source/RigVMDeveloper/Public/RigVMModel/RigVMNode.h
+++ b/Engine/Plugins/Runtime/RigVM/Source/RigVMDeveloper/Public/RigVMModel/RigVMNode.h
@@ -231,7 +231,7 @@ public:
 	UFUNCTION(BlueprintCallable, Category = RigVMNode)
 	virtual bool IsDefinedAsConstant() const { return false; }
 
-	// Returns true if the node is defined as non-varying
+	// Returns true if the node is defined as varying
 	UFUNCTION(BlueprintCallable, Category = RigVMNode)
 	virtual bool IsDefinedAsVarying() const { return false; }
 
diff --git a/Engine/Plugins/Runtime/RigVM/Source/RigVMEditor/Private/Editor/RigVMGraphDetailCustomization.cpp b/Engine/Plugins/Runtime/RigVM/Source/RigVMEditor/Private/Editor/RigVMGraphDetailCustomization.cpp
--- a/Engine/Plugins/Runtime/RigVM/Source/RigVMEditor/Private/Editor/RigVMGraphDetailCustomization.cpp
+++ b/Engine/Plugins/Runtime/RigVM/Source/RigVMEditor/Private/Editor/RigVMGraphDetailCustomization.cpp
@@ -1345,6 +1345,22 @@ void FRigVMGraphDetailCustomizationImpl::CustomizeDetails(IDetailLayoutBuilder&
             ]
         ];
 
+		// pure function
+		SettingsCategory.AddCustomRow(LOCTEXT("Pure", "Pure"))
+		.NameContent()
+		[
+			SNew(STextBlock)
+				.Text(LOCTEXT("Pure", "Pure"))
+				.Font(IDetailLayoutBuilder::GetDetailFont())
+		]
+		.ValueContent()
+		[
+			SNew(SCheckBox)
+				.IsChecked(this, &FRigVMGraphDetailCustomizationImpl::GetFunctionPureCheckState)
+				.OnCheckStateChanged(this, &FRigVMGraphDetailCustomizationImpl::OnFunctionPureCheckStateChanged)
+				.ToolTipText(LOCTEXT("PureTooltip", "Pure functions have no execute pins and cannot have side effects."))
+		];
+
 		// variant
 		if(CVarRigVMEnableVariants.GetValueOnAnyThread())
 		{
@@ -1767,6 +1783,42 @@ TSharedRef<ITableRow> FRigVMGraphDetailCustomizationImpl::HandleGenerateRowAcces
         ];
 }
 
+ECheckBoxState FRigVMGraphDetailCustomizationImpl::GetFunctionPureCheckState() const
+{
+	if (WeakModel.IsValid() && RigVMClientHost.IsValid())
+	{
+		if (URigVMGraph* RigVMGraph = WeakModel.Get())
+		{
+			if (URigVMCollapseNode* CollapseNode = RigVMGraph->GetTypedOuter<URigVMCollapseNode>())
+			{
+				return CollapseNode->IsPure() ? ECheckBoxState::Checked : ECheckBoxState::Unchecked;
+			}
+		}
+	}
+	return ECheckBoxState::Unchecked;
+}
+
+void FRigVMGraphDetailCustomizationImpl::OnFunctionPureCheckStateChanged(ECheckBoxState NewState)
+{
+	if (WeakModel.IsValid() && RigVMClientHost.IsValid())
+	{
+		if (URigVMGraph* RigVMGraph = WeakModel.Get())
+		{
+			if (URigVMLibraryNode* LibraryNode = RigVMGraph->GetTypedOuter<URigVMLibraryNode>())
+			{
+				if (URigVMGraph* FunctionLibrary = RigVMClientHost->GetLocalFunctionLibrary())
+				{
+					if (URigVMController* FunctionLibraryController = RigVMClientHost->GetOrCreateController(FunctionLibrary))
+					{
+						const bool bIsPure = (NewState == ECheckBoxState::Checked);
+						FunctionLibraryController->SetFunctionIsPure(LibraryNode->GetFName(), bIsPure);
+					}
+				}
+			}
+		}
+	}
+}
+
 bool FRigVMGraphDetailCustomizationImpl::IsValidFunction() const
 {
 	if (WeakModel.IsValid() && RigVMClientHost.IsValid())
diff --git a/Engine/Plugins/Runtime/RigVM/Source/RigVMEditor/Public/Editor/RigVMGraphDetailCustomization.h b/Engine/Plugins/Runtime/RigVM/Source/RigVMEditor/Public/Editor/RigVMGraphDetailCustomization.h
--- a/Engine/Plugins/Runtime/RigVM/Source/RigVMEditor/Public/Editor/RigVMGraphDetailCustomization.h
+++ b/Engine/Plugins/Runtime/RigVM/Source/RigVMEditor/Public/Editor/RigVMGraphDetailCustomization.h
@@ -281,6 +281,8 @@ private:
 	UE_API FText GetCurrentAccessSpecifierName() const;
 	UE_API void OnAccessSpecifierSelected( TSharedPtr<FRigVMStringWithTag> SpecifierName, ESelectInfo::Type SelectInfo );
 	UE_API TSharedRef<ITableRow> HandleGenerateRowAccessSpecifier( TSharedPtr<FRigVMStringWithTag> SpecifierName, const TSharedRef<STableViewBase>& OwnerTable );
+	UE_API ECheckBoxState GetFunctionPureCheckState() const;
+	UE_API void OnFunctionPureCheckStateChanged(ECheckBoxState NewState);
 	UE_API bool IsValidFunction() const;
 	UE_API FRigVMVariant GetVariant() const;
 	UE_API FRigVMVariantRef GetSubjectVariantRef() const;
diff --git a/Engine/Source/Programs/Shared/EpicGames.UHT/Exporters/CodeGen/UhtHeaderCodeGeneratorCppFile.cs b/Engine/Source/Programs/Shared/EpicGames.UHT/Exporters/CodeGen/UhtHeaderCodeGeneratorCppFile.cs
--- a/Engine/Source/Programs/Shared/EpicGames.UHT/Exporters/CodeGen/UhtHeaderCodeGeneratorCppFile.cs
+++ b/Engine/Source/Programs/Shared/EpicGames.UHT/Exporters/CodeGen/UhtHeaderCodeGeneratorCppFile.cs
@@ -1258,7 +1258,7 @@ namespace EpicGames.UHT.Exporters.CodeGen
 			if (scriptStruct.RigVMStructInfo != null)
 			{
 				string constPrefix = "";
-				if (!scriptStruct.RigVMStructInfo.HasAnyExecuteContextMember)
+				if (!scriptStruct.RigVMStructInfo.HasAnyExecuteContextMember || scriptStruct.RigVMStructInfo.IsPure)
 				{
 					constPrefix = "const ";
 				}
diff --git a/Engine/Source/Programs/Shared/EpicGames.UHT/Exporters/CodeGen/UhtHeaderCodeGeneratorHFile.cs b/Engine/Source/Programs/Shared/EpicGames.UHT/Exporters/CodeGen/UhtHeaderCodeGeneratorHFile.cs
--- a/Engine/Source/Programs/Shared/EpicGames.UHT/Exporters/CodeGen/UhtHeaderCodeGeneratorHFile.cs
+++ b/Engine/Source/Programs/Shared/EpicGames.UHT/Exporters/CodeGen/UhtHeaderCodeGeneratorHFile.cs
@@ -254,7 +254,7 @@ namespace EpicGames.UHT.Exporters.CodeGen
 					builder.Append("#define ").Append(scriptStruct.SourceName).Append('_').Append(methodInfo.Name).Append("() \\\r\n");
 					builder.Append('\t').Append(methodInfo.ReturnType).Append(' ').Append(scriptStruct.SourceName).Append("::Static").Append(methodInfo.Name).Append("( \\\r\n");
 					builder.Append("\t\t");
-					if (!scriptStruct.RigVMStructInfo.HasAnyExecuteContextMember)
+					if (!scriptStruct.RigVMStructInfo.HasAnyExecuteContextMember || scriptStruct.RigVMStructInfo.IsPure)
 					{
 						builder.Append("const ");
 					}
@@ -283,7 +283,7 @@ namespace EpicGames.UHT.Exporters.CodeGen
 			if (scriptStruct.RigVMStructInfo != null)
 			{
 				string constPrefix = "";
-				if (!scriptStruct.RigVMStructInfo.HasAnyExecuteContextMember)
+				if (!scriptStruct.RigVMStructInfo.HasAnyExecuteContextMember || scriptStruct.RigVMStructInfo.IsPure)
 				{
 					constPrefix = "const ";
 				}
@@ -348,7 +348,7 @@ namespace EpicGames.UHT.Exporters.CodeGen
 					builder.Append('\t').Append(methodInfo.DllStorage).Append(methodInfo.ReturnType).Append(' ').Append(methodInfo.Name).Append('(').Append(constPrefix).Append(scriptStruct.RigVMStructInfo.ExecuteContextType).Append("& InExecuteContext); \\\r\n");
 					builder.Append("\tstatic ").Append(methodInfo.ReturnType).Append(' ').Append(methodInfo.DllStorage).Append("Static").Append(methodInfo.Name).Append("( \\\r\n");
 					builder.Append("\t\t");
-					if (!scriptStruct.RigVMStructInfo.HasAnyExecuteContextMember)
+					if (!scriptStruct.RigVMStructInfo.HasAnyExecuteContextMember || scriptStruct.RigVMStructInfo.IsPure)
 					{
 						builder.Append("const ");
 					}
diff --git a/Engine/Source/Programs/Shared/EpicGames.UHT/Parsers/UhtScriptStructParser.cs b/Engine/Source/Programs/Shared/EpicGames.UHT/Parsers/UhtScriptStructParser.cs
--- a/Engine/Source/Programs/Shared/EpicGames.UHT/Parsers/UhtScriptStructParser.cs
+++ b/Engine/Source/Programs/Shared/EpicGames.UHT/Parsers/UhtScriptStructParser.cs
@@ -172,6 +172,7 @@ namespace EpicGames.UHT.Parsers
 				{
 					structInfo = new();
 					structInfo.Name = scriptStruct.EngineName;
+					structInfo.IsPure = scriptStruct.MetaData.ContainsKey("Pure");
 					scriptStruct.RigVMStructInfo = structInfo;
 				}
 
diff --git a/Engine/Source/Programs/Shared/EpicGames.UHT/Types/UhtScriptStruct.cs b/Engine/Source/Programs/Shared/EpicGames.UHT/Types/UhtScriptStruct.cs
--- a/Engine/Source/Programs/Shared/EpicGames.UHT/Types/UhtScriptStruct.cs
+++ b/Engine/Source/Programs/Shared/EpicGames.UHT/Types/UhtScriptStruct.cs
@@ -635,6 +635,11 @@ namespace EpicGames.UHT.Types
 		/// </summary>
 		public bool HasAnyExecuteContextMember { get; set; } = false;
 
+		/// <summary>
+		/// true if this struct has the Pure metadata (should receive const execute context)
+		/// </summary>
+		public bool IsPure { get; set; } = false;
+
 		/// <summary>
 		/// List of the members
 		/// </summary>
@@ -664,6 +669,7 @@ namespace EpicGames.UHT.Types
 			ExecuteContextType = reader.ReadString();
 			ExecuteContextMember = reader.ReadString();
 			HasAnyExecuteContextMember = reader.ReadBoolean();
+			IsPure = reader.ReadBoolean();
 			Members = reader.ReadList((reader) => new UhtRigVMParameter(reader));
 			Methods = reader.ReadList((reader) => new UhtRigVMMethodInfo(reader));
 		}
@@ -680,6 +686,7 @@ namespace EpicGames.UHT.Types
 			writer.WriteString(ExecuteContextType);
 			writer.WriteString(ExecuteContextMember);
 			writer.WriteBoolean(HasAnyExecuteContextMember);
+			writer.WriteBoolean(IsPure);
 			writer.WriteVariableLengthArray(Members, (writer, x) => x.Write(writer));
 			writer.WriteVariableLengthArray(Methods, (writer, x) => x.Write(writer));
 		}
