Control Rig: Fix node creation through shortcuts not working in new editor
#rnx

#Synced-CL 51819521

diff --git a/Engine/Plugins/Runtime/RigVM/Source/RigVMEditor/Private/Editor/RigVMNewEditor.cpp b/Engine/Plugins/Runtime/RigVM/Source/RigVMEditor/Private/Editor/RigVMNewEditor.cpp
--- a/Engine/Plugins/Runtime/RigVM/Source/RigVMEditor/Private/Editor/RigVMNewEditor.cpp
+++ b/Engine/Plugins/Runtime/RigVM/Source/RigVMEditor/Private/Editor/RigVMNewEditor.cpp
@@ -1639,8 +1639,8 @@ void FRigVMNewEditor::SetupGraphEditorEventsImpl(UEdGraph* InGraph, SGraphEditor
 	// TODO sara-s: Replicate what blueprint editor is doing
 	InEvents.OnSelectionChanged = SGraphEditor::FOnSelectionChanged::CreateSP( this, &FRigVMNewEditor::OnSelectedNodesChanged );
 	// InEvents.OnVerifyTextCommit = FOnNodeVerifyTextCommit::CreateSP(this, &FRigVMNewEditor::OnNodeVerifyTitleCommit);
-	// InEvents.OnTextCommitted = FOnNodeTextCommitted::CreateSP(this, &FRigVMNewEditor::OnNodeTitleCommitted);
-	// InEvents.OnSpawnNodeByShortcut = SGraphEditor::FOnSpawnNodeByShortcut::CreateSP(this, &FRigVMNewEditor::OnSpawnGraphNodeByShortcut, InGraph);
+	InEvents.OnTextCommitted = FOnNodeTextCommitted::CreateSP(this, &FRigVMNewEditor::OnNodeTitleCommitted);
+	InEvents.OnSpawnNodeByShortcutAtLocation = SGraphEditor::FOnSpawnNodeByShortcutAtLocation::CreateSP(this, &FRigVMNewEditor::OnSpawnGraphNodeByShortcut, InGraph);
 	// InEvents.OnNodeSpawnedByKeymap = SGraphEditor::FOnNodeSpawnedByKeymap::CreateSP(this, &FRigVMNewEditor::OnNodeSpawnedByKeymap );
 	// InEvents.OnDisallowedPinConnection = SGraphEditor::FOnDisallowedPinConnection::CreateSP(this, &FRigVMNewEditor::OnDisallowedPinConnection);
 	// InEvents.OnDoubleClicked = SGraphEditor::FOnDoubleClicked::CreateSP(this, &FRigVMNewEditor::NavigateToParentGraphByDoubleClick);
diff --git a/Engine/Plugins/Runtime/RigVM/Source/RigVMEditor/Public/Editor/RigVMLegacyEditor.h b/Engine/Plugins/Runtime/RigVM/Source/RigVMEditor/Public/Editor/RigVMLegacyEditor.h
--- a/Engine/Plugins/Runtime/RigVM/Source/RigVMEditor/Public/Editor/RigVMLegacyEditor.h
+++ b/Engine/Plugins/Runtime/RigVM/Source/RigVMEditor/Public/Editor/RigVMLegacyEditor.h
@@ -159,7 +159,7 @@ protected:
 	virtual bool CanPasteNodes() const override { return FRigVMEditorBase::CanPasteNodes(); }
 	virtual bool IsNativeParentClassCodeLinkEnabled() const override { return FRigVMEditorBase::IsNativeParentClassCodeLinkEnabled(); }
 	virtual bool ReparentBlueprint_IsVisible() const override { return FRigVMEditorBase::ReparentBlueprint_IsVisible(); }
-	virtual FReply OnSpawnGraphNodeByShortcut(FInputChord InChord, const FVector2f& InPosition, UEdGraph* InGraph) override { return FRigVMEditorBase::OnSpawnGraphNodeByShortcut(InChord, FDeprecateSlateVector2D(InPosition), InGraph); }
+	virtual FReply OnSpawnGraphNodeByShortcut(FInputChord InChord, const FVector2f& InPosition, UEdGraph* InGraph) override { return FRigVMEditorBase::OnSpawnGraphNodeByShortcut(InChord, InPosition, InGraph); }
 	virtual bool ShouldLoadBPLibrariesFromAssetRegistry() override { return FRigVMEditorBase::ShouldLoadBPLibrariesFromAssetRegistry(); }
 	virtual bool ShouldOpenGraphByDefault() const { return FRigVMEditorBase::ShouldOpenGraphByDefault(); }
 	virtual void AddNewFunctionVariant(const UEdGraph* InOriginalFunction) override { return FRigVMEditorBase::AddNewFunctionVariant(InOriginalFunction); }
