RigVM: Make the hash creation for FRigVMGraphFunction deterministic with callables

#rnx

#Synced-CL 51826256

diff --git a/Engine/Plugins/Runtime/RigVM/Source/RigVM/Public/RigVMCore/RigVMGraphFunctionDefinition.h b/Engine/Plugins/Runtime/RigVM/Source/RigVM/Public/RigVMCore/RigVMGraphFunctionDefinition.h
--- a/Engine/Plugins/Runtime/RigVM/Source/RigVM/Public/RigVMCore/RigVMGraphFunctionDefinition.h
+++ b/Engine/Plugins/Runtime/RigVM/Source/RigVM/Public/RigVMCore/RigVMGraphFunctionDefinition.h
@@ -239,10 +239,22 @@ struct FRigVMFunctionCompilationData
 			DataHash = HashCombine(DataHash, GetTypeHash(Pair.Value));
 		}
 
-		for (const TPair<FName, FRigVMOperand>& Pair : Data.InterfaceOperands)
 		{
-			DataHash = HashCombine(DataHash, GetTypeHash(Pair.Key));
-			DataHash = HashCombine(DataHash, GetTypeHash(Pair.Value));
+			TArray<TPair<FString, FRigVMOperand>> SortedInterfaceOperands;
+			SortedInterfaceOperands.Reserve(Data.InterfaceOperands.Num());
+			for (const TPair<FName, FRigVMOperand>& Pair : Data.InterfaceOperands)
+			{
+				SortedInterfaceOperands.Emplace(Pair.Key.ToString(), Pair.Value);
+			}
+			SortedInterfaceOperands.Sort([](const TPair<FString, FRigVMOperand>& A, const TPair<FString, FRigVMOperand>& B)
+			{
+				return A.Key < B.Key;
+			});
+			for (const TPair<FString, FRigVMOperand>& Pair : SortedInterfaceOperands)
+			{
+				DataHash = HashCombine(DataHash, GetTypeHash(Pair.Key));
+				DataHash = HashCombine(DataHash, GetTypeHash(Pair.Value));
+			}
 		}
 
 		return DataHash;
