RigVM: Mark nodes as pure
#rnx

#Synced-CL 51823286

diff --git a/Engine/Plugins/Runtime/RigVM/Source/RigVM/Public/RigVMFunctions/Execution/RigVMFunction_ForLoop.h b/Engine/Plugins/Runtime/RigVM/Source/RigVM/Public/RigVMFunctions/Execution/RigVMFunction_ForLoop.h
--- a/Engine/Plugins/Runtime/RigVM/Source/RigVM/Public/RigVMFunctions/Execution/RigVMFunction_ForLoop.h
+++ b/Engine/Plugins/Runtime/RigVM/Source/RigVM/Public/RigVMFunctions/Execution/RigVMFunction_ForLoop.h
@@ -8,7 +8,7 @@
 /**
  * Given a count, execute iteratively until the count is up
  */
-USTRUCT(meta=(DisplayName="For Loop", Category="Execution", TitleColor="1 0 0", NodeColor="1 1 1", Keywords="Iterate", Icon="EditorStyle|GraphEditor.Macro.Loop_16x", DocumentationPolicy="Strict"))
+USTRUCT(meta=(DisplayName="For Loop", Category="Execution", TitleColor="1 0 0", NodeColor="1 1 1", Keywords="Iterate", Icon="EditorStyle|GraphEditor.Macro.Loop_16x", DocumentationPolicy="Strict", Pure))
 struct FRigVMFunction_ForLoopCount : public FRigVMStructMutable
 {
 	GENERATED_BODY()
diff --git a/Engine/Plugins/Runtime/RigVM/Source/RigVM/Public/RigVMFunctions/Execution/RigVMFunction_Sequence.h b/Engine/Plugins/Runtime/RigVM/Source/RigVM/Public/RigVMFunctions/Execution/RigVMFunction_Sequence.h
--- a/Engine/Plugins/Runtime/RigVM/Source/RigVM/Public/RigVMFunctions/Execution/RigVMFunction_Sequence.h
+++ b/Engine/Plugins/Runtime/RigVM/Source/RigVM/Public/RigVMFunctions/Execution/RigVMFunction_Sequence.h
@@ -8,7 +8,7 @@
 /**
  * Allows for a single execution pulse to trigger a series of events in order.
  */
-USTRUCT(meta=(DisplayName="Sequence", Category="Execution", TitleColor="1 0 0", NodeColor="1 1 1", Icon="EditorStyle|GraphEditor.Sequence_16x"))
+USTRUCT(meta=(DisplayName="Sequence", Category="Execution", TitleColor="1 0 0", NodeColor="1 1 1", Icon="EditorStyle|GraphEditor.Sequence_16x", Pure))
 struct FRigVMFunction_Sequence : public FRigVMStruct
 {
 	GENERATED_BODY()
diff --git a/Engine/Plugins/Runtime/RigVM/Source/RigVM/Public/RigVMFunctions/RigVMDispatch_If.h b/Engine/Plugins/Runtime/RigVM/Source/RigVM/Public/RigVMFunctions/RigVMDispatch_If.h
--- a/Engine/Plugins/Runtime/RigVM/Source/RigVM/Public/RigVMFunctions/RigVMDispatch_If.h
+++ b/Engine/Plugins/Runtime/RigVM/Source/RigVM/Public/RigVMFunctions/RigVMDispatch_If.h
@@ -9,7 +9,7 @@
 /*
  * Chooses between two values based on a condition
  */
-USTRUCT(meta=(DisplayName = "If", Category = "Execution", Keywords = "Branch,Condition", NodeColor = "0,1,0,1"))
+USTRUCT(meta=(DisplayName = "If", Category = "Execution", Keywords = "Branch,Condition", NodeColor = "0,1,0,1", Pure))
 struct FRigVMDispatch_If : public FRigVMDispatch_CoreBase
 {
 	GENERATED_BODY()
diff --git a/Engine/Plugins/Runtime/RigVM/Source/RigVM/Public/RigVMFunctions/RigVMDispatch_Select.h b/Engine/Plugins/Runtime/RigVM/Source/RigVM/Public/RigVMFunctions/RigVMDispatch_Select.h
--- a/Engine/Plugins/Runtime/RigVM/Source/RigVM/Public/RigVMFunctions/RigVMDispatch_Select.h
+++ b/Engine/Plugins/Runtime/RigVM/Source/RigVM/Public/RigVMFunctions/RigVMDispatch_Select.h
@@ -9,7 +9,7 @@
 /*
  * Pick from a list of values based on an integer index
  */
-USTRUCT(meta=(DisplayName = "Select", Category = "Execution", Keywords = "Switch,Case", NodeColor = "0,1,0,1"))
+USTRUCT(meta=(DisplayName = "Select", Category = "Execution", Keywords = "Switch,Case", NodeColor = "0,1,0,1", Pure))
 struct FRigVMDispatch_SelectInt32 : public FRigVMDispatch_CoreBase
 {
 	GENERATED_BODY()
diff --git a/Engine/Plugins/Runtime/RigVM/Source/RigVM/Public/RigVMFunctions/RigVMDispatch_Switch.h b/Engine/Plugins/Runtime/RigVM/Source/RigVM/Public/RigVMFunctions/RigVMDispatch_Switch.h
--- a/Engine/Plugins/Runtime/RigVM/Source/RigVM/Public/RigVMFunctions/RigVMDispatch_Switch.h
+++ b/Engine/Plugins/Runtime/RigVM/Source/RigVM/Public/RigVMFunctions/RigVMDispatch_Switch.h
@@ -34,7 +34,7 @@ protected:
 /*
  * Run a branch based on an integer index
  */
-USTRUCT(meta=(DisplayName = "Switch", Category = "Execution", Keywords = "Case", NodeColor = "0,1,0,1"))
+USTRUCT(meta=(DisplayName = "Switch", Category = "Execution", Keywords = "Case", NodeColor = "0,1,0,1", Pure))
 struct FRigVMDispatch_SwitchInt32 : public FRigVMDispatch_SwitchBase
 {
 	GENERATED_BODY()
@@ -66,7 +66,7 @@ protected:
 /*
  * Run a branch based on an enum index
  */
-USTRUCT(meta=(DisplayName = "Switch Enum", Category = "Execution", Keywords = "Case", NodeColor = "0,1,0,1"))
+USTRUCT(meta=(DisplayName = "Switch Enum", Category = "Execution", Keywords = "Case", NodeColor = "0,1,0,1", Pure))
 struct FRigVMDispatch_SwitchEnum : public FRigVMDispatch_SwitchBase
 {
 	GENERATED_BODY()
diff --git a/Engine/Plugins/Runtime/RigVM/Source/RigVM/Public/RigVMFunctions/RigVMFunction_ControlFlow.h b/Engine/Plugins/Runtime/RigVM/Source/RigVM/Public/RigVMFunctions/RigVMFunction_ControlFlow.h
--- a/Engine/Plugins/Runtime/RigVM/Source/RigVM/Public/RigVMFunctions/RigVMFunction_ControlFlow.h
+++ b/Engine/Plugins/Runtime/RigVM/Source/RigVM/Public/RigVMFunctions/RigVMFunction_ControlFlow.h
@@ -17,7 +17,7 @@ struct FRigVMFunction_ControlFlowBase : public FRigVMStruct
 /**
  * Executes either the True or False branch based on the condition
  */
-USTRUCT(meta = (DisplayName = "Branch", Keywords = "If"))
+USTRUCT(meta = (DisplayName = "Branch", Keywords = "If", Pure))
 struct FRigVMFunction_ControlFlowBranch : public FRigVMFunction_ControlFlowBase
 {
 	GENERATED_BODY()
