Control Rig: Fix hide unrelated nodes is hiding used variable nodes
#rnx

#Synced-CL 51836820

diff --git a/Engine/Plugins/Runtime/RigVM/Source/RigVMDeveloper/Private/EdGraph/RigVMEdGraphNode.cpp b/Engine/Plugins/Runtime/RigVM/Source/RigVMDeveloper/Private/EdGraph/RigVMEdGraphNode.cpp
--- a/Engine/Plugins/Runtime/RigVM/Source/RigVMDeveloper/Private/EdGraph/RigVMEdGraphNode.cpp
+++ b/Engine/Plugins/Runtime/RigVM/Source/RigVMDeveloper/Private/EdGraph/RigVMEdGraphNode.cpp
@@ -1517,13 +1517,44 @@ float URigVMEdGraphNode::GetFadedOutState() const
 	{
 		return FadedOutOverride.GetValue();
 	}
-	if(GetInstructionIndex(true) == INDEX_NONE)
+	
+	bool bIsVariableUsed = false;
+	if (URigVMNode* ModelNode = GetModelNode())
+	{
+		if (URigVMVariableNode* VariableNode = Cast<URigVMVariableNode>(ModelNode))
+		{
+			if (VariableNode->IsGetter())
+			{
+				if (URigVMEdGraph* Graph = Cast<URigVMEdGraph>(GetGraph()))
+				{
+					TArray<URigVMNode*> TargetNodes = VariableNode->GetLinkedTargetNodes();
+					for (URigVMNode* TargetNode : TargetNodes)
+					{
+						if (URigVMEdGraphNode* TargetEdNode = Cast<URigVMEdGraphNode>(Graph->FindNodeForModelNodeName(TargetNode->GetFName())))
+						{
+							if (!TargetEdNode->IsFadedOut())
+							{
+								bIsVariableUsed = true;
+								break;
+							}
+						}
+					}
+				}
+			}
+		}
+	}
+	
+	if (!bIsVariableUsed)
 	{
-		if (GetDefault<URigVMEditorSettings>()->bFadeOutUnrelatedNodes)
-    	{
-			return 0.25f;
+		if(GetInstructionIndex(true) == INDEX_NONE)
+		{
+			if (GetDefault<URigVMEditorSettings>()->bFadeOutUnrelatedNodes)
+			{
+				return 0.25f;
+			}
 		}
 	}
+	
 	if (IsExcludedFromPreview())
 	{
 		bool bWasRunDuringLastEval = false;
