Control Rig: Add find toolbar button in new editor
#rnx

#Synced-CL 51840311

diff --git a/Engine/Plugins/Runtime/RigVM/Source/RigVMEditor/Private/Editor/RigVMEditor.cpp b/Engine/Plugins/Runtime/RigVM/Source/RigVMEditor/Private/Editor/RigVMEditor.cpp
--- a/Engine/Plugins/Runtime/RigVM/Source/RigVMEditor/Private/Editor/RigVMEditor.cpp
+++ b/Engine/Plugins/Runtime/RigVM/Source/RigVMEditor/Private/Editor/RigVMEditor.cpp
@@ -3727,6 +3727,7 @@ void FRigVMEditorBase::FillToolbar(FToolBarBuilder& ToolbarBuilder, bool bEndSec
 			false);
 
 		AddSettingsAndDefaultWidget(ToolbarBuilder);
+		AddFindWidget(ToolbarBuilder);
 	}
 
 	if(bEndSection)
diff --git a/Engine/Plugins/Runtime/RigVM/Source/RigVMEditor/Private/Editor/RigVMEditorCommands.cpp b/Engine/Plugins/Runtime/RigVM/Source/RigVMEditor/Private/Editor/RigVMEditorCommands.cpp
--- a/Engine/Plugins/Runtime/RigVM/Source/RigVMEditor/Private/Editor/RigVMEditorCommands.cpp
+++ b/Engine/Plugins/Runtime/RigVM/Source/RigVMEditor/Private/Editor/RigVMEditorCommands.cpp
@@ -37,7 +37,7 @@ void FRigVMEditorCommands::RegisterCommands()
 	UI_COMMAND(TogglePreviewHere, "Toggle Preview Here", "Run the graph up to this node", EUserInterfaceActionType::Button, FInputChord(EKeys::F9));
 	UI_COMMAND(PreviewHereStepForward, "Step Forwards", "Run the graph up to the next node.", EUserInterfaceActionType::Button, FInputChord(EKeys::F10));
 
-	UI_COMMAND(FindReferences, "Find References", "Find references of this item", EUserInterfaceActionType::Button, FInputChord(EKeys::F, EModifierKey::Shift | EModifierKey::Alt));
+	UI_COMMAND(FindReferences, "Find References", "Find references of this item", EUserInterfaceActionType::Button, FInputChord(EKeys::F, EModifierKey::Control));
 	UI_COMMAND(FindReferencesByNameLocal, "By Name", "Find references of this item by name", EUserInterfaceActionType::Button, FInputChord());
 	UI_COMMAND(FindReferencesByNameGlobal, "By Name (All)", "Find references of this item by name in all blueprints", EUserInterfaceActionType::Button, FInputChord());
 }
diff --git a/Engine/Plugins/Runtime/RigVM/Source/RigVMEditor/Private/Editor/RigVMNewEditor.cpp b/Engine/Plugins/Runtime/RigVM/Source/RigVMEditor/Private/Editor/RigVMNewEditor.cpp
--- a/Engine/Plugins/Runtime/RigVM/Source/RigVMEditor/Private/Editor/RigVMNewEditor.cpp
+++ b/Engine/Plugins/Runtime/RigVM/Source/RigVMEditor/Private/Editor/RigVMNewEditor.cpp
@@ -727,6 +727,28 @@ void FRigVMNewEditor::AddSettingsAndDefaultWidget(FToolBarBuilder& ToolbarBuilde
 	ToolbarBuilder.AddToolBarButton(FRigVMEditorCommands::Get().EditClassDefaults);
 }
 
+void FRigVMNewEditor::AddFindWidget(FToolBarBuilder& ToolbarBuilder)
+{
+	
+	ToolbarBuilder.AddToolBarButton(
+		FRigVMEditorCommands::Get().FindReferences,
+		NAME_None,
+		LOCTEXT("FindReferences_Label", "Find"),
+			LOCTEXT("FindReferences_ToolTip", "Find all references to the selected variable"),
+		TAttribute<FSlateIcon>::CreateLambda(
+			[]() -> FSlateIcon
+			{
+				const FSlateIcon GlobalFindResultsIcon(FAppStyle::GetAppStyleSetName(), "BlueprintEditor.RigVMFindInBlueprints.MenuIcon");
+				return GlobalFindResultsIcon;
+			}
+		),
+		NAME_None,
+		FNewMenuDelegate(),
+		TAttribute<EVisibility>(),
+		LOCTEXT("FindReferences_Label", "Find")
+	);
+}
+
 void FRigVMNewEditor::AddSelectedDebugObjectWidget(FToolBarBuilder& ToolbarBuilder)
 {
 	ToolbarBuilder.AddWidget(SNew(SRigVMEditorSelectedDebugObjectWidget, SharedThis(this)));
@@ -793,6 +815,11 @@ void FRigVMNewEditor::ZoomToSelection_Clicked()
 	}
 }
 
+void FRigVMNewEditor::FindInBlueprint_Clicked()
+{
+	SummonSearchUI(true, FString(), true);
+}
+
 void FRigVMNewEditor::RestoreEditedObjectState()
 {
 	check(IsEditingSingleBlueprint());
@@ -1246,6 +1273,12 @@ void FRigVMNewEditor::CreateDefaultCommands()
 			FExecuteAction::CreateSP(this, &FRigVMNewEditor::EditClassDefaults_Clicked),
 			FCanExecuteAction(),
 			FIsActionChecked::CreateSP(this, &FRigVMNewEditor::IsDetailsPanelEditingClassDefaults));
+		
+		ToolkitCommands->MapAction(
+			FRigVMEditorCommands::Get().FindReferences,
+			FExecuteAction::CreateSP(this, &FRigVMNewEditor::FindInBlueprint_Clicked),
+			FCanExecuteAction()
+		);
 	}
 
 	FRigVMEditorBase::CreateDefaultCommandsImpl();
diff --git a/Engine/Plugins/Runtime/RigVM/Source/RigVMEditor/Private/Widgets/SRigVMEditorGraphExplorer.cpp b/Engine/Plugins/Runtime/RigVM/Source/RigVMEditor/Private/Widgets/SRigVMEditorGraphExplorer.cpp
--- a/Engine/Plugins/Runtime/RigVM/Source/RigVMEditor/Private/Widgets/SRigVMEditorGraphExplorer.cpp
+++ b/Engine/Plugins/Runtime/RigVM/Source/RigVMEditor/Private/Widgets/SRigVMEditorGraphExplorer.cpp
@@ -210,6 +210,7 @@ void FRigVMEditorGraphExplorerCommands::RegisterCommands()
 	UI_COMMAND(PasteLocalVariable, "Paste Local Variable", "Pastes the local variable.", EUserInterfaceActionType::Button, FInputChord());
 	UI_COMMAND(RemoveUnusedFunctions, "Remove Unused Functions", "Removes unsued functions.", EUserInterfaceActionType::Button, FInputChord());
 	UI_COMMAND(RemoveUnusedVariables, "Remove Unused Variables", "Removes unused variables.", EUserInterfaceActionType::Button, FInputChord());
+	UI_COMMAND(FindReferences, "Find references", "Find references.", EUserInterfaceActionType::Button, FInputChord(EKeys::F, EModifierKey::Control));
 }
 
 
@@ -217,6 +218,7 @@ void SRigVMEditorGraphExplorer::Construct(const FArguments& InArgs, TWeakPtr<IRi
 {
 	bNeedsRefresh = false;
 	RigVMEditor = InRigVMEditor;
+	UICommandList = MakeShared<FUICommandList>();
 
 	RegisterCommands();
 
@@ -229,6 +231,7 @@ void SRigVMEditorGraphExplorer::Construct(const FArguments& InArgs, TWeakPtr<IRi
 	if (Editor)
 	{
 		Editor->OnRefresh().AddSP(this, &SRigVMEditorGraphExplorer::Refresh);
+		Editor->GetKeyDownDelegate().BindSP(this, &SRigVMEditorGraphExplorer::OnKeyDown);
 
 		// Acquire a registry for function and variable nodes
 		if (const TScriptInterface<IRigVMEditorAssetInterface> AssetInterface = Editor->GetRigVMAssetInterface())
@@ -323,100 +326,110 @@ void SRigVMEditorGraphExplorer::HandleGraphModifiedEvent(ERigVMGraphNotifType In
 	}
 }
 
+FReply SRigVMEditorGraphExplorer::OnKeyDown(const FGeometry& MyGeometry, const FKeyEvent& InKeyEvent)
+{
+	if (UICommandList.IsValid() && UICommandList->ProcessCommandBindings(InKeyEvent))
+	{
+		return FReply::Handled();
+	}
+	return FReply::Unhandled();
+}
+
 void SRigVMEditorGraphExplorer::RegisterCommands()
 {
 	TSharedPtr<IRigVMEditor> Editor = RigVMEditor.Pin();
 
 	if (Editor)
 	{
-		TSharedPtr<FUICommandList> ToolKitCommandList = Editor->GetToolkitCommands();
-
-		ToolKitCommandList->MapAction(FRigVMEditorGraphExplorerCommands::Get().OpenGraph,
+		UICommandList->MapAction(FRigVMEditorGraphExplorerCommands::Get().OpenGraph,
 		    FExecuteAction::CreateSP(this, &SRigVMEditorGraphExplorer::OnOpenGraph, false),
 			FCanExecuteAction(), FGetActionCheckState(),
 			FIsActionButtonVisible::CreateSP(this, &SRigVMEditorGraphExplorer::CanOpenGraph));
 
-		ToolKitCommandList->MapAction(FRigVMEditorGraphExplorerCommands::Get().OpenGraphInNewTab,
+		UICommandList->MapAction(FRigVMEditorGraphExplorerCommands::Get().OpenGraphInNewTab,
 			FExecuteAction::CreateSP(this, &SRigVMEditorGraphExplorer::OnOpenGraph, true),
 			FCanExecuteAction(), FGetActionCheckState(),
 			FIsActionButtonVisible::CreateSP(this, &SRigVMEditorGraphExplorer::CanOpenGraph));
 
-		ToolKitCommandList->MapAction(FRigVMEditorGraphExplorerCommands::Get().CreateGraph,
+		UICommandList->MapAction(FRigVMEditorGraphExplorerCommands::Get().CreateGraph,
 			FExecuteAction::CreateSP(this, &SRigVMEditorGraphExplorer::OnCreateGraph),
 			FCanExecuteAction::CreateSP(this, &SRigVMEditorGraphExplorer::CanCreateGraph)
 			);
 
-		ToolKitCommandList->MapAction(FRigVMEditorGraphExplorerCommands::Get().CreateFunction,
+		UICommandList->MapAction(FRigVMEditorGraphExplorerCommands::Get().CreateFunction,
 			FExecuteAction::CreateSP(this, &SRigVMEditorGraphExplorer::OnCreateFunction),
 			FCanExecuteAction::CreateSP(this, &SRigVMEditorGraphExplorer::CanCreateFunction)
 			);
 
-		ToolKitCommandList->MapAction(FRigVMEditorGraphExplorerCommands::Get().CreateVariable,
+		UICommandList->MapAction(FRigVMEditorGraphExplorerCommands::Get().CreateVariable,
 			FExecuteAction::CreateSP(this, &SRigVMEditorGraphExplorer::OnCreateVariable),
 			FCanExecuteAction::CreateSP(this, &SRigVMEditorGraphExplorer::CanCreateVariable)
 			);
 
-		ToolKitCommandList->MapAction(FRigVMEditorGraphExplorerCommands::Get().CreateLocalVariable,
+		UICommandList->MapAction(FRigVMEditorGraphExplorerCommands::Get().CreateLocalVariable,
 			FExecuteAction::CreateSP(this, &SRigVMEditorGraphExplorer::OnCreateLocalVariable),
 			FCanExecuteAction::CreateSP(this, &SRigVMEditorGraphExplorer::CanCreateLocalVariable)
 			);
 
-		ToolKitCommandList->MapAction(FRigVMEditorGraphExplorerCommands::Get().AddFunctionVariant,
+		UICommandList->MapAction(FRigVMEditorGraphExplorerCommands::Get().AddFunctionVariant,
 			FExecuteAction::CreateSP(this, &SRigVMEditorGraphExplorer::OnAddFunctionVariant),
 			FCanExecuteAction::CreateSP(this, &SRigVMEditorGraphExplorer::CanAddFunctionVariant)
 			);
 
-		ToolKitCommandList->MapAction(FGenericCommands::Get().Rename,
+		UICommandList->MapAction(FGenericCommands::Get().Rename,
 			FExecuteAction::CreateSP(this, &SRigVMEditorGraphExplorer::OnRenameEntry),
 			FCanExecuteAction::CreateSP(this, &SRigVMEditorGraphExplorer::CanRenameEntry));
 
-		ToolKitCommandList->MapAction(FGenericCommands::Get().Copy,
+		UICommandList->MapAction(FGenericCommands::Get().Copy,
 			FExecuteAction::CreateSP(this, &SRigVMEditorGraphExplorer::OnCopy),
 			FCanExecuteAction::CreateSP(this, &SRigVMEditorGraphExplorer::CanCopy));
 		
-		ToolKitCommandList->MapAction(FGenericCommands::Get().Cut,
+		UICommandList->MapAction(FGenericCommands::Get().Cut,
 			FExecuteAction::CreateSP(this, &SRigVMEditorGraphExplorer::OnCut),
 			FCanExecuteAction::CreateSP(this, &SRigVMEditorGraphExplorer::CanCut));
 
-		ToolKitCommandList->MapAction(FGenericCommands::Get().Duplicate,
+		UICommandList->MapAction(FGenericCommands::Get().Duplicate,
 			FExecuteAction::CreateSP(this, &SRigVMEditorGraphExplorer::OnDuplicate),
 			FCanExecuteAction::CreateSP(this, &SRigVMEditorGraphExplorer::CanDuplicate));
 
-		ToolKitCommandList->MapAction(FGenericCommands::Get().Paste,
+		UICommandList->MapAction(FGenericCommands::Get().Paste,
 			FExecuteAction::CreateSP(this, &SRigVMEditorGraphExplorer::OnPasteGeneric),
 			FCanExecuteAction(), FIsActionChecked(),
 			FIsActionButtonVisible::CreateSP(this, &SRigVMEditorGraphExplorer::CanPasteGeneric));
 
-		ToolKitCommandList->MapAction(FRigVMEditorGraphExplorerCommands::Get().PasteFunction,
+		UICommandList->MapAction(FRigVMEditorGraphExplorerCommands::Get().PasteFunction,
 			FExecuteAction::CreateSP(this, &SRigVMEditorGraphExplorer::OnPasteFunction),
 			FCanExecuteAction(), FIsActionChecked(),
 			FIsActionButtonVisible::CreateSP(this, &SRigVMEditorGraphExplorer::CanPasteFunction));
 
-		ToolKitCommandList->MapAction(FRigVMEditorGraphExplorerCommands::Get().PasteVariable,
+		UICommandList->MapAction(FRigVMEditorGraphExplorerCommands::Get().PasteVariable,
 			FExecuteAction::CreateSP(this, &SRigVMEditorGraphExplorer::OnPasteVariable),
 			FCanExecuteAction(), FIsActionChecked(),
 			FIsActionButtonVisible::CreateSP(this, &SRigVMEditorGraphExplorer::CanPasteVariable));
 
-		ToolKitCommandList->MapAction(FRigVMEditorGraphExplorerCommands::Get().PasteLocalVariable,
+		UICommandList->MapAction(FRigVMEditorGraphExplorerCommands::Get().PasteLocalVariable,
 			FExecuteAction::CreateSP(this, &SRigVMEditorGraphExplorer::OnPasteLocalVariable),
 			FCanExecuteAction(), FIsActionChecked(),
 			FIsActionButtonVisible::CreateSP(this, &SRigVMEditorGraphExplorer::CanPasteLocalVariable));
 
-		ToolKitCommandList->MapAction(FGenericCommands::Get().Delete,
+		UICommandList->MapAction(FGenericCommands::Get().Delete,
 			FExecuteAction::CreateSP(this, &SRigVMEditorGraphExplorer::OnDeleteEntry),
 			FCanExecuteAction::CreateSP(this, &SRigVMEditorGraphExplorer::CanDeleteEntry));
-
-		// Find references local
+		
+		// Bind widget-level command for Ctrl+F (only active when this widget has focus)
 		{
 			constexpr bool bSearchAllBlueprints = false;
 
-			ToolKitCommandList->MapAction(FRigVMEditorCommands::Get().FindReferences,
+			UICommandList->MapAction(FRigVMEditorGraphExplorerCommands::Get().FindReferences,
 				FExecuteAction::CreateSP(this, &SRigVMEditorGraphExplorer::OnFindReferences, bSearchAllBlueprints),
-				FCanExecuteAction(),
-				FIsActionChecked(),
-				FIsActionButtonVisible::CreateSP(this, &SRigVMEditorGraphExplorer::CanFindReferences));
+				FCanExecuteAction::CreateSP(this, &SRigVMEditorGraphExplorer::CanFindReferences));
+		}
+
+		// Find references local
+		{
+			constexpr bool bSearchAllBlueprints = false;
 
-			ToolKitCommandList->MapAction(FRigVMEditorCommands::Get().FindReferencesByNameLocal,
+			UICommandList->MapAction(FRigVMEditorCommands::Get().FindReferencesByNameLocal,
 				FExecuteAction::CreateSP(this, &SRigVMEditorGraphExplorer::OnFindReferences, bSearchAllBlueprints),
 				FCanExecuteAction(),
 				FIsActionChecked(),
@@ -427,7 +440,7 @@ void SRigVMEditorGraphExplorer::RegisterCommands()
 		{
 			constexpr bool bSearchAllBlueprints = true;
 
-			ToolKitCommandList->MapAction(FRigVMEditorCommands::Get().FindReferencesByNameGlobal,
+			UICommandList->MapAction(FRigVMEditorCommands::Get().FindReferencesByNameGlobal,
 				FExecuteAction::CreateSP(this, &SRigVMEditorGraphExplorer::OnFindReferences, bSearchAllBlueprints),
 				FCanExecuteAction(),
 				FIsActionChecked(),
@@ -436,13 +449,13 @@ void SRigVMEditorGraphExplorer::RegisterCommands()
 
 		// Remove unused members
 		{
-			ToolKitCommandList->MapAction(FRigVMEditorGraphExplorerCommands::Get().RemoveUnusedFunctions,
+			UICommandList->MapAction(FRigVMEditorGraphExplorerCommands::Get().RemoveUnusedFunctions,
 				FExecuteAction::CreateSP(this, &SRigVMEditorGraphExplorer::OnRemoveUnusedFunctions),
 				FCanExecuteAction(),
 				FIsActionChecked(),
 				FIsActionButtonVisible::CreateSP(this, &SRigVMEditorGraphExplorer::CanRemoveUnusedFunctions));
 
-			ToolKitCommandList->MapAction(FRigVMEditorGraphExplorerCommands::Get().RemoveUnusedVariables,
+			UICommandList->MapAction(FRigVMEditorGraphExplorerCommands::Get().RemoveUnusedVariables,
 				FExecuteAction::CreateSP(this, &SRigVMEditorGraphExplorer::OnRemoveUnusedVariables),
 				FCanExecuteAction(),
 				FIsActionChecked(),
@@ -592,7 +605,7 @@ void SRigVMEditorGraphExplorer::CreateWidgets()
 
 TSharedRef<SWidget> SRigVMEditorGraphExplorer::CreateAddNewMenuWidget()
 {
-	FMenuBuilder MenuBuilder(/* bShouldCloseWindowAfterMenuSelection= */true, RigVMEditor.Pin()->GetToolkitCommands());
+	FMenuBuilder MenuBuilder(/* bShouldCloseWindowAfterMenuSelection= */true, UICommandList);
 
 	BuildAddNewMenu(MenuBuilder);
 
@@ -1335,7 +1348,7 @@ TSharedPtr<SWidget> SRigVMEditorGraphExplorer::OnContextMenuOpening()
 				MenuBuilder.AddMenuEntry(FRigVMEditorGraphExplorerCommands::Get().OpenGraph);
 				MenuBuilder.AddMenuEntry(FRigVMEditorGraphExplorerCommands::Get().OpenGraphInNewTab);
 				MenuBuilder.AddMenuEntry(FGenericCommands::Get().Rename);
-				MenuBuilder.AddMenuEntry(FRigVMEditorCommands::Get().FindReferences);
+				MenuBuilder.AddMenuEntry(FRigVMEditorGraphExplorerCommands::Get().FindReferences);
 				MenuBuilder.AddMenuEntry(FGenericCommands::Get().Cut);
 				MenuBuilder.AddMenuEntry(FGenericCommands::Get().Copy);
 				MenuBuilder.AddMenuEntry(FGenericCommands::Get().Delete);
@@ -2251,6 +2264,7 @@ void SRigVMEditorGraphExplorer::OnFindReferences(bool bSearchInAllBlueprints)
 	if (!CanFindReferences() ||
 		SelectedKeys.Num() != 1)
 	{
+		PinnedEditor->SummonSearchUI(true);
 		return;
 	}
 
@@ -2275,6 +2289,7 @@ void SRigVMEditorGraphExplorer::OnFindReferences(bool bSearchInAllBlueprints)
 		}
 		break;
 	}
+	case ERigVMExplorerElementType::Event:
 	case ERigVMExplorerElementType::Function:
 	case ERigVMExplorerElementType::Variable:
 	case ERigVMExplorerElementType::LocalVariable:
diff --git a/Engine/Plugins/Runtime/RigVM/Source/RigVMEditor/Public/Editor/RigVMEditor.h b/Engine/Plugins/Runtime/RigVM/Source/RigVMEditor/Public/Editor/RigVMEditor.h
--- a/Engine/Plugins/Runtime/RigVM/Source/RigVMEditor/Public/Editor/RigVMEditor.h
+++ b/Engine/Plugins/Runtime/RigVM/Source/RigVMEditor/Public/Editor/RigVMEditor.h
@@ -38,6 +38,7 @@ class SGraphEditor;
 class FPreviewScene;
 
 DECLARE_MULTICAST_DELEGATE_TwoParams(FRigVMEditorClosed, const IRigVMEditor*, FRigVMEditorAssetInterfacePtr);
+DECLARE_DELEGATE_RetVal_TwoParams(FReply, FRigVMEditorKeyDownDelegate, const FGeometry&, const FKeyEvent&);
 
 
 /* Enums to use when grouping the blueprint members in the list panel. The order here will determine the order in the list */
@@ -133,6 +134,7 @@ public:
 	virtual FPreviewHostUpdated& OnPreviewHostUpdated()  = 0;
 
 	virtual FRigVMEditorClosed& OnEditorClosed() = 0;
+	virtual FRigVMEditorKeyDownDelegate& GetKeyDownDelegate() = 0;
 	virtual UEdGraph* GetFocusedGraph() const = 0;
 	virtual URigVMGraph* GetFocusedModel() const = 0;
 	virtual FNotifyHook* GetNotifyHook() = 0;
@@ -254,6 +256,7 @@ protected:
 	virtual void RegisterToolbarTab(const TSharedRef<class FTabManager>& TabManager) = 0;
 	virtual void AddCompileWidget(FToolBarBuilder& ToolbarBuilder) = 0;
 	virtual void AddSettingsAndDefaultWidget(FToolBarBuilder& ToolbarBuilder) = 0;
+	virtual void AddFindWidget(FToolBarBuilder& ToolbarBuilder) = 0;
 	virtual void AddSelectedDebugObjectWidget(FToolBarBuilder& ToolbarBuilder) = 0;
 	virtual void AddAutoCompileWidget(FToolBarBuilder& ToolbarBuilder) = 0;
 	virtual void AddRunOnceWidget(FToolBarBuilder& ToolbarBuilder) = 0;
@@ -398,6 +401,7 @@ protected:
 	virtual FPreviewHostUpdated& OnPreviewHostUpdated() override { return PreviewHostUpdated;  }
 
 	virtual FRigVMEditorClosed& OnEditorClosed() override { return RigVMEditorClosedDelegate; }
+	virtual FRigVMEditorKeyDownDelegate& GetKeyDownDelegate() override { return OnKeyDownDelegate; }
 
 	UE_API virtual bool SelectLocalVariable(const UEdGraph* Graph, const FName& VariableName) override;
 	UE_API virtual bool SelectVariable(const FName& VariableName) override;
@@ -580,6 +584,7 @@ protected:
 	TSharedPtr<SRigVMEditorGraphExplorer> GraphExplorerWidget;
 
 	FRigVMEditorClosed RigVMEditorClosedDelegate;
+	FRigVMEditorKeyDownDelegate OnKeyDownDelegate;
 
 	UE_API virtual void SetGraphExplorerWidget(TSharedPtr<SRigVMEditorGraphExplorer> InWidget) override;
 	virtual TSharedPtr<SRigVMEditorGraphExplorer> GetGraphExplorerWidget() override { return GraphExplorerWidget; }
diff --git a/Engine/Plugins/Runtime/RigVM/Source/RigVMEditor/Public/Editor/RigVMLegacyEditor.h b/Engine/Plugins/Runtime/RigVM/Source/RigVMEditor/Public/Editor/RigVMLegacyEditor.h
--- a/Engine/Plugins/Runtime/RigVM/Source/RigVMEditor/Public/Editor/RigVMLegacyEditor.h
+++ b/Engine/Plugins/Runtime/RigVM/Source/RigVMEditor/Public/Editor/RigVMLegacyEditor.h
@@ -120,6 +120,7 @@ protected:
 	virtual const TArray<UObject*>* GetObjectsCurrentlyBeingEdited() const override { return FBlueprintEditor::GetObjectsCurrentlyBeingEdited(); }
 	RIGVMEDITOR_API virtual void AddCompileWidget(FToolBarBuilder& ToolbarBuilder) override;
 	virtual void AddSettingsAndDefaultWidget(FToolBarBuilder& ToolbarBuilder) override {}
+	virtual void AddFindWidget(FToolBarBuilder& ToolbarBuilder) override {}
 	RIGVMEDITOR_API virtual void AddSelectedDebugObjectWidget(FToolBarBuilder& ToolbarBuilder) override;
 	RIGVMEDITOR_API virtual void AddAutoCompileWidget(FToolBarBuilder& ToolbarBuilder) override;
 	virtual FString GetCustomDebugObjectLabel(UObject* ObjectBeingDebugged) const override { return FRigVMEditorBase::GetCustomDebugObjectLabel(ObjectBeingDebugged); }
diff --git a/Engine/Plugins/Runtime/RigVM/Source/RigVMEditor/Public/Editor/RigVMNewEditor.h b/Engine/Plugins/Runtime/RigVM/Source/RigVMEditor/Public/Editor/RigVMNewEditor.h
--- a/Engine/Plugins/Runtime/RigVM/Source/RigVMEditor/Public/Editor/RigVMNewEditor.h
+++ b/Engine/Plugins/Runtime/RigVM/Source/RigVMEditor/Public/Editor/RigVMNewEditor.h
@@ -116,6 +116,7 @@ protected:
 	UE_API virtual void GetPinTypeSelectorFilters(TArray<TSharedPtr<IPinTypeSelectorFilter>>& OutFilters) const override;
 	UE_API virtual void OnAddNewVariable() override;
 	UE_API virtual void ZoomToSelection_Clicked() override;
+	UE_API void FindInBlueprint_Clicked();
 public:
 	UE_API virtual void SummonSearchUI(bool bSetFindWithinBlueprint, FString NewSearchTerms, bool bSelectFirstResult) override;
 	UE_API virtual void RestoreEditedObjectState() override;
@@ -128,6 +129,7 @@ public:
 	virtual const TArray<UObject*>* GetObjectsCurrentlyBeingEdited() const override { return FAssetEditorToolkit::GetObjectsCurrentlyBeingEdited(); }
 	UE_API virtual void AddCompileWidget(FToolBarBuilder& ToolbarBuilder) override;
 	UE_API virtual void AddSettingsAndDefaultWidget(FToolBarBuilder& ToolbarBuilder) override;
+	UE_API virtual void AddFindWidget(FToolBarBuilder& ToolbarBuilder) override;
 	UE_API virtual void AddSelectedDebugObjectWidget(FToolBarBuilder& ToolbarBuilder) override;
 	virtual void AddAutoCompileWidget(FToolBarBuilder& ToolbarBuilder) override {}
 	virtual void Compile() { FRigVMEditorBase::Compile(); }
diff --git a/Engine/Plugins/Runtime/RigVM/Source/RigVMEditor/Public/Widgets/SRigVMEditorGraphExplorer.h b/Engine/Plugins/Runtime/RigVM/Source/RigVMEditor/Public/Widgets/SRigVMEditorGraphExplorer.h
--- a/Engine/Plugins/Runtime/RigVM/Source/RigVMEditor/Public/Widgets/SRigVMEditorGraphExplorer.h
+++ b/Engine/Plugins/Runtime/RigVM/Source/RigVMEditor/Public/Widgets/SRigVMEditorGraphExplorer.h
@@ -27,6 +27,7 @@ class FRigVMEditorGraphExplorerTreeElement;
 
 class FMenuBuilder;
 class SSearchBox;
+class FUICommandList;
 //
 
 class FRigVMGraphExplorerDragDropOp : public FGraphEditorDragDropAction 
@@ -82,6 +83,7 @@ public:
 	TSharedPtr<FUICommandInfo> PasteLocalVariable;
 	TSharedPtr<FUICommandInfo> RemoveUnusedFunctions;
 	TSharedPtr<FUICommandInfo> RemoveUnusedVariables;
+	TSharedPtr<FUICommandInfo> FindReferences;
 
 	void RegisterCommands() override;
 };
@@ -102,6 +104,7 @@ public:
 
 	// SWidget overrides
 	void Tick(const FGeometry& InAllottedGeometry, const double InCurrentTime, const float InDeltaTime) override;
+	virtual FReply OnKeyDown(const FGeometry& MyGeometry, const FKeyEvent& InKeyEvent);
 
 	void SetLastPinTypeUsed(const FEdGraphPinType& InType) { LastPinType = InType; }
 	FEdGraphPinType GetLastPinTypeUsed() const { return LastPinType; }
@@ -123,6 +126,9 @@ private:
 	bool bNeedsRefresh = true;
 
 	FEdGraphPinType LastPinType;
+
+	/** Command list for widget-level keyboard shortcuts */
+	TSharedPtr<FUICommandList> UICommandList;
 	
 	void RegisterCommands();
 	void CreateWidgets();
