Control Rig: Fix icons for class settings and defaults
#rnx

#Synced-CL 51840534

diff --git a/Engine/Plugins/Runtime/RigVM/Source/RigVMEditor/Private/Editor/RigVMNewEditor.cpp b/Engine/Plugins/Runtime/RigVM/Source/RigVMEditor/Private/Editor/RigVMNewEditor.cpp
--- a/Engine/Plugins/Runtime/RigVM/Source/RigVMEditor/Private/Editor/RigVMNewEditor.cpp
+++ b/Engine/Plugins/Runtime/RigVM/Source/RigVMEditor/Private/Editor/RigVMNewEditor.cpp
@@ -723,8 +723,16 @@ void FRigVMNewEditor::AddSettingsAndDefaultWidget(FToolBarBuilder& ToolbarBuilde
 		FSlateIcon(FAppStyle::GetAppStyleSetName(), "GraphEditor.ToggleHideUnrelatedNodes")
 	);
 
-	ToolbarBuilder.AddToolBarButton(FRigVMEditorCommands::Get().EditGlobalOptions);
-	ToolbarBuilder.AddToolBarButton(FRigVMEditorCommands::Get().EditClassDefaults);
+	ToolbarBuilder.AddToolBarButton(FRigVMEditorCommands::Get().EditGlobalOptions, 
+		NAME_None, 
+		TAttribute<FText>(), 
+		TAttribute<FText>(), 
+		FSlateIcon(FAppStyle::GetAppStyleSetName(), "FullBlueprintEditor.EditGlobalOptions"));
+	ToolbarBuilder.AddToolBarButton(FRigVMEditorCommands::Get().EditClassDefaults, 
+		NAME_None, 
+		TAttribute<FText>(), 
+		TAttribute<FText>(), 
+		FSlateIcon(FAppStyle::GetAppStyleSetName(), "FullBlueprintEditor.EditClassDefaults"));
 }
 
 void FRigVMNewEditor::AddFindWidget(FToolBarBuilder& ToolbarBuilder)
