Control Rig: Fix missing connection between entry and return in pure functions
#rnx

#Synced-CL 51883189

diff --git a/Engine/Plugins/Runtime/RigVM/Source/RigVMDeveloper/Private/RigVMModel/RigVMController.cpp b/Engine/Plugins/Runtime/RigVM/Source/RigVMDeveloper/Private/RigVMModel/RigVMController.cpp
--- a/Engine/Plugins/Runtime/RigVM/Source/RigVMDeveloper/Private/RigVMModel/RigVMController.cpp
+++ b/Engine/Plugins/Runtime/RigVM/Source/RigVMDeveloper/Private/RigVMModel/RigVMController.cpp
@@ -15194,7 +15194,7 @@ URigVMLibraryNode* URigVMController::AddFunctionToLibrary(const FName& InFunctio
 		ReturnNode->Position = FVector2D(250.f, 0.f);
 		CollapseController->Notify(ERigVMGraphNotifType::NodePositionChanged, ReturnNode);
 
-		if (bMutable)
+		if (EntryNode->IsMutable() && ReturnNode->IsMutable())
 		{
 			CollapseController->AddLink(EntryNode->FindPin(FRigVMStruct::ExecuteContextName.ToString()), ReturnNode->FindPin(FRigVMStruct::ExecuteContextName.ToString()), false);
 		}
