UE: Add more info to an assert

diff --git a/Engine/Plugins/Runtime/RigVM/Source/RigVM/Private/RigVMCore/RigVMFunction.cpp b/Engine/Plugins/Runtime/RigVM/Source/RigVM/Private/RigVMCore/RigVMFunction.cpp
--- a/Engine/Plugins/Runtime/RigVM/Source/RigVM/Private/RigVMCore/RigVMFunction.cpp
+++ b/Engine/Plugins/Runtime/RigVM/Source/RigVM/Private/RigVMCore/RigVMFunction.cpp
@@ -324,7 +324,7 @@ void FRigVMFunction::Serialize_NoLock(FArchive& Ar, FRigVMRegistryHandle& InRegi
 			{
 				FunctionPtr = GlobalFunction->FunctionPtr;
 #if WITH_EDITOR
-				check(GlobalFunction->Struct == Struct);
+				checkf(GlobalFunction->Struct == Struct, TEXT("Struct mismatch when looking for function '%s'. GlobalFunction->Struct:%s Struct:%s"), *Name, *GetPathNameSafe(GlobalFunction->Struct), *GetPathNameSafe(Struct));
 
 				auto CompareArguments = [GlobalFunction](const FRigVMFunctionArgument& SerializedArg, const FRigVMFunctionArgument& CurrentArg)
 				{
