RigVM Nativization System

Converts RigVM bytecode into C++ source code for improved runtime performance.

Architecture

RigVM Bytecode ? FRigVMCodeConverter ? JSON ? Inja Templates ? C++ Source

Core Components:

+---------------------------------------------------------------------------------------------------+
?         Class         ?                                  Purpose                                  ?
+-----------------------+---------------------------------------------------------------------------?
? FRigVMCodeConverter   ? Parses bytecode into JSON (variables, memory, instructions, control flow) ?
+-----------------------+---------------------------------------------------------------------------?
? FRigVMCodeEnvironment ? Inja template environment with UE-specific helpers                        ?
+-----------------------+---------------------------------------------------------------------------?
? FRigVMCodeGenerator   ? Template discovery and rendering orchestration                            ?
+-----------------------+---------------------------------------------------------------------------?
? FRigVMCodeTemplate    ? Lazy-loaded template file representation                                  ?
+-----------------------+---------------------------------------------------------------------------?
? FRigVMCodeOutput      ? Rendered output with save-to-disk support                                 ?
+---------------------------------------------------------------------------------------------------+

Bytecode Parsing (14-step process)

- External variables and memory layouts (literal/work)
- Property paths, callables, blocks, entries
- Instructions with operand resolution
- Dependency tracking (includes, libraries)

Template System

- Templates in Content/CodeGeneration/CPlusPlus/*.template
- Uses Inja (Jinja2-like C++ library)
- Custom functions: sanitize(), tabs(), cpp_get_default_for_native_type(), etc.
- Generates AssetNameVM.h and AssetNameVM.cpp

Usage

FRigVMCodeConversionSettings Settings;
Settings.TargetModule = TEXT("MyModule");

auto Converter = MakeShared<FRigVMCodeConverter>(EditorAsset, Settings);
if (Converter->HasError()) return;

FRigVMCodeGenerator Generator(ERigVMCodeLanguage::CPlusPlus);
TArray<TSharedPtr<FRigVMCodeOutput>> Outputs = Generator.GenerateAll(Converter);

Error Handling

- ensureMsgf() for graceful validation failures
- HasError() API for error state checking
- Descriptive error messages via UE_LOG(LogRigVMDeveloper, Error, ...)

#rnx

#Synced-CL 51941404

diff --git a/Engine/Plugins/Animation/ControlRig/Source/ControlRig/Private/ControlRig.cpp b/Engine/Plugins/Animation/ControlRig/Source/ControlRig/Private/ControlRig.cpp
--- a/Engine/Plugins/Animation/ControlRig/Source/ControlRig/Private/ControlRig.cpp
+++ b/Engine/Plugins/Animation/ControlRig/Source/ControlRig/Private/ControlRig.cpp
@@ -1017,18 +1017,7 @@ bool UControlRig::Execute(const FName& InEventName)
 	if (!IsReadOnly())
 	{
 		ExtendedExecuteContext.SetInstructionVisitInfo(&InstructionVisitInfo);
-		
-		if (GeneratedBy)
-		{
-			DebugInfo->CopyFrom(GeneratedBy->GetDebugInfo());
-		}
-		else if (UControlRig* CDO = GetClass()->GetDefaultObject<UControlRig>())
-		{
-			// Copy the breakpoints. This will not override the state of the breakpoints
-			DebugInfo->CopyFrom(CDO->DebugInfo.Get());
-		}
-
-		ExtendedExecuteContext.SetDebugInfo(DebugInfo);
+		ExtendedExecuteContext.SetDebugInfo(DebugInfo, false);
 		
 		if (IsProfilingEnabled())
 		{
diff --git a/Engine/Plugins/Animation/ControlRig/Source/ControlRigEditor/Private/Editor/ControlRigEditor.cpp b/Engine/Plugins/Animation/ControlRig/Source/ControlRigEditor/Private/Editor/ControlRigEditor.cpp
--- a/Engine/Plugins/Animation/ControlRig/Source/ControlRigEditor/Private/Editor/ControlRigEditor.cpp
+++ b/Engine/Plugins/Animation/ControlRig/Source/ControlRigEditor/Private/Editor/ControlRigEditor.cpp
@@ -4272,19 +4272,22 @@ const TArray<FRigDirectManipulationTarget> FControlRigBaseEditor::GetDirectManip
 	{
 		if (UControlRig* DebuggedControlRig = Cast<UControlRig>(GetRigVMAssetInterface()->GetObjectBeingDebugged()))
 		{
-			if(const URigVMUnitNode* Node = DirectManipulationSubject.Get())
+			if (DebuggedControlRig->GetVM() && !DebuggedControlRig->GetVM()->IsNativized())
 			{
-				if(Node->IsPartOfRuntime(DebuggedControlRig))
+				if(const URigVMUnitNode* Node = DirectManipulationSubject.Get())
 				{
-					const TSharedPtr<FStructOnScope> NodeInstance = Node->ConstructLiveStructInstance(DebuggedControlRig);
-					if(NodeInstance.IsValid() && NodeInstance->IsValid())
+					if(Node->IsPartOfRuntime(DebuggedControlRig))
 					{
-						if(const FRigUnit* UnitInstance = UControlRig::GetRigUnitInstanceFromScope(NodeInstance))
+						const TSharedPtr<FStructOnScope> NodeInstance = Node->ConstructLiveStructInstance(DebuggedControlRig);
+						if(NodeInstance.IsValid() && NodeInstance->IsValid())
 						{
-							TArray<FRigDirectManipulationTarget> Targets;
-							if(UnitInstance->GetDirectManipulationTargets(Node, NodeInstance, DebuggedControlRig->GetHierarchy(), Targets, nullptr))
+							if(const FRigUnit* UnitInstance = UControlRig::GetRigUnitInstanceFromScope(NodeInstance))
 							{
-								return Targets;
+								TArray<FRigDirectManipulationTarget> Targets;
+								if(UnitInstance->GetDirectManipulationTargets(Node, NodeInstance, DebuggedControlRig->GetHierarchy(), Targets, nullptr))
+								{
+									return Targets;
+								}
 							}
 						}
 					}
@@ -4321,7 +4324,7 @@ EVisibility FControlRigBaseEditor::GetPreviewingNodeVisibility() const
 {
 	if (UControlRig* DebuggedControlRig = Cast<UControlRig>(GetRigVMAssetInterface()->GetObjectBeingDebugged()))
 	{
-		if (DebuggedControlRig->GetDebugInfo().GetInstructionForEarlyExit() != INDEX_NONE)
+		if (DebuggedControlRig->GetDebugInfo().GetStepCondition().OriginInstruction != INDEX_NONE)
 		{
 			return EVisibility::Visible;
 		}
@@ -4338,7 +4341,7 @@ void FControlRigBaseEditor::OnPreviewingNodeJumpTo()
 	{
 		return;
 	}
-	const int32 InstructionToJumpTo = DebuggedControlRig->GetDebugInfo().GetInstructionForEarlyExit();
+	const int32 InstructionToJumpTo = DebuggedControlRig->GetDebugInfo().GetStepCondition().OriginInstruction;
 	if (InstructionToJumpTo == INDEX_NONE)
 	{
 		return;
diff --git a/Engine/Plugins/Runtime/RigVM/Source/RigVM/Private/RigVMCore/RigVM.cpp b/Engine/Plugins/Runtime/RigVM/Source/RigVM/Private/RigVMCore/RigVM.cpp
--- a/Engine/Plugins/Runtime/RigVM/Source/RigVM/Private/RigVMCore/RigVM.cpp
+++ b/Engine/Plugins/Runtime/RigVM/Source/RigVM/Private/RigVMCore/RigVM.cpp
@@ -548,6 +548,11 @@ uint32 URigVM::GetVMHash() const
 }
 
 uint32 URigVM::ComputeVMHash() const
+{
+	return ComputeVMHash(ExternalVariables);
+}
+
+uint32 URigVM::ComputeVMHash(const TArray<FRigVMExternalVariableDef>& InExternalVariables) const
 {
 	uint32 Hash = 0;
 	for(const FName& FunctionName : GetFunctionNames())
@@ -557,7 +562,7 @@ uint32 URigVM::ComputeVMHash() const
 
 	Hash = HashCombine(Hash, GetTypeHash(GetByteCode()));
 
-	for(const FRigVMExternalVariableDef& ExternalVariable : ExternalVariables)
+	for(const FRigVMExternalVariableDef& ExternalVariable : InExternalVariables)
 	{
 		Hash = HashCombine(Hash, GetTypeHash(ExternalVariable.GetName().ToString()));
 		Hash = HashCombine(Hash, GetTypeHash(ExternalVariable.GetExtendedCPPType().ToString()));
@@ -565,35 +570,18 @@ uint32 URigVM::ComputeVMHash() const
 
 	Hash = HashCombine(Hash, LiteralMemoryStorage.GetMemoryHash());
 	Hash = HashCombine(Hash, DefaultWorkMemoryStorage.GetMemoryHash());
-	
+
 	return Hash;
 }
 
-UClass* URigVM::GetNativizedClass(const TArray<FRigVMExternalVariableDef>& InExternalVariables)
+uint32 URigVM::ComputeVMHash(const TArray<FRigVMExternalVariable>& InExternalVariables) const
 {
-	TSharedPtr<TGuardValue<TArray<FRigVMExternalVariableDef>>> GuardPtr;
-	if(!InExternalVariables.IsEmpty())
-	{
-		GuardPtr = MakeShareable(new TGuardValue<TArray<FRigVMExternalVariableDef>>(ExternalVariables, InExternalVariables));
-	}
-
-	const uint32 VMHash = GetVMHash();
-	
-	for (TObjectIterator<UClass> ClassIterator; ClassIterator; ++ClassIterator)
-	{
-		if (ClassIterator->IsChildOf(URigVMNativized::StaticClass()) && (*ClassIterator != URigVMNativized::StaticClass()))
-		{
-			if(const URigVM* NativizedVMCDO = ClassIterator->GetDefaultObject<URigVM>())
-			{
-				if(NativizedVMCDO->GetVMHash() == VMHash)
-				{
-					return *ClassIterator;
-				}
-			}
-		}
-	}
+	return ComputeVMHash(RigVMTypeUtils::GetExternalVariableDefs(InExternalVariables));
+}
 
-	return nullptr;
+UClass* URigVM::GetNativizedClass(uint32 InHash) const
+{
+	return URigVMNativized::FindClassForHash(InHash);
 }
 
 bool URigVM::ValidateBytecode()
@@ -781,6 +769,7 @@ const URigVMHost* URigVM::GetHostCDO() const
 void URigVM::Reset_Internal()
 {
 	CachedVMHash = 0;
+	CachedNativizedVMHash.Reset();
 	FunctionNamesStorage.Reset();
 	FunctionsStorage.Reset();
 	FactoriesStorage.Reset();
@@ -1760,17 +1749,7 @@ ERigVMExecuteResult URigVM::ExecuteVM(FRigVMExtendedExecuteContext& Context, con
 
 	if(bIsRootEntry)
 	{
-		if (FRigVMInstructionVisitInfo* RigVMInstructionVisitInfo = Context.GetRigVMInstructionVisitInfo())
-		{
-			if (RigVMInstructionVisitInfo->GetFirstEntryEventInEventQueue() == NAME_None || RigVMInstructionVisitInfo->GetFirstEntryEventInEventQueue() == InEntryName)
-			{
-				RigVMInstructionVisitInfo->SetupInstructionTracking(Instructions.Num(), ByteCode.NumCallables());
-				if (FRigVMProfilingInfo* RigVMProfilingInfo = Context.GetRigVMProfilingInfo())
-				{
-					RigVMProfilingInfo->SetupInstructionTracking(Instructions.Num(), ByteCode.NumCallables(), true);
-				}
-			}
-		}
+		SetupInstructionTracking(Context, Instructions.Num(), ByteCode.NumCallables(), InEntryName);
 	}
 #endif
 
@@ -1798,7 +1777,11 @@ ERigVMExecuteResult URigVM::ExecuteVM(FRigVMExtendedExecuteContext& Context, con
 			return Context.CurrentExecuteResult = ERigVMExecuteResult::Failed;
 		}
 		
-		SetInstructionIndex(Context, ByteCode.GetEntry(EntryIndex).InstructionIndex);
+		if (!SetInstructionIndex(Context, ContextPublicData, ByteCode.GetEntry(EntryIndex).InstructionIndex))
+		{
+			return Context.CurrentExecuteResult = ERigVMExecuteResult::Failed;
+		}
+		
 		EntryIndexToPush = EntryIndex;
 
 		if(bIsRootEntry)
@@ -1818,7 +1801,10 @@ ERigVMExecuteResult URigVM::ExecuteVM(FRigVMExtendedExecuteContext& Context, con
 				check(Callable);
 				FirstInstructionIndex = FMath::Max(FirstInstructionIndex, Callable->LastInstruction + 1);
 			}
-			ContextPublicData.InstructionIndex = FirstInstructionIndex;
+			if (!SetInstructionIndex(Context, ContextPublicData, FirstInstructionIndex))
+			{
+				return Context.CurrentExecuteResult = ERigVMExecuteResult::Failed;
+			}
 		}
 		
 		for(int32 EntryIndex = 0; EntryIndex < ByteCode.NumEntries(); EntryIndex++)
@@ -1887,34 +1873,11 @@ ERigVMExecuteResult URigVM::ExecuteVM(FRigVMExtendedExecuteContext& Context, con
 #if WITH_EDITOR
 	TSharedPtr<FRigVMDebugInfo> RigVMDebugInfo = Context.GetRigVMDebugInfoWeak().Pin();
 
-	FRigVMDebugInfo PreviousDebugInfo;
-
-	if (RigVMDebugInfo)
+	if (bIsRootEntry && RigVMDebugInfo)
 	{
-		if (bIsRootEntry)
-		{
-			PreviousDebugInfo.CopyFrom(*RigVMDebugInfo);
-
-			if (!RigVMDebugInfo->IsEmpty())
-			{
-				// if the instruction we are about to early exit from is outside of the set of instructions
-				// we are about to run - disable the debug info.
-				const int32 EarlyExistInstructionIndex = RigVMDebugInfo->GetInstructionForEarlyExit();
-				if (EarlyExistInstructionIndex != FMath::Clamp(EarlyExistInstructionIndex, ContextPublicData.InstructionIndex, LastInstruction))
-				{
-					RigVMDebugInfo->Reset();
-				}
-				else
-				{
-					LastInstruction = EarlyExistInstructionIndex;
-				}
-			}
-
-			if (!RigVMDebugInfo->IsEmpty())
-			{
-				RigVMDebugInfo->ExecutionHalted().Broadcast(INDEX_NONE, nullptr, InEntryName);
-			}
-		}
+		// The marker needs to make sure to record until it has passed
+		// the origin to be able to step forward, into etc.
+		RigVMDebugInfo->BeginExecution();
 	}
 #endif
 
@@ -1927,17 +1890,14 @@ ERigVMExecuteResult URigVM::ExecuteVM(FRigVMExtendedExecuteContext& Context, con
 		
 		if (RigVMDebugInfo)
 		{
-			if (!RigVMDebugInfo->IsEmpty())
+			if (RigVMDebugInfo->IsActive())
 			{
 				if (Context.CurrentExecuteResult != ERigVMExecuteResult::Failed)
 				{
 					Context.CurrentExecuteResult = ERigVMExecuteResult::Halted;
-					RigVMDebugInfo->ExecutionHalted().Broadcast(RigVMDebugInfo->GetInstructionForEarlyExit(), nullptr, InEntryName);
+					RigVMDebugInfo->ExecutionHalted().Broadcast(RigVMDebugInfo->GetStepCondition().OriginInstruction, nullptr, InEntryName);
 				}
 			}
-
-			// restore it
-			RigVMDebugInfo->CopyFrom(PreviousDebugInfo);
 		}
 	}
 #endif
@@ -1950,17 +1910,28 @@ struct FRigVMCallstackGuard
 public:
 	FRigVMCallstackGuard(FRigVMExtendedExecuteContext& InContext, const FRigVMByteCode& ByteCode, int32 InCallableIndex)
 		: Context(InContext)
-		, PreviousHash(InContext.CallstackHash)
+		, PublicData(InContext.GetPublicData<>())
+		, PreviousInstructionIndex(INDEX_NONE)
+		, PreviousCallstackHash(InContext.CallstackHash)
 		, bIsValid(false)
 	{
-		const int32 InstructionIndex = Context.GetInstructionIndex();
-		bIsValid = InstructionIndex != INDEX_NONE && ByteCode.IsValidCallableIndex(InCallableIndex);
+		PreviousInstructionIndex = Context.GetInstructionIndex();
+		bIsValid = PreviousInstructionIndex != INDEX_NONE && ByteCode.IsValidCallableIndex(InCallableIndex);
 		if (bIsValid)
 		{
 			const FRigVMCallableInfo* Callable = ByteCode.GetCallable(InCallableIndex);
-			Context.Callstack.Emplace(InstructionIndex, Callable);
-			Context.CallstackHash = HashCombine(InstructionIndex, InCallableIndex, Context.CallstackHash);
+			PublicData.InstructionIndex = Callable->FirstInstruction;
+			Context.Callstack.Push({PreviousInstructionIndex, Callable});
+			Context.CallstackHash = HashCombine(PreviousInstructionIndex, InCallableIndex, Context.CallstackHash);
 			Context.BeginCallableSlice(Callable->Index);
+
+#if WITH_EDITOR
+			if (FRigVMInstructionVisitInfo* RigVMInstructionVisitInfo = Context.GetRigVMInstructionVisitInfo())
+			{
+				RigVMInstructionVisitInfo->SetCallableVisitedDuringLastRun(Callable->Index);
+				RigVMInstructionVisitInfo->AddCallableIndexToVisitOrder(Callable->Index);
+			}
+#endif
 		}
 	}
 
@@ -1969,14 +1940,17 @@ public:
 		if (bIsValid)
 		{
 			Context.EndCallableSlice();
-			Context.CallstackHash = PreviousHash;
+			Context.CallstackHash = PreviousCallstackHash;
 			Context.Callstack.Pop();
+			PublicData.InstructionIndex = PreviousInstructionIndex;
 		}
 	}
 		
 private:
 	FRigVMExtendedExecuteContext& Context;
-	uint32 PreviousHash;
+	FRigVMExecuteContext& PublicData;
+	int32 PreviousInstructionIndex;
+	uint32 PreviousCallstackHash;
 	bool bIsValid;
 };
 
@@ -2006,6 +1980,16 @@ ERigVMExecuteResult URigVM::ExecuteInstructions(FRigVMExtendedExecuteContext& Co
 		Context.GetPublicDataSafe<>().Log(EMessageSeverity::Error, InstructionBracket.ErrorMessage);
 		return ERigVMExecuteResult::Failed;
 	}
+
+	TSharedPtr<TGuardValue<FRigVMDebugInfo*>> GuardStrongDebugInfo;
+	const TWeakPtr<FRigVMDebugInfo>& WeakDebugInfo = Context.GetRigVMDebugInfoWeak();
+	if (WeakDebugInfo.IsValid())
+	{
+		if (FRigVMDebugInfo* CurrentDebugInfo = WeakDebugInfo.Pin().Get())
+		{
+			GuardStrongDebugInfo = MakeShared<TGuardValue<FRigVMDebugInfo*>>(Context.DebugInfoStrong, CurrentDebugInfo);
+		}
+	}
 #endif
 	
 	TArray<const FRigVMFunction*>& Functions = GetFunctions();
@@ -2024,12 +2008,14 @@ ERigVMExecuteResult URigVM::ExecuteInstructions(FRigVMExtendedExecuteContext& Co
 	};
 	check(PreCachedMemoryHandles.Num() == MemoryHandleCount);
 	check(PreCachedMemoryHandlesMemoryType.Num() == MemoryHandleCount);
-
+	
 #if WITH_EDITOR
 	TArray<FName>& FunctionNames = GetFunctionNames();
 	FRigVMInstructionVisitInfo* RigVMInstructionVisitInfo = Context.GetRigVMInstructionVisitInfo();
 	FRigVMProfilingInfo* RigVMProfilingInfo = Context.GetRigVMProfilingInfo();
 #endif
+
+	bool bMovingToNextInstructionSucceeded = true;
 	
 	while (Instructions.IsValidIndex(ContextPublicData.InstructionIndex))
 	{
@@ -2131,7 +2117,7 @@ ERigVMExecuteResult URigVM::ExecuteInstructions(FRigVMExtendedExecuteContext& Co
 				}
 #endif
 
-				ContextPublicData.InstructionIndex++;
+				bMovingToNextInstructionSucceeded = IncrementInstructionIndex(Context, ContextPublicData);
 				break;
 			}
 			case ERigVMOpCode::InvokeCallable:
@@ -2143,7 +2129,7 @@ ERigVMExecuteResult URigVM::ExecuteInstructions(FRigVMExtendedExecuteContext& Co
 				// it's possible to have empty callables, so just step over this 
 				if (Callable->GetNumInstructions() == 0)
 				{
-					ContextPublicData.InstructionIndex++;
+					bMovingToNextInstructionSucceeded = IncrementInstructionIndex(Context, ContextPublicData);
 					break;
 				}
 					
@@ -2185,31 +2171,28 @@ ERigVMExecuteResult URigVM::ExecuteInstructions(FRigVMExtendedExecuteContext& Co
 				}
 
 				// invoke the range of the callable
-				const FRigVMCallstackGuard CallstackGuard(Context, ByteCode, Callable->Index);
-#if WITH_EDITOR
-				if (RigVMInstructionVisitInfo != nullptr)
+				ERigVMExecuteResult CallableResult = ERigVMExecuteResult::Failed; 
 				{
-					RigVMInstructionVisitInfo->SetCallableVisitedDuringLastRun(Callable->Index);
-					RigVMInstructionVisitInfo->AddCallableIndexToVisitOrder(Callable->Index);
-				}
-#endif
-				const ERigVMExecuteResult CallableResult = ExecuteInstructions(Context, Callable->FirstInstruction, Callable->LastInstruction);
+					const FRigVMCallstackGuard CallstackGuard(Context, ByteCode, Callable->Index);
+					
+					CallableResult = ExecuteInstructions(Context, Callable->FirstInstruction, Callable->LastInstruction);
 
-				// reset the forwarded handles
-				for (int32 ArgumentIndex = 0; ArgumentIndex < ArgumentCount; ArgumentIndex++)
-				{
-					const FRigVMOperand& ForwardedOperand = Callable->Arguments[ArgumentIndex].ForwardedOperand;
-					FRigVMMemoryStorageStruct* PropertyBag = GetMemoryByType(Context, ForwardedOperand.GetMemoryType());
-					FRigVMForwardedMemoryHandle* ForwardHandle = PropertyBag->GetData<FRigVMForwardedMemoryHandle>(ForwardedOperand.GetRegisterIndex());
-					ForwardHandle->Reset();
+					// reset the forwarded handles
+					for (int32 ArgumentIndex = 0; ArgumentIndex < ArgumentCount; ArgumentIndex++)
+					{
+						const FRigVMOperand& ForwardedOperand = Callable->Arguments[ArgumentIndex].ForwardedOperand;
+						FRigVMMemoryStorageStruct* PropertyBag = GetMemoryByType(Context, ForwardedOperand.GetMemoryType());
+						FRigVMForwardedMemoryHandle* ForwardHandle = PropertyBag->GetData<FRigVMForwardedMemoryHandle>(ForwardedOperand.GetRegisterIndex());
+						ForwardHandle->Reset();
+					}
 				}
-					
+
 				if (CallableResult != ERigVMExecuteResult::Succeeded)
 				{
 					return CallableResult;
 				}
 
-				ContextPublicData.InstructionIndex++;
+				bMovingToNextInstructionSucceeded = IncrementInstructionIndex(Context, ContextPublicData);
 				break;
 			}
 			case ERigVMOpCode::Zero:
@@ -2233,7 +2216,7 @@ ERigVMExecuteResult URigVM::ExecuteInstructions(FRigVMExtendedExecuteContext& Co
 				}
 #endif
 
-				ContextPublicData.InstructionIndex++;
+				bMovingToNextInstructionSucceeded = IncrementInstructionIndex(Context, ContextPublicData);
 				break;
 			}
 			case ERigVMOpCode::BoolFalse:
@@ -2242,7 +2225,7 @@ ERigVMExecuteResult URigVM::ExecuteInstructions(FRigVMExtendedExecuteContext& Co
 				const FRigVMMemoryHandle Handle = GeneratePatchedHandle(Context.ExternalVariableRuntimeData, VMMemoryPtrs, PreCachedMemoryHandlesMemoryType[FirstHandleIndex], PreCachedMemoryHandles[FirstHandleIndex]);
 
 				*((bool*)Handle.GetInputData()) = false;
-				ContextPublicData.InstructionIndex++;
+				bMovingToNextInstructionSucceeded = IncrementInstructionIndex(Context, ContextPublicData);
 				break;
 			}
 			case ERigVMOpCode::BoolTrue:
@@ -2251,7 +2234,7 @@ ERigVMExecuteResult URigVM::ExecuteInstructions(FRigVMExtendedExecuteContext& Co
 				const FRigVMMemoryHandle Handle = GeneratePatchedHandle(Context.ExternalVariableRuntimeData, VMMemoryPtrs, PreCachedMemoryHandlesMemoryType[FirstHandleIndex], PreCachedMemoryHandles[FirstHandleIndex]);
 
 				*((bool*)Handle.GetInputData()) = true;
-				ContextPublicData.InstructionIndex++;
+				bMovingToNextInstructionSucceeded = IncrementInstructionIndex(Context, ContextPublicData);
 				break;
 			}
 			case ERigVMOpCode::Copy:
@@ -2271,7 +2254,7 @@ ERigVMExecuteResult URigVM::ExecuteInstructions(FRigVMExtendedExecuteContext& Co
 				}
 #endif
 					
-				ContextPublicData.InstructionIndex++;
+				bMovingToNextInstructionSucceeded = IncrementInstructionIndex(Context, ContextPublicData);
 				break;
 			}
 			case ERigVMOpCode::Increment:
@@ -2287,7 +2270,7 @@ ERigVMExecuteResult URigVM::ExecuteInstructions(FRigVMExtendedExecuteContext& Co
 					CopyOperandForDebuggingIfNeeded(Context, Op.Arg, Handle);
 				}
 #endif
-				ContextPublicData.InstructionIndex++;
+				bMovingToNextInstructionSucceeded = IncrementInstructionIndex(Context, ContextPublicData);
 				break;
 			}
 			case ERigVMOpCode::Decrement:
@@ -2303,7 +2286,7 @@ ERigVMExecuteResult URigVM::ExecuteInstructions(FRigVMExtendedExecuteContext& Co
 					CopyOperandForDebuggingIfNeeded(Context, Op.Arg, Handle);
 				}
 #endif
-				ContextPublicData.InstructionIndex++;
+				bMovingToNextInstructionSucceeded = IncrementInstructionIndex(Context, ContextPublicData);
 				break;
 			}
 			case ERigVMOpCode::Equals:
@@ -2318,25 +2301,25 @@ ERigVMExecuteResult URigVM::ExecuteInstructions(FRigVMExtendedExecuteContext& Co
 				const bool Result = Instruction.OpCode == ERigVMOpCode::Equals ? Identical : !Identical;
 
 				*((bool*)HandleR.GetOutputData()) = Result;
-				ContextPublicData.InstructionIndex++;
+				bMovingToNextInstructionSucceeded = IncrementInstructionIndex(Context, ContextPublicData);
 				break;
 			}
 			case ERigVMOpCode::JumpAbsolute:
 			{
 				const FRigVMJumpOp& Op = ByteCode.GetOpAt<FRigVMJumpOp>(Instruction);
-				ContextPublicData.InstructionIndex = FRigVMJumpOp::GetTargetInstruction(Instruction.OpCode, ContextPublicData.InstructionIndex, Op.InstructionIndex);
+				bMovingToNextInstructionSucceeded = SetInstructionIndex(Context, ContextPublicData, FRigVMJumpOp::GetTargetInstruction(Instruction.OpCode, ContextPublicData.InstructionIndex, Op.InstructionIndex));
 				break;
 			}
 			case ERigVMOpCode::JumpForward:
 			{
 				const FRigVMJumpOp& Op = ByteCode.GetOpAt<FRigVMJumpOp>(Instruction);
-				ContextPublicData.InstructionIndex = FRigVMJumpOp::GetTargetInstruction(Instruction.OpCode, ContextPublicData.InstructionIndex, Op.InstructionIndex);
+				bMovingToNextInstructionSucceeded = SetInstructionIndex(Context, ContextPublicData, FRigVMJumpOp::GetTargetInstruction(Instruction.OpCode, ContextPublicData.InstructionIndex, Op.InstructionIndex));
 				break;
 			}
 			case ERigVMOpCode::JumpBackward:
 			{
 				const FRigVMJumpOp& Op = ByteCode.GetOpAt<FRigVMJumpOp>(Instruction);
-				ContextPublicData.InstructionIndex = FRigVMJumpOp::GetTargetInstruction(Instruction.OpCode, ContextPublicData.InstructionIndex, Op.InstructionIndex);
+				bMovingToNextInstructionSucceeded = SetInstructionIndex(Context, ContextPublicData, FRigVMJumpOp::GetTargetInstruction(Instruction.OpCode, ContextPublicData.InstructionIndex, Op.InstructionIndex));
 				break;
 			}
 			case ERigVMOpCode::JumpAbsoluteIf:
@@ -2349,11 +2332,11 @@ ERigVMExecuteResult URigVM::ExecuteInstructions(FRigVMExtendedExecuteContext& Co
 				const bool Condition = *(bool*)Handle.GetInputData();
 				if (Condition == Op.Condition)
 				{
-					ContextPublicData.InstructionIndex = FRigVMJumpOp::GetTargetInstruction(Instruction.OpCode, ContextPublicData.InstructionIndex, Op.InstructionIndex);
+					bMovingToNextInstructionSucceeded = SetInstructionIndex(Context, ContextPublicData, FRigVMJumpOp::GetTargetInstruction(Instruction.OpCode, ContextPublicData.InstructionIndex, Op.InstructionIndex));
 				}
 				else
 				{
-					ContextPublicData.InstructionIndex++;
+					bMovingToNextInstructionSucceeded = IncrementInstructionIndex(Context, ContextPublicData);
 				}
 				break;
 			}
@@ -2367,11 +2350,11 @@ ERigVMExecuteResult URigVM::ExecuteInstructions(FRigVMExtendedExecuteContext& Co
 				const bool Condition = *(bool*)Handle.GetInputData();
 				if (Condition == Op.Condition)
 				{
-					ContextPublicData.InstructionIndex = FRigVMJumpOp::GetTargetInstruction(Instruction.OpCode, ContextPublicData.InstructionIndex, Op.InstructionIndex);
+					bMovingToNextInstructionSucceeded = SetInstructionIndex(Context, ContextPublicData, FRigVMJumpOp::GetTargetInstruction(Instruction.OpCode, ContextPublicData.InstructionIndex, Op.InstructionIndex));
 				}
 				else
 				{
-					ContextPublicData.InstructionIndex++;
+					bMovingToNextInstructionSucceeded = IncrementInstructionIndex(Context, ContextPublicData);
 				}
 				break;
 			}
@@ -2385,11 +2368,11 @@ ERigVMExecuteResult URigVM::ExecuteInstructions(FRigVMExtendedExecuteContext& Co
 				const bool Condition = *(bool*)Handle.GetInputData();
 				if (Condition == Op.Condition)
 				{
-					ContextPublicData.InstructionIndex = FRigVMJumpOp::GetTargetInstruction(Instruction.OpCode, ContextPublicData.InstructionIndex, Op.InstructionIndex);
+					bMovingToNextInstructionSucceeded = SetInstructionIndex(Context, ContextPublicData, FRigVMJumpOp::GetTargetInstruction(Instruction.OpCode, ContextPublicData.InstructionIndex, Op.InstructionIndex));
 				}
 				else
 				{
-					ContextPublicData.InstructionIndex++;
+					bMovingToNextInstructionSucceeded = IncrementInstructionIndex(Context, ContextPublicData);
 				}
 				break;
 			}
@@ -2429,13 +2412,13 @@ ERigVMExecuteResult URigVM::ExecuteInstructions(FRigVMExtendedExecuteContext& Co
 				const int32 Count = (*((int32*)CountHandle.GetInputData()));
 				const int32 Index = (*((int32*)IndexHandle.GetInputData()));
 				Context.BeginLoopSlice(Count, Index);
-				ContextPublicData.InstructionIndex++;
+				bMovingToNextInstructionSucceeded = IncrementInstructionIndex(Context, ContextPublicData);
 				break;
 			}
 			case ERigVMOpCode::EndBlock:
 			{
 				Context.EndLoopSlice();
-				ContextPublicData.InstructionIndex++;
+				bMovingToNextInstructionSucceeded = IncrementInstructionIndex(Context, ContextPublicData);
 				break;
 			}
 			case ERigVMOpCode::InvokeEntry:
@@ -2457,7 +2440,7 @@ ERigVMExecuteResult URigVM::ExecuteInstructions(FRigVMExtendedExecuteContext& Co
 					}
 				}
 					
-				ContextPublicData.InstructionIndex++;
+				bMovingToNextInstructionSucceeded = IncrementInstructionIndex(Context, ContextPublicData);
 				break;
 			}
 			case ERigVMOpCode::JumpToBranch:
@@ -2494,7 +2477,7 @@ ERigVMExecuteResult URigVM::ExecuteInstructions(FRigVMExtendedExecuteContext& Co
 						}
 						if(Branch.Label == BranchLabel)
 						{
-							ContextPublicData.InstructionIndex = Branch.FirstInstruction;
+							bMovingToNextInstructionSucceeded = SetInstructionIndex(Context, ContextPublicData, Branch.FirstInstruction);
 							bBranchFound = true;
 							break;
 						}
@@ -2536,7 +2519,7 @@ ERigVMExecuteResult URigVM::ExecuteInstructions(FRigVMExtendedExecuteContext& Co
 					}
 				}
 
-				ContextPublicData.InstructionIndex++;
+				bMovingToNextInstructionSucceeded = IncrementInstructionIndex(Context, ContextPublicData);
 				break;
 			}
 			case ERigVMOpCode::SetupTraits:
@@ -2588,7 +2571,7 @@ ERigVMExecuteResult URigVM::ExecuteInstructions(FRigVMExtendedExecuteContext& Co
 					}
 				}
 
-				ContextPublicData.InstructionIndex++;
+				bMovingToNextInstructionSucceeded = IncrementInstructionIndex(Context, ContextPublicData);
 				break;
 			}
 			case ERigVMOpCode::Invalid:
@@ -2634,6 +2617,12 @@ ERigVMExecuteResult URigVM::ExecuteInstructions(FRigVMExtendedExecuteContext& Co
 				RigVMProfilingInfo->AddOverallCycles(Cycles);
 			}
 		}
+
+		if (!bMovingToNextInstructionSucceeded)
+		{
+			return ERigVMExecuteResult::Halted;
+		}
+		
 #if UE_RIGVM_DEBUG_EXECUTION
 		if (CVarControlRigDebugAllVMExecutions->GetBool() || ContextPublicData.bDebugExecution)
 		{
@@ -3142,6 +3131,49 @@ bool URigVM::CreateLocalizedRegistryIfRequired_NoLock(bool bForce)
 	return true;
 }
 
+FString URigVM::GetLocalizedRegistryAsString() const
+{
+	if (!LocalizedRegistry.IsValid())
+	{
+		return FString();
+	}
+
+	FRigVMObjectArchive Archive;
+	{
+		FRigVMObjectArchiveWriter Writer(Archive, nullptr);
+		LocalizedRegistry->Serialize_NoLock(Writer);
+	}
+	Archive.Compress();
+	return Archive.ToString();
+}
+
+void URigVM::SetLocalizedRegistryFromString(const FString& InData)
+{
+	LocalizedRegistry.Reset();
+	
+	if (InData.IsEmpty())
+	{
+		return;
+	}
+	
+	FRigVMObjectArchive Archive;
+	if (!Archive.SetFromString(InData))
+	{
+		return;
+	}
+
+	if (!Archive.Decompress())
+	{
+		return;
+	}
+
+	LocalizedRegistry = FRigVMRegistry_NoLock::CreateLocalizedRegistry();
+	{
+		FRigVMObjectArchiveReader Reader(Archive, nullptr);
+		LocalizedRegistry->Serialize_NoLock(Reader);
+	}
+}
+
 #if WITH_EDITOR
 
 TArray<FString> URigVM::DumpByteCodeAsTextArray(FRigVMExtendedExecuteContext& Context, const TArray<int32>& InInstructionOrder, bool bIncludeLineNumbers, TFunction<FString(const FString& RegisterName, const FString& RegisterOffsetName)> OperandFormatFunction)
@@ -3523,10 +3555,13 @@ void URigVM::PreCacheSingleMemoryHandle(const FRigVMRegistryHandle& InRegistry,
 	// make sure the handle points to valid memory
 	if (PropertyPath)
 	{
-		uint8* MemoryPtr = (uint8*)Memory->GetContainerPtr() + (PTRINT)Handle.GetOutputData();
-		// don't check the result - since it may be an array element
-		// that doesn't exist yet. 
-		PropertyPath->GetData<uint8>(MemoryPtr, Property);
+		if (!Handle.bIsForwarded)
+		{
+			uint8* MemoryPtr = (uint8*)Memory->GetContainerPtr() + (PTRINT)Handle.GetOutputData();
+			// don't check the result - since it may be an array element
+			// that doesn't exist yet. 
+			PropertyPath->GetData<uint8>(MemoryPtr, Property);
+		}
 	}
 #endif
 }
@@ -3756,9 +3791,16 @@ void URigVM::GetRequiredPlugins(TArray<FString>& OutPlugins)
 	}
 }
 
-void URigVM::SetupInstructionTracking(FRigVMExtendedExecuteContext& Context, int32 InInstructionCount, int32 InCallableCount)
+void URigVM::SetupInstructionTracking(FRigVMExtendedExecuteContext& Context, int32 InInstructionCount, int32 InCallableCount, const FName& InEntryName)
 {
 #if WITH_EDITOR
+	if (FRigVMInstructionVisitInfo* RigVMInstructionVisitInfo = Context.GetRigVMInstructionVisitInfo())
+	{
+		if (RigVMInstructionVisitInfo->GetFirstEntryEventInEventQueue() == NAME_None || RigVMInstructionVisitInfo->GetFirstEntryEventInEventQueue() == InEntryName)
+		{
+			RigVMInstructionVisitInfo->SetupInstructionTracking(InInstructionCount, InCallableCount);
+		}
+	}
 	if (FRigVMProfilingInfo* RigVMProfilingInfo = Context.GetRigVMProfilingInfo())
 	{
 		RigVMProfilingInfo->SetupInstructionTracking(InInstructionCount, InCallableCount, true);
@@ -3774,7 +3816,7 @@ void URigVM::StartProfiling(FRigVMExtendedExecuteContext& Context)
 		RigVMProfilingInfo->StartProfiling(true);
 	}
 #endif
-
+	
 	// if we are going to trace the VM for Rewind Debugger
 	// we should set up the debug memory and prepare the host.
 	TRACE_SETUP_RIGVM_EVALUATION(Context);
diff --git a/Engine/Plugins/Runtime/RigVM/Source/RigVM/Private/RigVMCore/RigVMByteCode.cpp b/Engine/Plugins/Runtime/RigVM/Source/RigVM/Private/RigVMCore/RigVMByteCode.cpp
--- a/Engine/Plugins/Runtime/RigVM/Source/RigVM/Private/RigVMCore/RigVMByteCode.cpp
+++ b/Engine/Plugins/Runtime/RigVM/Source/RigVM/Private/RigVMCore/RigVMByteCode.cpp
@@ -51,7 +51,7 @@ void FRigVMCallableInfo::Serialize(FArchive& Ar)
 {
 	Ar << Index;
 	Ar << Name;
-	Ar << Hash;
+	Ar << FunctionHash;
 	Ar << Arguments;
 	Ar << FirstInstruction;
 	Ar << LastInstruction;
@@ -1081,7 +1081,7 @@ int32 FRigVMByteCode::GetCallableIndex(const uint32& InHash) const
 {
 	return CallableInfos.IndexOfByPredicate([InHash](const FRigVMCallableInfo& InCallableInfo) -> bool
 	{
-		return InCallableInfo.Hash == InHash;
+		return InCallableInfo.FunctionHash == InHash;
 	});
 }
 
@@ -1089,7 +1089,7 @@ const FRigVMCallableInfo* FRigVMByteCode::FindCallable(const uint32& InHash) con
 {
 	return CallableInfos.FindByPredicate([InHash](const FRigVMCallableInfo& InCallableInfo) -> const FRigVMCallableInfo*
 	{
-		if (InCallableInfo.Hash == InHash)
+		if (InCallableInfo.FunctionHash == InHash)
 		{
 			return &InCallableInfo;
 		}
@@ -1109,6 +1109,18 @@ const FRigVMCallableInfo* FRigVMByteCode::FindCallable(const FName& InName) cons
 	});
 }
 
+const FRigVMCallableInfo* FRigVMByteCode::FindCallableForInstruction(int32 InInstructionIndex) const
+{
+	return CallableInfos.FindByPredicate([InInstructionIndex](const FRigVMCallableInfo& InCallableInfo) -> const FRigVMCallableInfo*
+	{
+		if (FMath::IsWithinInclusive(InInstructionIndex, InCallableInfo.FirstInstruction, InCallableInfo.LastInstruction))
+		{
+			return &InCallableInfo;
+		}
+		return nullptr;
+	});
+}
+
 int32 FRigVMByteCode::GetOpNumBytesAt(int32 InByteCodeIndex, bool bIncludeOperands) const
 {
 	const ERigVMOpCode OpCode = GetOpCodeAt(InByteCodeIndex);
@@ -1765,6 +1777,30 @@ int32 FRigVMByteCode::AddPredicateBranch(const FRigVMPredicateBranch& InPredicat
 	return PredicateBranches.Add(InPredicateBranch);
 }
 
+const FRigVMBranchInfo* FRigVMByteCode::FindBranch(int32 InInstructionIndex, int32 InArgumentIndex) const
+{
+	for(const FRigVMBranchInfo& BranchInfo : BranchInfos)
+	{
+		if (BranchInfo.InstructionIndex == InInstructionIndex && BranchInfo.ArgumentIndex == InArgumentIndex)
+		{
+			return &BranchInfo;
+		}
+	}
+	return nullptr;
+}
+
+const FRigVMBranchInfo* FRigVMByteCode::FindBranch(int32 InInstructionIndex, const FName& InLabel) const
+{
+	for(const FRigVMBranchInfo& BranchInfo : BranchInfos)
+	{
+		if (BranchInfo.InstructionIndex == InInstructionIndex && BranchInfo.Label == InLabel)
+		{
+			return &BranchInfo;
+		}
+	}
+	return nullptr;
+}
+
 FRigVMOperandArray FRigVMByteCode::GetOperandsForOp(const FRigVMInstruction& InInstruction) const
 {
 	switch(InInstruction.OpCode)
diff --git a/Engine/Plugins/Runtime/RigVM/Source/RigVM/Private/RigVMCore/RigVMDebugInfo.cpp b/Engine/Plugins/Runtime/RigVM/Source/RigVM/Private/RigVMCore/RigVMDebugInfo.cpp
--- a/Engine/Plugins/Runtime/RigVM/Source/RigVM/Private/RigVMCore/RigVMDebugInfo.cpp
+++ b/Engine/Plugins/Runtime/RigVM/Source/RigVM/Private/RigVMCore/RigVMDebugInfo.cpp
@@ -1,5 +1,244 @@
 ﻿// Copyright Epic Games, Inc. All Rights Reserved.
 
 #include "RigVMCore/RigVMDebugInfo.h"
+#include "RigVMCore/RigVMExecuteContext.h"
 
 #include UE_INLINE_GENERATED_CPP_BY_NAME(RigVMDebugInfo)
+
+bool FRigVMDebugInfo::ShouldStop(const FRigVMExtendedExecuteContext& Context, int32 CurrentInstruction)
+{
+	const int32 CallstackCompare = CompareCallstackPath(Context, StepCondition.OriginCallstackPath);
+	const int32 IterCompare = CompareIterationPath(Context, StepCondition.OriginIterationPath);
+
+	// Update bHasPassedOrigin for stepping modes
+	if (!StepCondition.bHasPassedOrigin)
+	{
+		if (StepCondition.Mode == ERigVMStepMode::Into
+			|| StepCondition.Mode == ERigVMStepMode::Forward
+			|| StepCondition.Mode == ERigVMStepMode::Out)
+		{
+			if (CallstackCompare == 0
+				&& CurrentInstruction == StepCondition.OriginInstruction
+				&& IterCompare == 0)
+			{
+				StepCondition.bHasPassedOrigin = true;
+			}
+		}
+	}
+
+	switch (StepCondition.Mode)
+	{
+	case ERigVMStepMode::None:
+		return false;
+
+	case ERigVMStepMode::Stopped:
+		// Stopped state - halt if we're at the exact same location
+		return CurrentInstruction == StepCondition.OriginInstruction
+			&& CallstackCompare == 0
+			&& IterCompare == 0;
+
+	case ERigVMStepMode::Into:
+		{
+			if (!StepCondition.bHasPassedOrigin)
+			{
+				return false;
+			}
+
+			// Check if we've entered a callable (callstack is deeper)
+			if (CallstackCompare < 0)
+			{
+				TransitionToStopped(Context, CurrentInstruction);
+				return true;
+			}
+
+			// If we haven't entered a callable, check for forward progress at same callstack
+			if (CallstackCompare == 0)
+			{
+				if (IterCompare > 0)
+				{
+					// Advanced in iteration or exited loop - stop
+					TransitionToStopped(Context, CurrentInstruction);
+					return true;
+				}
+
+				if (IterCompare == 0 && CurrentInstruction > StepCondition.OriginInstruction)
+				{
+					// Same iteration, advanced instruction - stop
+					TransitionToStopped(Context, CurrentInstruction);
+					return true;
+				}
+			}
+
+			// Exited a callable (callstack is shallower) - stop at next instruction
+			if (CallstackCompare > 0)
+			{
+				TransitionToStopped(Context, CurrentInstruction);
+				return true;
+			}
+
+			return false;
+		}
+
+	case ERigVMStepMode::Forward:
+		{
+			if (!StepCondition.bHasPassedOrigin)
+			{
+				return false;
+			}
+
+			// Exited a callable (callstack is shallower) - stop
+			if (CallstackCompare > 0)
+			{
+				TransitionToStopped(Context, CurrentInstruction);
+				return true;
+			}
+
+			if (CallstackCompare == 0)
+			{
+				if (IterCompare > 0)
+				{
+					// Advanced in iteration or exited loop - stop
+					TransitionToStopped(Context, CurrentInstruction);
+					return true;
+				}
+
+				if (IterCompare == 0 && CurrentInstruction > StepCondition.OriginInstruction)
+				{
+					// Same iteration, advanced instruction - stop
+					TransitionToStopped(Context, CurrentInstruction);
+					return true;
+				}
+			}
+
+			return false;
+		}
+
+	case ERigVMStepMode::Out:
+		{
+			if (!StepCondition.bHasPassedOrigin)
+			{
+				return false;
+			}
+
+			// We've passed the origin, check if we've exited to a shallower callstack
+			if (CallstackCompare > 0)
+			{
+				TransitionToStopped(Context, CurrentInstruction);
+				return true;
+			}
+
+			return false;
+		}
+
+	case ERigVMStepMode::RunTo:
+		if (CurrentInstruction == StepCondition.TargetInstruction)
+		{
+			TransitionToStopped(Context, CurrentInstruction);
+			return true;
+		}
+		return false;
+	}
+
+	return false;
+}
+
+void FRigVMDebugInfo::TransitionToStopped(const FRigVMExtendedExecuteContext& Context, int32 CurrentInstruction)
+{
+	UE_LOG(LogRigVM, Display, TEXT("<-- Step Condition %s"), *GetStepCondition().ToString());
+
+	StepCondition.Mode = ERigVMStepMode::Stopped;
+	StepCondition.OriginInstruction = CurrentInstruction;
+	StepCondition.TargetInstruction = INDEX_NONE;
+
+	// Capture current callstack and iteration paths
+	CaptureCallstackPath(Context, StepCondition.OriginCallstackPath);
+	CaptureIterationPath(Context, StepCondition.OriginIterationPath);
+
+	UE_LOG(LogRigVM, Display, TEXT("--> Step Condition %s"), *GetStepCondition().ToString());
+}
+
+void FRigVMDebugInfo::CaptureCallstackPath(const FRigVMExtendedExecuteContext& Context, TArray<int32, TInlineAllocator<16>>& OutPath)
+{
+	OutPath.Reset();
+	const FRigVMExecuteCallstack& Callstack = Context.GetCallstack();
+	for (int32 i = 0; i < Callstack.Num(); i++)
+	{
+		OutPath.Add(Callstack[i].InstructionIndex);
+	}
+}
+
+int32 FRigVMDebugInfo::CompareCallstackPath(const FRigVMExtendedExecuteContext& Context, const TArray<int32, TInlineAllocator<16>>& InPath)
+{
+	const FRigVMExecuteCallstack& Callstack = Context.GetCallstack();
+	const int32 MinLen = FMath::Min(Callstack.Num(), InPath.Num());
+
+	// Compare common prefix
+	for (int32 i = 0; i < MinLen; i++)
+	{
+		const int32 CurrentInstr = Callstack[i].InstructionIndex;
+		const int32 OriginInstr = InPath[i];
+
+		if (CurrentInstr > OriginInstr)
+			return +1;  // Current is past origin at this level
+		if (CurrentInstr < OriginInstr)
+			return -1;  // Current is before origin at this level
+	}
+
+	// Prefixes match - compare by depth
+	if (Callstack.Num() > InPath.Num())
+		return -1;  // Current is deeper (entered a callable)
+	if (Callstack.Num() < InPath.Num())
+		return +1;  // Current is shallower (exited a callable)
+
+	return 0;  // Exact match
+}
+
+void FRigVMDebugInfo::CaptureIterationPath(const FRigVMExtendedExecuteContext& Context, TArray<int32, TInlineAllocator<16>>& OutPath)
+{
+	OutPath.Reset();
+	for (const FRigVMSlice& Slice : Context.Slices)
+	{
+		if (Slice.IsLoop())
+		{
+			OutPath.Add(Slice.GetRelativeIndex());
+		}
+	}
+}
+
+int32 FRigVMDebugInfo::CompareIterationPath(const FRigVMExtendedExecuteContext& Context, const TArray<int32, TInlineAllocator<16>>& InPath)
+{
+	int32 PathIndex = 0;
+	for (const FRigVMSlice& Slice : Context.Slices)
+	{
+		if (Slice.IsLoop())
+		{
+			if (PathIndex >= InPath.Num())
+			{
+				// Origin path is shorter - we've compared all we can
+				// Current is deeper in loops, but at same point up to origin's depth
+				break;
+			}
+
+			const int32 CurrentIter = Slice.GetRelativeIndex();
+			const int32 OriginIter = InPath[PathIndex];
+
+			if (CurrentIter > OriginIter)
+				return +1;  // Advanced in this loop
+			if (CurrentIter < OriginIter)
+				return -1;  // Before origin in this loop
+
+			PathIndex++;
+		}
+	}
+
+	// If we get here, all compared elements matched
+	// PathIndex tells us how many loops we compared
+
+	if (PathIndex < InPath.Num())
+	{
+		// Current has fewer loops than origin - we've exited a loop scope
+		return +1;
+	}
+
+	return 0;  // Same path (up to min length)
+}
diff --git a/Engine/Plugins/Runtime/RigVM/Source/RigVM/Private/RigVMCore/RigVMExecuteContext.cpp b/Engine/Plugins/Runtime/RigVM/Source/RigVM/Private/RigVMCore/RigVMExecuteContext.cpp
--- a/Engine/Plugins/Runtime/RigVM/Source/RigVM/Private/RigVMCore/RigVMExecuteContext.cpp
+++ b/Engine/Plugins/Runtime/RigVM/Source/RigVM/Private/RigVMCore/RigVMExecuteContext.cpp
@@ -190,6 +190,42 @@ void FRigVMExtendedExecuteContext::Initialize(const UScriptStruct* InScriptStruc
 	SetDefaultNameCache();
 }
 
+bool FRigVMExtendedExecuteContext::StepForward()
+{
+#if WITH_EDITOR
+	if (TSharedPtr<FRigVMDebugInfo> CurrentDebugInfo = DebugInfoWeak.Pin())
+	{
+		CurrentDebugInfo->StepForward();
+		return true;
+	}
+#endif
+	return false;
+}
+
+bool FRigVMExtendedExecuteContext::StepInto()
+{
+#if WITH_EDITOR
+	if (TSharedPtr<FRigVMDebugInfo> CurrentDebugInfo = DebugInfoWeak.Pin())
+	{
+		CurrentDebugInfo->StepInto();
+		return true;
+	}
+#endif
+	return false;
+}
+
+bool FRigVMExtendedExecuteContext::StepOut()
+{
+#if WITH_EDITOR
+	if (TSharedPtr<FRigVMDebugInfo> CurrentDebugInfo = DebugInfoWeak.Pin())
+	{
+		CurrentDebugInfo->StepOut();
+		return true;
+	}
+#endif
+	return false;
+}
+
 int32 FRigVMExtendedExecuteContext::GetInstructionIndex() const
 {
 	return GetPublicData<>().InstructionIndex;
diff --git a/Engine/Plugins/Runtime/RigVM/Source/RigVM/Private/RigVMCore/RigVMFunction.cpp b/Engine/Plugins/Runtime/RigVM/Source/RigVM/Private/RigVMCore/RigVMFunction.cpp
--- a/Engine/Plugins/Runtime/RigVM/Source/RigVM/Private/RigVMCore/RigVMFunction.cpp
+++ b/Engine/Plugins/Runtime/RigVM/Source/RigVM/Private/RigVMCore/RigVMFunction.cpp
@@ -121,6 +121,18 @@ FString FRigVMFunction::GetModuleRelativeHeaderPath() const
 	return FString();
 }
 
+int32 FRigVMFunction::GetArgumentIndex(const FName& InArgumentName) const
+{
+	for (int32 ArgumentIndex = 0; ArgumentIndex < Arguments.Num(); ArgumentIndex++)
+	{
+		if (Arguments[ArgumentIndex].Name == InArgumentName)
+		{
+			return ArgumentIndex;
+		}
+	}
+	return INDEX_NONE;
+}
+
 const TArray<TRigVMTypeIndex>& FRigVMFunction::GetArgumentTypeIndices() const
 {
 	FRigVMRegistryWriteLock Lock;
@@ -412,3 +424,39 @@ void FRigVMFunction::Serialize_NoLock(FArchive& Ar, FRigVMRegistryHandle& InRegi
 		}
 	}
 }
+
+#if WITH_EDITOR
+
+bool FRigVMFunction::IsFixedSizeArray(const FName& InArgumentName) const
+{
+	if (Factory)
+	{
+		static const FString TrueString = TEXT("True");
+		return Factory->GetArgumentMetaData(InArgumentName, FRigVMStruct::FixedSizeArrayMetaName) == TrueString;
+	}
+	if (Struct)
+	{
+		const FProperty* Property = Struct->FindPropertyByName(InArgumentName);
+		check(Property);
+		return Property->HasMetaData(FRigVMStruct::FixedSizeArrayMetaName); 
+	}
+	return false;
+}
+
+bool FRigVMFunction::IsLazyValue(const FName& InArgumentName) const
+{
+	if (Factory)
+	{
+		static const FString TrueString = TEXT("True");
+		return Factory->GetArgumentMetaData(InArgumentName, FRigVMStruct::ComputeLazilyMetaName) == TrueString;
+	}
+	if (Struct)
+	{
+		const FProperty* Property = Struct->FindPropertyByName(InArgumentName);
+		check(Property);
+		return Property->HasMetaData(FRigVMStruct::ComputeLazilyMetaName); 
+	}
+	return false;
+}
+
+#endif
\ No newline at end of file
diff --git a/Engine/Plugins/Runtime/RigVM/Source/RigVM/Private/RigVMCore/RigVMMemoryStorage.cpp b/Engine/Plugins/Runtime/RigVM/Source/RigVM/Private/RigVMCore/RigVMMemoryStorage.cpp
--- a/Engine/Plugins/Runtime/RigVM/Source/RigVM/Private/RigVMCore/RigVMMemoryStorage.cpp
+++ b/Engine/Plugins/Runtime/RigVM/Source/RigVM/Private/RigVMCore/RigVMMemoryStorage.cpp
@@ -89,9 +89,9 @@ void FRigVMBranchInfo::Serialize(FArchive& Ar)
 
 ERigVMExecuteResult FRigVMLazyBranch::Execute(FRigVMExtendedExecuteContext& Context)
 {
-	check(VM);
 	if(BranchInfo.IsValid())
 	{
+		check(VM);
 		return VM->ExecuteBranch(Context, BranchInfo);
 	}
 
@@ -101,10 +101,9 @@ ERigVMExecuteResult FRigVMLazyBranch::Execute(FRigVMExtendedExecuteContext& Cont
 
 ERigVMExecuteResult FRigVMLazyBranch::ExecuteIfRequired(FRigVMExtendedExecuteContext& Context, uint32 InSliceHash)
 {
-	check(VM);
-
 	FRigVMInstructionSetExecuteState& ExecuteState = Context.LazyBranchExecuteState[BranchInfo.Index];
-	if(ExecuteState.RequiresExecute(InSliceHash, Context.GetHashForLazyBranch()))
+	const uint32 HashForLazyBranch = Context.GetHashForLazyBranch();
+	if(ExecuteState.RequiresExecute(InSliceHash, HashForLazyBranch))
 	{
 		return Execute(Context);
 	}
diff --git a/Engine/Plugins/Runtime/RigVM/Source/RigVM/Private/RigVMCore/RigVMMemoryStorageStruct.cpp b/Engine/Plugins/Runtime/RigVM/Source/RigVM/Private/RigVMCore/RigVMMemoryStorageStruct.cpp
--- a/Engine/Plugins/Runtime/RigVM/Source/RigVM/Private/RigVMCore/RigVMMemoryStorageStruct.cpp
+++ b/Engine/Plugins/Runtime/RigVM/Source/RigVM/Private/RigVMCore/RigVMMemoryStorageStruct.cpp
@@ -32,6 +32,12 @@ FRigVMMemoryStorageStruct::~FRigVMMemoryStorageStruct()
 {
 }
 
+void FRigVMMemoryStorageStruct::InitializeFromNativeStruct(const UScriptStruct* InStruct)
+{
+	Value.InitializeAs(InStruct);
+	RefreshLinkedProperties();
+}
+
 bool FRigVMMemoryStorageStruct::Serialize(FArchive& Ar)
 {
 	UE_RIGVM_ARCHIVETRACE_SCOPE(Ar, TEXT("FRigVMMemoryStorageStruct"));
@@ -574,6 +580,10 @@ bool FRigVMMemoryStorageStruct::CopyProperty(
 		InSourceHandle.GetPropertyPathRef());
 }
 
+const UStruct* FRigVMMemoryStorageStruct::GetStruct() const
+{
+	return Value.IsValid() ? Value.GetScriptStruct() : nullptr;
+}
 
 //****************************************************************************
 
@@ -581,9 +591,9 @@ void FRigVMMemoryStorageStruct::RefreshLinkedProperties()
 {
 	LinkedProperties.Reset();
 
-	if (const UPropertyBag* CurrentBagStruct = GetPropertyBagStruct())
+	if (const UStruct* Struct = GetStruct())
 	{
-		const FProperty* Property = CastField<FProperty>(CurrentBagStruct->ChildProperties);
+		const FProperty* Property = CastField<FProperty>(Struct->ChildProperties);
 		while (Property)
 		{
 			LinkedProperties.Add(Property);
@@ -615,7 +625,9 @@ void FRigVMMemoryStorageStruct::RefreshPropertyPaths()
 #if WITH_EDITOR
 			if (!PropertyPaths[PropertyPathIndex].IsValid())
 			{
+#if UE_BUILD_DEBUG
 				UE_LOG(LogRigVM, Warning, TEXT("Unable to resolve property path '%s' for property '%s'."), *PropertyPathDescriptions[PropertyPathIndex].SegmentPath, *LinkedProperties[PropertyIndex]->GetPathName());
+#endif
 			}
 #endif
 		}
diff --git a/Engine/Plugins/Runtime/RigVM/Source/RigVM/Private/RigVMCore/RigVMNativized.cpp b/Engine/Plugins/Runtime/RigVM/Source/RigVM/Private/RigVMCore/RigVMNativized.cpp
--- a/Engine/Plugins/Runtime/RigVM/Source/RigVM/Private/RigVMCore/RigVMNativized.cpp
+++ b/Engine/Plugins/Runtime/RigVM/Source/RigVM/Private/RigVMCore/RigVMNativized.cpp
@@ -1,12 +1,18 @@
 // Copyright Epic Games, Inc. All Rights Reserved.
 
 #include "RigVMCore/RigVMNativized.h"
+#include "RigVMModule.h"
 
 #include UE_INLINE_GENERATED_CPP_BY_NAME(RigVMNativized)
 
+FRigVMNativizedContext::FRigVMNativizedContext(FRigVMExtendedExecuteContext& InExecuteContext, FRigVMExecuteContext& InPublicContext)
+: ExecuteContext(InExecuteContext)
+, PublicContextBase(InPublicContext)
+{
+}
+
 URigVMNativized::URigVMNativized()
 	: URigVM()
-	, TemporaryArrayIndex(INDEX_NONE)
 {
 }
 
@@ -14,16 +20,34 @@ URigVMNativized::~URigVMNativized()
 {
 }
 
+UClass* URigVMNativized::FindClassForHash(uint32 InVMHash)
+{
+	if (UClass** ClassPtr = GetNativizedClassMap().Find(InVMHash))
+	{
+		return *ClassPtr;
+	}
+	return nullptr;
+}
+
+void URigVMNativized::RegisterClassForHash(uint32 InVMHash, UClass* InClass)
+{
+	GetNativizedClassMap().FindOrAdd(InVMHash, InClass);
+}
+
+TMap<uint32, UClass*>& URigVMNativized::GetNativizedClassMap()
+{
+	static TMap<uint32, UClass*> Lookup;
+	return Lookup;
+}
+
 void URigVMNativized::Serialize(FArchive& Ar)
 {
 	UObject::Serialize(Ar);
 }
 
-void URigVMNativized::Reset(bool IsIgnoringArchetypeRef)
+void URigVMNativized::Reset(FRigVMExtendedExecuteContext& Context)
 {
 	// don't call super on purpose
-	LazyMemoryHandles.Reset();
-	LazyMemoryBranches.Reset();
 	ByteCodeStorage.Reset();
 }
 
@@ -39,50 +63,96 @@ ERigVMExecuteResult URigVMNativized::ExecuteVM(FRigVMExtendedExecuteContext& Con
 	return ERigVMExecuteResult::Failed;
 }
 
+FRigVMMemoryStorageStruct* URigVMNativized::GetMemoryByType(FRigVMExtendedExecuteContext& Context, ERigVMMemoryType InMemoryType)
+{
+	return Super::GetMemoryByType(Context, InMemoryType);
+}
+
+const FRigVMMemoryStorageStruct* URigVMNativized::GetMemoryByType(const FRigVMExtendedExecuteContext& Context, ERigVMMemoryType InMemoryType) const
+{
+	return Super::GetMemoryByType(Context, InMemoryType);
+}
+
 const FRigVMInstructionArray& URigVMNativized::GetInstructions()
 {
-	if(GetByteCode().GetNumInstructions() > 0)
+	if(URigVM* SourceVM = GetSourceVM())
 	{
-		return Super::GetInstructions();
+		return SourceVM->GetInstructions();
 	}
 	static const FRigVMInstructionArray EmptyInstructions;
 	return EmptyInstructions;
 }
 
-#if WITH_EDITOR
+const FRigVMByteCode& URigVMNativized::GetByteCode() const
+{
+	if(const URigVM* SourceVM = GetSourceVM())
+	{
+		return SourceVM->GetByteCode();
+	}
+	return Super::GetByteCode();
+}
 
-void URigVMNativized::SetByteCode(const FRigVMByteCode& InByteCode)
+const TArray<const FRigVMFunction*>& URigVMNativized::GetFunctions() const
 {
-	ByteCodeStorage = InByteCode;
-	ByteCodePtr = &ByteCodeStorage;
+	if(const URigVM* SourceVM = GetSourceVM())
+	{
+		return SourceVM->GetFunctions();
+	}
+	return Super::GetFunctions();
 }
 
-#endif
+const TArray<FName>& URigVMNativized::GetFunctionNames() const
+{
+	if(const URigVM* SourceVM = GetSourceVM())
+	{
+		return SourceVM->GetFunctionNames();
+	}
+	return Super::GetFunctionNames();
+}
 
-const FRigVMMemoryHandle& URigVMNativized::GetLazyMemoryHandle(int32 InIndex, uint8* InMemory, const FProperty* InProperty, TFunction<ERigVMExecuteResult()> InLambda)
+const TArray<const FRigVMDispatchFactory*>& URigVMNativized::GetFactories() const
 {
-	check(InIndex != INDEX_NONE);
-	check(LazyMemoryHandles.IsValidIndex(InIndex));
+	if(const URigVM* SourceVM = GetSourceVM())
+	{
+		return SourceVM->GetFactories();
+	}
+	return Super::GetFactories();
+}
 
-	FRigVMMemoryHandle& Handle = LazyMemoryHandles[InIndex];
-	if(!Handle.IsLazy())
+const FRigVMFunction* URigVMNativized::FindDispatchFunction(const FString& InIdentifier) const
+{
+	if (LocalizedRegistry.IsValid())
 	{
-		check(LazyMemoryBranches.IsValidIndex(InIndex));
-		FRigVMLazyBranch& LazyBranch = LazyMemoryBranches[InIndex];
-		LazyBranch.VM = this;
-		LazyBranch.FunctionPtr = InLambda;
+		return LocalizedRegistry->FindFunction_NoLock(*InIdentifier);
+	}
+	return FRigVMRegistry::Get().FindFunction(*InIdentifier);
+}
 
-		Handle = FRigVMMemoryHandle(InMemory, InProperty, nullptr, &LazyBranch);
+URigVM* URigVMNativized::GetSourceVM() const
+{
+#if WITH_EDITOR
+	if (WeakSourceVM.IsValid())
+	{
+		return WeakSourceVM.Get();
 	}
+#endif
+	return nullptr;
+}
+
+#if WITH_EDITOR
 
-	return Handle;
+void URigVMNativized::SetSourceVM(URigVM* InSourceVM)
+{
+	WeakSourceVM = InSourceVM;
 }
 
-void URigVMNativized::AllocateLazyMemoryHandles(int32 InCount)
+void URigVMNativized::FDefaultValueImportErrorPipe::Serialize(const TCHAR* V, ELogVerbosity::Type Verbosity, const class FName& Category)
 {
-	if(LazyMemoryHandles.Num() != InCount)
+	if (Verbosity == ELogVerbosity::Error || Verbosity == ELogVerbosity::Warning)
 	{
-		LazyMemoryHandles.SetNum(InCount);
-		LazyMemoryBranches.SetNum(InCount);
+		UE_LOG(LogRigVM, Error, TEXT("%s: Error Importing Default Value: %s"), *Owner, V);
+		NumErrors++;
 	}
 }
+
+#endif
diff --git a/Engine/Plugins/Runtime/RigVM/Source/RigVM/Private/RigVMCore/RigVMObjectArchive.cpp b/Engine/Plugins/Runtime/RigVM/Source/RigVM/Private/RigVMCore/RigVMObjectArchive.cpp
--- a/Engine/Plugins/Runtime/RigVM/Source/RigVM/Private/RigVMCore/RigVMObjectArchive.cpp
+++ b/Engine/Plugins/Runtime/RigVM/Source/RigVM/Private/RigVMCore/RigVMObjectArchive.cpp
@@ -4,6 +4,8 @@
 #include "Misc/Compression.h"
 #include "UObject/UE5MainStreamObjectVersion.h"
 #include "RigVMObjectVersion.h"
+#include "RigVMStringUtils.h"
+#include "Misc/Base64.h"
 
 #include UE_INLINE_GENERATED_CPP_BY_NAME(RigVMObjectArchive)
 
@@ -64,11 +66,11 @@ void FRigVMObjectArchive::Compress()
 
 }
 
-void FRigVMObjectArchive::Decompress()
+bool FRigVMObjectArchive::Decompress()
 {
 	if(!IsCompressed())
 	{
-		return;
+		return true;
 	}
 
 	CompressedSize = Buffer.Num();
@@ -80,7 +82,94 @@ void FRigVMObjectArchive::Decompress()
 	{
 		Buffer = MoveTemp(UncompressedBuffer);
 		bIsCompressed = false;
+		return true;
+	}
+	return false;
+}
+
+FString FRigVMObjectArchive::ToString()
+{
+	if (IsEmpty())
+	{
+		return FString();
+	}
+	
+	Compress();
+
+	TArray<FString> VersionStrings;
+	for (const TPair<FGuid,int32>& Pair : Versions)
+	{
+		VersionStrings.Add(FString::Printf(TEXT("%s=%d"), *Pair.Key.ToString(EGuidFormats::Base36Encoded), Pair.Value));
+	}
+	const FString VersionsString = RigVMStringUtils::JoinStrings(VersionStrings, TEXT("|"));
+	const FString Base64Header = FString::Printf(TEXT("%d,%d,%s,Base64:"), CompressedSize, UncompressedSize, *VersionsString);
+	return Base64Header + FBase64::Encode(Buffer.GetData(), Buffer.Num());
+}
+
+bool FRigVMObjectArchive::SetFromString(const FString& InString)
+{
+	Reset();
+
+	check(!InString.IsEmpty());
+	check(InString.Contains(TEXT(",Base64:"), ESearchCase::CaseSensitive));
+
+	FStringView CompleteView(InString);
+	int32 Separator = INDEX_NONE;
+	if(!CompleteView.FindChar(TEXT(':'), Separator))
+	{
+		Reset();
+		return false;
 	}
+	FStringView Header = CompleteView.Left(Separator);
+	FStringView Data = CompleteView.Mid(Separator + 1);
+
+	if(!Header.FindChar(TEXT(','), Separator))
+	{
+		Reset();
+		return false;
+	}
+	CompressedSize = FCString::Atoi(*FString(Header.Left(Separator)));
+	Header = Header.Mid(Separator + 1);
+	if(!Header.FindChar(TEXT(','), Separator))
+	{
+		Reset();
+		return false;
+	}
+	UncompressedSize = FCString::Atoi(*FString(Header.Left(Separator)));
+	Header = Header.Mid(Separator + 1);
+	if(!Header.FindChar(TEXT(','), Separator))
+	{
+		Reset();
+		return false;
+	}
+	const FString VersionsString = FString(Header.Left(Separator));
+	TArray<FString> VersionStrings;
+	RigVMStringUtils::SplitString(VersionsString, TEXT("|"), VersionStrings);
+	for (const FString& VersionString: VersionStrings)
+	{
+		FString Guid,Value;
+		if(!VersionString.Split(TEXT("="), &Guid, &Value))
+		{
+			Reset();
+			return false;
+		}
+		Versions.Add(FGuid(Guid), FCString::Atoi(*Value));
+	}
+	Header = Header.Mid(Separator + 1);
+	if(Header != TEXT("Base64"))
+	{
+		Reset();
+		return false;		
+	}
+
+	if (!FBase64::Decode(FString(Data), Buffer))
+	{
+		Reset();
+		return false;
+	}
+
+	bIsCompressed = CompressedSize != INDEX_NONE && CompressedSize < UncompressedSize;
+	return true;
 }
 
 void FRigVMObjectArchive::SetVersionsFromArchive(const FArchive& Ar)
diff --git a/Engine/Plugins/Runtime/RigVM/Source/RigVM/Private/RigVMCore/RigVMPropertyPath.cpp b/Engine/Plugins/Runtime/RigVM/Source/RigVM/Private/RigVMCore/RigVMPropertyPath.cpp
--- a/Engine/Plugins/Runtime/RigVM/Source/RigVM/Private/RigVMCore/RigVMPropertyPath.cpp
+++ b/Engine/Plugins/Runtime/RigVM/Source/RigVM/Private/RigVMCore/RigVMPropertyPath.cpp
@@ -131,6 +131,59 @@ FRigVMPropertyPath::FRigVMPropertyPath(const FProperty* InProperty, const FStrin
 	}
 }
 
+FString FRigVMPropertyPath::GetName() const
+{
+	TArray<FString> Parts;
+	Parts.Reset(Segments.Num() + 2);
+	Parts.Add(TEXT("PropertyPath"));
+	Parts.Add(GetHeadCPPType().Replace(TEXT("<"), TEXT("_")).Replace(TEXT(">"), TEXT("_")));
+	for(const FRigVMPropertyPathSegment& Segment : Segments)
+	{
+		if (Segment.Type == ERigVMPropertyPathSegmentType::ArrayElement)
+		{
+			Parts.Add(FString::FromInt(Segment.Index));
+		}
+		else
+		{
+			Parts.Add(Segment.Name.ToString());
+		}
+	}
+	return RigVMStringUtils::JoinStrings(Parts, TEXT("_"));
+}
+
+FString FRigVMPropertyPath::GetHeadCPPType() const
+{
+	if(Segments.IsEmpty())
+	{
+		return FString();
+	}
+	
+	const FRigVMPropertyPathSegment& FirstSegment = Segments[0];
+	switch(FirstSegment.Type)
+	{
+		case ERigVMPropertyPathSegmentType::ArrayElement:
+		case ERigVMPropertyPathSegmentType::MapValue:
+		{
+			// for array elements map values the segment refers to the array itself
+			FString ExtendedType;
+			FString BaseType = FirstSegment.Property->GetCPPType(&ExtendedType);  
+			return BaseType + ExtendedType;
+		}
+		case ERigVMPropertyPathSegmentType::StructMember:
+		{
+			UScriptStruct* OwnerStruct = Cast<UScriptStruct>(FirstSegment.Property->GetOwnerStruct());
+			check(OwnerStruct);
+			return OwnerStruct->GetStructCPPName();
+		}
+		default:
+		{
+			checkNoEntry();
+			break;
+		}
+	}
+	return FString();
+}
+
 const FProperty* FRigVMPropertyPath::GetTailProperty() const
 {
 	if(Segments.IsEmpty())
diff --git a/Engine/Plugins/Runtime/RigVM/Source/RigVM/Private/RigVMHost.cpp b/Engine/Plugins/Runtime/RigVM/Source/RigVM/Private/RigVMHost.cpp
--- a/Engine/Plugins/Runtime/RigVM/Source/RigVM/Private/RigVMHost.cpp
+++ b/Engine/Plugins/Runtime/RigVM/Source/RigVM/Private/RigVMHost.cpp
@@ -27,8 +27,8 @@
 // CVar to disable all rigvm execution 
 static TAutoConsoleVariable<int32> CVarRigVMDisableExecutionAll(TEXT("RigVM.DisableExecutionAll"), 0, TEXT("if nonzero we disable all execution of rigvms."));
 
-// CVar to disable swapping to nativized vms 
-static TAutoConsoleVariable<int32> CVarRigVMDisableNativizedVMs(TEXT("RigVM.DisableNativizedVMs"), 1, TEXT("if nonzero we disable swapping to nativized VMs."));
+// CVar to enable swapping to nativized vms 
+static TAutoConsoleVariable<int32> CVarRigVMEnableNativizedVMs(TEXT("RigVM.EnableNativizedVMs"), 0, TEXT("if nonzero we enable swapping to nativized VMs."));
 
 URigVMHost::URigVMHost(const FObjectInitializer& ObjectInitializer)
 	: Super(ObjectInitializer)
@@ -38,6 +38,7 @@ URigVMHost::URigVMHost(const FObjectInitializer& ObjectInitializer)
 	, bAccumulateTime(true)
 #if WITH_EDITOR
 	, bIsBeingDebugged(false)
+	, bEnableNativizedVMs(true)
 #endif
 #if RIGVM_TRACE_ENABLED
 	, bIsPlayingRewindDebugTrace(false)
@@ -213,9 +214,6 @@ void URigVMHost::BeginDestroy()
 
 	InitializedEvent.Clear();
 	ExecutedEvent.Clear();
-#if WITH_EDITOR
-	DebugInfo->Reset();
-#endif
 
 	TRACE_OBJECT_LIFETIME_END(this);
 }
@@ -294,9 +292,12 @@ bool URigVMHost::InitializeVM(const FName& InEventName)
 	FRigVMExtendedExecuteContext& ExtendedExecuteContext = GetRigVMExtendedExecuteContext();
 
 	const TArray<FRigVMExternalVariable> ExternalVariables = GetExternalVariablesImpl(false);
-	if (VM->GetExternalVariableDefs().Num() != ExternalVariables.Num())
+	if (!VM->IsNativized())
 	{
-		return false;	// The rig did compile with errors
+		if (VM->GetExternalVariableDefs().Num() != ExternalVariables.Num())
+		{
+			return false;	// The rig did compile with errors
+		}
 	}
 
 	// update the VM's external variables
@@ -604,12 +605,7 @@ bool URigVMHost::Execute(const FName& InEventName)
 		RuntimeAssetInterface = GetClass();
 	}
 
-	if (RuntimeAssetInterface)
-	{
-		DebugInfo->CopyFrom(RuntimeAssetInterface->GetDebugInfo());
-	}
-
-	ExtendedExecuteContext.SetDebugInfo(DebugInfo);
+	ExtendedExecuteContext.SetDebugInfo(DebugInfo, false);
 	
 	if (IsProfilingEnabled())
 	{
@@ -1060,44 +1056,55 @@ void URigVMHost::SwapVMToNativizedIfRequired(UClass* InNativizedClass)
 	{
 		return;
 	}
+	
 	if(VM == nullptr)
 	{
 		return;
 	}
 
-	const bool bNativizedVMDisabled = AreNativizedVMsDisabled();
+	const bool bNativizedVMEnabled = AreNativizedVMsEnabled();
 
 	// GetNativizedClass can be pretty costly, let's try to skip this if it is not absolutely necessary
-	if(InNativizedClass == nullptr && !bNativizedVMDisabled)
+	if(InNativizedClass == nullptr && bNativizedVMEnabled)
 	{
-		if(!HasAnyFlags(RF_ClassDefaultObject))
+		const uint32 Hash = VM->GetVMHash(); 
+		if (LastNativizedVMHash.Get(0) == Hash && WeakNativizedVMClass.IsSet())
 		{
-			if(URigVM* SourceVM = GetSourceVM())
+			InNativizedClass = WeakNativizedVMClass.GetValue().Get();
+		}
+		else
+		{
+			const TArray<FRigVMExternalVariableDef> ExternalVariableDefs = RigVMTypeUtils::GetExternalVariableDefs(GetExternalVariables());
+			if(!HasAnyFlags(RF_ClassDefaultObject))
 			{
-				if(SourceVM->IsNativized())
-				{
-					InNativizedClass = SourceVM->GetClass();
-				}
-				else
+				if(URigVM* SourceVM = GetSourceVM())
 				{
-					InNativizedClass = SourceVM->GetNativizedClass(RigVMTypeUtils::GetExternalVariableDefs(GetExternalVariables()));
+					if(SourceVM->IsNativized())
+					{
+						InNativizedClass = SourceVM->GetClass();
+					}
+					else
+					{
+						InNativizedClass = SourceVM->GetNativizedClass(Hash);
+					}
 				}
 			}
-		}
-		else
-		{
-			InNativizedClass = VM->GetNativizedClass(RigVMTypeUtils::GetExternalVariableDefs(GetExternalVariables()));
+			else
+			{
+				InNativizedClass = VM->GetNativizedClass(Hash);
+			}
+
+			LastNativizedVMHash = Hash;
+			WeakNativizedVMClass = InNativizedClass;
 		}
 	}
 
 	if(VM->IsNativized())
 	{
-		if((InNativizedClass == nullptr) || bNativizedVMDisabled)
+		if((InNativizedClass == nullptr) || !bNativizedVMEnabled)
 		{
-			const EObjectFlags PreviousFlags = VM->GetFlags();
-			VM->Rename(nullptr, GetTransientPackage(), REN_DoNotDirty | REN_DontCreateRedirectors | REN_NonTransactional);
-			VM->MarkAsGarbage();
-			VM = NewObject<URigVM>(this, TEXT("RigVM_NVMA"), PreviousFlags);
+			VM = GetSourceVM();
+			GetRigVMExtendedExecuteContext().ResetNumExecutions();
 #if UE_RIGVM_PROFILE_EXECUTE_UNITS_NUM
 			ProfilingRunsLeft = 0;
 			AccumulatedCycles = 0;
@@ -1106,42 +1113,65 @@ void URigVMHost::SwapVMToNativizedIfRequired(UClass* InNativizedClass)
 	}
 	else
 	{
-		if(InNativizedClass && !bNativizedVMDisabled)
+		if(InNativizedClass && bNativizedVMEnabled)
 		{
-			const EObjectFlags PreviousFlags = VM->GetFlags();
-			VM->Rename(nullptr, GetTransientPackage(), REN_DoNotDirty | REN_DontCreateRedirectors | REN_NonTransactional);
-			VM->MarkAsGarbage();
-			VM = NewObject<URigVM>(this, InNativizedClass, TEXT("RigVM_NVMB"), PreviousFlags);
-
+#if WITH_EDITOR
+			URigVM* SourceVM = VM;
+#endif
+			VM = InNativizedClass->GetDefaultObject<URigVM>();
 #if UE_RIGVM_PROFILE_EXECUTE_UNITS_NUM
 			ProfilingRunsLeft = 0;
 			AccumulatedCycles = 0;
 #endif
-		}
-	}
-	
-#if WITH_EDITOR
 
-	// if we are a nativized VM,
-	// let's set the bytecode for UI purposes.
-	// this is only used for traversing node from execute stack to node and back etc. 
-	// since the hash between nativized VM and current matches we assume the bytecode is identical as well.
-	if(VM->IsNativized())
-	{
-		if(URigVMNativized* NativizedVM = Cast<URigVMNativized>(VM))
-		{
-			if (URigVM* SourceVM = GetSourceVM())
+#if WITH_EDITOR
+			// let's set the bytecode for UI purposes.
+			// this is only used for traversing node from execute stack to node and back etc. 
+			// since the hash between nativized VM and current matches we assume the bytecode is identical as well.
+			if(URigVMNativized* NativizedVM = Cast<URigVMNativized>(VM))
 			{
-				NativizedVM->SetByteCode(SourceVM->GetByteCode());
+				NativizedVM->SetSourceVM(SourceVM);
 			}
+			GetRigVMExtendedExecuteContext().ResetNumExecutions();
+#endif
 		}
 	}
+}
+
+bool URigVMHost::AreNativizedVMsEnabled() const
+{
+#if WITH_EDITOR
+	return bEnableNativizedVMs && (CVarRigVMEnableNativizedVMs->GetInt() != 0);
+#else
+	return CVarRigVMEnableNativizedVMs->GetInt() != 0;
 #endif
 }
 
-bool URigVMHost::AreNativizedVMsDisabled()
+bool URigVMHost::CanSwapVMToNativized() const
 {
-	return (CVarRigVMDisableNativizedVMs->GetInt() != 0);
+	if (VM == nullptr)
+	{
+		return false;
+	}
+	
+	if (VM->IsNativized())
+	{
+		return true;
+	}
+
+	const uint32 Hash = VM->GetVMHash();
+	if (LastNativizedVMHash.Get(0) != Hash || !WeakNativizedVMClass.IsSet())
+	{
+		LastNativizedVMHash = Hash;
+		WeakNativizedVMClass = VM->GetNativizedClass(LastNativizedVMHash.GetValue());
+	}
+
+	if (WeakNativizedVMClass.IsSet())
+	{
+		return WeakNativizedVMClass.GetValue().IsValid();
+	}
+
+	return false;
 }
 
 #if WITH_EDITORONLY_DATA
@@ -1363,6 +1393,20 @@ void URigVMHost::GenerateRequiredPluginsData(FRigVMExtendedExecuteContext& InCon
 	}
 }
 
+#if WITH_EDITOR
+
+void URigVMHost::EnableNativizedVM(bool bEnable)
+{
+	if (bEnableNativizedVMs == bEnable)
+	{
+		return;
+	}
+	bEnableNativizedVMs = bEnable;
+	Initialize(true);
+}
+
+#endif
+
 void URigVMHost::HandleExecutionReachedExit(const FName& InEventName)
 {
 	if (bAccumulateTime)
diff --git a/Engine/Plugins/Runtime/RigVM/Source/RigVM/Public/RigVMCore/RigVM.h b/Engine/Plugins/Runtime/RigVM/Source/RigVM/Public/RigVMCore/RigVM.h
--- a/Engine/Plugins/Runtime/RigVM/Source/RigVM/Public/RigVMCore/RigVM.h
+++ b/Engine/Plugins/Runtime/RigVM/Source/RigVM/Public/RigVMCore/RigVM.h
@@ -188,13 +188,15 @@ public:
 
 	// Generates a unique hash to compare VMs
 	RIGVM_API virtual uint32 ComputeVMHash() const;
+	RIGVM_API virtual uint32 ComputeVMHash(const TArray<FRigVMExternalVariableDef>& InExternalVariables) const;
+	RIGVM_API virtual uint32 ComputeVMHash(const TArray<FRigVMExternalVariable>& InExternalVariables) const;
 
-	UE_DEPRECATED(5.3, "Please, use GetNativizedClass with FRigVMExternalVariableDef array.")
+	UE_DEPRECATED(5.3, "Please, use GetNativizedClass with Hash argument.")
 	UClass* GetNativizedClass(const TArray<FRigVMExternalVariable>& InExternalVariables) { return nullptr; }
 
 	// returns the VM's matching nativized class if it exists
-	RIGVM_API UClass* GetNativizedClass(const TArray<FRigVMExternalVariableDef>& InExternalVariables = TArray<FRigVMExternalVariableDef>());
-	
+	RIGVM_API UClass* GetNativizedClass(uint32 InHash) const;
+
 	UE_DEPRECATED(5.4, "Please, use Reset with Context param")
 	virtual void Reset(bool IsIgnoringArchetypeRef = false) {}
 
@@ -393,8 +395,8 @@ public:
 	UPROPERTY()
 	FRigVMByteCode ByteCodeStorage;
 	FRigVMByteCode* ByteCodePtr;
-	FRigVMByteCode& GetByteCode() { return *ByteCodePtr; }
-	const FRigVMByteCode& GetByteCode() const { return *ByteCodePtr; }
+	FRigVMByteCode& GetByteCode() { return const_cast<FRigVMByteCode&>(const_cast<const URigVM*>(this)->GetByteCode()); }
+	virtual const FRigVMByteCode& GetByteCode() const { return *ByteCodePtr; }
 
 	// Returns the instructions of the VM
 	RIGVM_API virtual const FRigVMInstructionArray& GetInstructions();
@@ -837,7 +839,7 @@ public:
 	RIGVM_API void SetExternalVariableDefs(const TArray<FRigVMExternalVariable>& InExternalVariables);
 
 	// Sets the external variables instance data required for execution
-	void SetExternalVariablesInstanceData(FRigVMExtendedExecuteContext& Context, const TArray<FRigVMExternalVariable>& InExternalVariables, bool bAllowNullMemory = false)
+	virtual void SetExternalVariablesInstanceData(FRigVMExtendedExecuteContext& Context, const TArray<FRigVMExternalVariable>& InExternalVariables, bool bAllowNullMemory = false)
 	{
 		const int32 NumExternalVariables = InExternalVariables.Num();
 		check(ExternalVariables.Num() == NumExternalVariables);
@@ -910,12 +912,15 @@ public:
 	RIGVM_API bool CreateLocalizedRegistryIfRequired(bool bForce = false);
 	RIGVM_API bool CreateLocalizedRegistryIfRequired_NoLock(bool bForce = false);
 
-private:
+protected:
 	//  The local registry to use by the VM
 	TSharedPtr<FRigVMRegistry_NoLock> LocalizedRegistry;
 
-public:
+	RIGVM_API FString GetLocalizedRegistryAsString() const;
+	RIGVM_API void SetLocalizedRegistryFromString(const FString& InData);
 
+public:
+	
 #if WITH_EDITOR
 	UE_DEPRECATED(5.4, "Please use DumpByteCodeAsTextArray with ExtendedExecuteContext parameter.")
 	TArray<FString> DumpByteCodeAsTextArray(const TArray<int32>& InInstructionOrder = TArray<int32>(), bool bIncludeLineNumbers = true, TFunction<FString(const FString& RegisterName, const FString& RegisterOffsetName)> OperandFormatFunction = nullptr) { return TArray<FString>(); }
@@ -1013,7 +1018,25 @@ private:
 
 protected:
 
-	virtual void SetInstructionIndex(FRigVMExtendedExecuteContext& Context, int32 InInstructionIndex) { Context.GetPublicData<>().InstructionIndex = InInstructionIndex; }
+	[[nodiscard]] virtual bool SetInstructionIndex(FRigVMExtendedExecuteContext& Context, FRigVMExecuteContext& PublicData, int32 InInstructionIndex) const
+	{
+#if WITH_EDITOR
+		if (Context.DebugInfoStrong != nullptr && Context.DebugInfoStrong->IsActive())
+		{
+			if (Context.DebugInfoStrong->ShouldStop(Context, PublicData.InstructionIndex))
+			{
+				return false;
+			}
+		}
+#endif
+		PublicData.InstructionIndex = InInstructionIndex;
+		return true;
+	}
+
+	[[nodiscard]] bool IncrementInstructionIndex(FRigVMExtendedExecuteContext& Context, FRigVMExecuteContext& PublicData) const
+	{
+		return SetInstructionIndex(Context, PublicData, PublicData.InstructionIndex + 1);
+	}
 
 #if WITH_EDITORONLY_DATA
 	UPROPERTY(transient)
@@ -1033,20 +1056,22 @@ private:
 	UPROPERTY()
 	TArray<FName> FunctionNamesStorage;
 	TArray<FName>* FunctionNamesPtr;
-	TArray<FName>& GetFunctionNames() { return *FunctionNamesPtr; }
-	const TArray<FName>& GetFunctionNames() const { return *FunctionNamesPtr; }
-
+public:
+	TArray<FName>& GetFunctionNames() { return const_cast<TArray<FName>&>(const_cast<const URigVM*>(this)->GetFunctionNames()); }
+	virtual const TArray<FName>& GetFunctionNames() const { return *FunctionNamesPtr; }
+private:
 	TArray<const FRigVMFunction*> FunctionsStorage;
 	TArray<const FRigVMFunction*>* FunctionsPtr;
 public:
-	TArray<const FRigVMFunction*>& GetFunctions() { return *FunctionsPtr; }
-	const TArray<const FRigVMFunction*>& GetFunctions() const { return *FunctionsPtr; }
-
+	TArray<const FRigVMFunction*>& GetFunctions() { return const_cast<TArray<const FRigVMFunction*>&>(const_cast<const URigVM*>(this)->GetFunctions()); }
+	virtual const TArray<const FRigVMFunction*>& GetFunctions() const { return *FunctionsPtr; }
 private:
 	TArray<const FRigVMDispatchFactory*> FactoriesStorage;
 	TArray<const FRigVMDispatchFactory*>* FactoriesPtr;
-	TArray<const FRigVMDispatchFactory*>& GetFactories() { return *FactoriesPtr; }
-	const TArray<const FRigVMDispatchFactory*>& GetFactories() const { return *FactoriesPtr; }
+public:
+	TArray<const FRigVMDispatchFactory*>& GetFactories() { return const_cast<TArray<const FRigVMDispatchFactory*>&>(const_cast<const URigVM*>(this)->GetFactories()); }
+	virtual const TArray<const FRigVMDispatchFactory*>& GetFactories() const { return *FactoriesPtr; }
+private:
 
 	UPROPERTY()
 	TArray<FRigVMParameter> Parameters;
@@ -1169,6 +1194,8 @@ private:
 	UPROPERTY()
 	uint32 CachedVMHash = 0;
 
+	mutable TOptional<uint32> CachedNativizedVMHash;
+
 	mutable TArray<FName> EntryNames;
 
 	bool bIsSerializing = false;
@@ -1180,7 +1207,7 @@ protected:
 	RIGVM_API void ClearMemory_Internal();
 	RIGVM_API void InvalidateCachedMemory_Internal();
 
-	RIGVM_API void SetupInstructionTracking(FRigVMExtendedExecuteContext& Context, int32 InInstructionCount, int32 InCallableCount);
+	RIGVM_API void SetupInstructionTracking(FRigVMExtendedExecuteContext& Context, int32 InInstructionCount, int32 InCallableCount, const FName& InEntryName);
 	RIGVM_API void StartProfiling(FRigVMExtendedExecuteContext& Context);
 	RIGVM_API void StopProfiling(FRigVMExtendedExecuteContext& Context);
 
@@ -1189,5 +1216,5 @@ protected:
 
 	friend class URigVMCompiler;
 	friend struct FRigVMCompilerWorkData;
-	friend struct FRigVMCodeGenerator;
+	friend class FRigVMCodeConverter;
 };
diff --git a/Engine/Plugins/Runtime/RigVM/Source/RigVM/Public/RigVMCore/RigVMByteCode.h b/Engine/Plugins/Runtime/RigVM/Source/RigVM/Public/RigVMCore/RigVMByteCode.h
--- a/Engine/Plugins/Runtime/RigVM/Source/RigVM/Public/RigVMCore/RigVMByteCode.h
+++ b/Engine/Plugins/Runtime/RigVM/Source/RigVM/Public/RigVMCore/RigVMByteCode.h
@@ -163,7 +163,7 @@ struct FRigVMCallableArgument
 	friend uint32 GetTypeHash(const FRigVMCallableArgument& InArgument)
 	{
 		return HashCombine(
-			GetTypeHash(InArgument.Name),
+			GetTypeHash(InArgument.Name.ToString()),
 			GetTypeHash(InArgument.TypeString),
 			GetTypeHash(InArgument.InterfaceOperand),
 			GetTypeHash(InArgument.ForwardedOperand),
@@ -182,12 +182,21 @@ struct FRigVMCallableInfo
 	FRigVMCallableInfo()
 		: Index(INDEX_NONE)
 		, Name(NAME_None)
-		, Hash(0)
+		, FunctionHash(0)
 		, FirstInstruction(INDEX_NONE)
 		, LastInstruction(INDEX_NONE)
 	{
 	}
 
+	FRigVMCallableInfo(int32 InIndex, const FName& InName, uint32 InHash, int32 InFirstInstruction, int32 InLastInstruction)
+	: Index(InIndex)
+	, Name(InName)
+	, FunctionHash(InHash)
+	, FirstInstruction(InFirstInstruction)
+	, LastInstruction(InLastInstruction)
+	{
+	}
+
 	UPROPERTY()
 	int32 Index;
 	
@@ -195,7 +204,7 @@ struct FRigVMCallableInfo
 	FName Name;
 	
 	UPROPERTY()
-	uint32 Hash;
+	uint32 FunctionHash;
 
 	UPROPERTY()
 	TArray<FRigVMCallableArgument> Arguments;
@@ -222,7 +231,17 @@ struct FRigVMCallableInfo
 
 	friend uint32 GetTypeHash(const FRigVMCallableInfo& InCallableInfo)
 	{
-		return InCallableInfo.Hash;
+		uint32 Hash = HashCombine(
+			GetTypeHash(InCallableInfo.Index),
+			GetTypeHash(InCallableInfo.Name.ToString()),
+			GetTypeHash(InCallableInfo.FirstInstruction),
+			GetTypeHash(InCallableInfo.LastInstruction)
+		);
+		for (const FRigVMCallableArgument& Argument : InCallableInfo.Arguments)
+		{
+			Hash = HashCombine(Hash, GetTypeHash(Argument));
+		}
+		return Hash;
 	}
 };
 
@@ -1371,6 +1390,9 @@ public:
 	// returns a callable given its name
 	UE_API const FRigVMCallableInfo* FindCallable(const FName& InName) const;
 
+	// returns a callable given an instruction within the callable
+	UE_API const FRigVMCallableInfo* FindCallableForInstruction(int32 InInstructionIndex) const;
+
 	// adds an execute operator given its function index operands
 	UE_API int32 AddExecuteOp(int32 InFunctionIndex, const FRigVMOperandArray& InOperands, const int32 InStartPredicateIndex, const int32 InPredicateCount);
 	
@@ -1451,6 +1473,12 @@ public:
 		return BranchInfos[InIndex];
 	}
 
+	// returns the branch for a given argument index
+	UE_API const FRigVMBranchInfo* FindBranch(int32 InInstructionIndex, int32 InArgumentIndex) const; 
+
+	// returns the branch for a given label
+	UE_API const FRigVMBranchInfo* FindBranch(int32 InInstructionIndex, const FName& InLabel) const; 
+
 	// returns the number of predicate branches within this bytecode
 	int32 NumPredicateBranches() const { return PredicateBranches.Num(); }
 
diff --git a/Engine/Plugins/Runtime/RigVM/Source/RigVM/Public/RigVMCore/RigVMDebugInfo.h b/Engine/Plugins/Runtime/RigVM/Source/RigVM/Public/RigVMCore/RigVMDebugInfo.h
--- a/Engine/Plugins/Runtime/RigVM/Source/RigVM/Public/RigVMCore/RigVMDebugInfo.h
+++ b/Engine/Plugins/Runtime/RigVM/Source/RigVM/Public/RigVMCore/RigVMDebugInfo.h
@@ -15,6 +15,94 @@
 
 class UObject;
 class URigVMNode;
+struct FRigVMExtendedExecuteContext;
+
+UENUM()
+enum class ERigVMStepMode : uint8
+{
+	None,
+	Stopped,	// Parked at a specific instruction
+	Forward,    // Next instruction at same callable depth
+	Into,       // Enter next callable
+	Out,        // Exit current callable
+	RunTo       // Run until target instruction (first hit)
+};
+
+/*
+FRigVMStepCondition tracks:
+- Mode - None, Stopped, Forward, Into, Out, RunTo
+- OriginInstruction - instruction index where we're stopped
+- OriginCallstackPath - array of call site instruction indices (full callstack context)
+- OriginIterationPath - array of loop iteration indices
+- bHasPassedOrigin - flag to ensure we've reached the origin before detecting forward progress
+
+Step behaviors:
+- Into - stops when entering a callable (deeper), advancing at same level, or exiting a callable (shallower)
+- Forward - stops when advancing at same level or exiting a callable (shallower)
+- Out - stops only when exiting a callable (shallower)
+- RunTo - stops at target instruction regardless of context
+- Stopped - re-halts at exact same location (callstack + instruction + iteration)
+*/
+USTRUCT()
+struct FRigVMStepCondition
+{
+	GENERATED_BODY()
+
+	FRigVMStepCondition()
+		: Mode(ERigVMStepMode::None)
+		, OriginInstruction(INDEX_NONE)
+		, TargetInstruction(INDEX_NONE)
+		, bHasPassedOrigin(false)
+	{
+	}
+
+	ERigVMStepMode Mode;
+	int32 OriginInstruction;
+	int32 TargetInstruction;  // For RunTo mode
+	bool bHasPassedOrigin;    // True once we've passed the origin point during re-execution
+
+	// Callstack path - instruction index at each call site
+	TArray<int32, TInlineAllocator<16>> OriginCallstackPath;
+
+	// Iteration path when stopped - for re-halting at same location
+	TArray<int32, TInlineAllocator<16>> OriginIterationPath;
+
+	FString ToString() const
+	{
+		static const TCHAR* ModeNames[] = {
+			TEXT("None"),
+			TEXT("Stopped"),
+			TEXT("Forward"),
+			TEXT("Into"),
+			TEXT("Out"),
+			TEXT("RunTo")
+		};
+
+		FString CallstackPathStr = TEXT("[");
+		for (int32 i = 0; i < OriginCallstackPath.Num(); i++)
+		{
+			if (i > 0) CallstackPathStr += TEXT(", ");
+			CallstackPathStr += FString::Printf(TEXT("%d"), OriginCallstackPath[i]);
+		}
+		CallstackPathStr += TEXT("]");
+
+		FString IterationPathStr = TEXT("[");
+		for (int32 i = 0; i < OriginIterationPath.Num(); i++)
+		{
+			if (i > 0) IterationPathStr += TEXT(", ");
+			IterationPathStr += FString::Printf(TEXT("%d"), OriginIterationPath[i]);
+		}
+		IterationPathStr += TEXT("]");
+
+		return FString::Printf(TEXT("Mode=%s, OriginInstr=%d, TargetInstr=%d, PassedOrigin=%d, Callstack=%s, IterPath=%s"),
+			ModeNames[static_cast<uint8>(Mode)],
+			OriginInstruction,
+			TargetInstruction,
+			bHasPassedOrigin ? 1 : 0,
+			*CallstackPathStr,
+			*IterationPathStr);
+	}
+};
 
 USTRUCT()
 struct FRigVMDebugInfo
@@ -27,47 +115,67 @@ struct FRigVMDebugInfo
 
 	void Reset()
 	{
-		EarlyExitInstruction.Reset();
+		StepCondition = FRigVMStepCondition();
 	}
 
-	bool IsEmpty() const
+	bool IsActive() const
 	{
-		return !EarlyExitInstruction.IsSet();
+		return StepCondition.Mode != ERigVMStepMode::None;
 	}
 
-	int32 GetInstructionForEarlyExit() const
+	UE_API bool ShouldStop(const FRigVMExtendedExecuteContext& Context, int32 CurrentInstruction);
+
+	void StepForward()
 	{
-		return EarlyExitInstruction.Get(INDEX_NONE);
+		StepCondition.Mode = ERigVMStepMode::Forward;
 	}
 
-	void SetInstructionForEarlyExit(const int32 InInstruction)
+	void StepInto()
 	{
-		EarlyExitInstruction = InInstruction;
+		StepCondition.Mode = ERigVMStepMode::Into;
 	}
 
-	bool ShouldEarlyExitBeforeInstruction(const int32 InInstruction) const
+	void StepOut()
 	{
-		if (EarlyExitInstruction.IsSet())
-		{
-			return InInstruction > EarlyExitInstruction.GetValue();
-		}
-		return false;
+		StepCondition.Mode = ERigVMStepMode::Out;
 	}
 
-	void CopyFrom(const FRigVMDebugInfo& Other)
+	void RunToInstruction(int32 Instruction)
 	{
-		EarlyExitInstruction = Other.EarlyExitInstruction;
+		StepCondition.Mode = ERigVMStepMode::RunTo;
+		StepCondition.TargetInstruction = Instruction;
+	}
+
+	void BeginExecution()
+	{
+		StepCondition.bHasPassedOrigin = false;
+	}
+
+	const FRigVMStepCondition& GetStepCondition() const
+	{
+		return StepCondition;
+	}
+
+	void SetStepCondition(const FRigVMStepCondition& InStepCondition)
+	{
+		StepCondition = InStepCondition;
 	}
 
 	DECLARE_EVENT_ThreeParams(URigVM, FExecutionHaltedEvent, int32, UObject*, const FName&);
-	inline FExecutionHaltedEvent& ExecutionHalted()
+	FExecutionHaltedEvent& ExecutionHalted()
 	{
 		return OnExecutionHalted;
 	}
 
 private:
+	UE_API void TransitionToStopped(const FRigVMExtendedExecuteContext& Context, int32 CurrentInstruction);
+	UE_API static void CaptureCallstackPath(const FRigVMExtendedExecuteContext& Context, TArray<int32, TInlineAllocator<16>>& OutPath);
+	UE_API static int32 CompareCallstackPath(const FRigVMExtendedExecuteContext& Context, const TArray<int32, TInlineAllocator<16>>& InPath);
+	UE_API static void CaptureIterationPath(const FRigVMExtendedExecuteContext& Context, TArray<int32, TInlineAllocator<16>>& OutPath);
+	UE_API static int32 CompareIterationPath(const FRigVMExtendedExecuteContext& Context, const TArray<int32, TInlineAllocator<16>>& InPath);
+
+	FRigVMStepCondition StepCondition;
 	FExecutionHaltedEvent OnExecutionHalted;
-	TOptional<int32> EarlyExitInstruction;
 };
 
 #undef UE_API
diff --git a/Engine/Plugins/Runtime/RigVM/Source/RigVM/Public/RigVMCore/RigVMDispatchFactory.h b/Engine/Plugins/Runtime/RigVM/Source/RigVM/Public/RigVMCore/RigVMDispatchFactory.h
--- a/Engine/Plugins/Runtime/RigVM/Source/RigVM/Public/RigVMCore/RigVMDispatchFactory.h
+++ b/Engine/Plugins/Runtime/RigVM/Source/RigVM/Public/RigVMCore/RigVMDispatchFactory.h
@@ -301,4 +301,5 @@ protected:
 	friend struct FRigVMRegistry_NoLock;
 	friend struct FRigVMFunction;
 	friend class URigVM;
+	friend class FRigVMCodeConverter;
 };
diff --git a/Engine/Plugins/Runtime/RigVM/Source/RigVM/Public/RigVMCore/RigVMExecuteContext.h b/Engine/Plugins/Runtime/RigVM/Source/RigVM/Public/RigVMCore/RigVMExecuteContext.h
--- a/Engine/Plugins/Runtime/RigVM/Source/RigVM/Public/RigVMCore/RigVMExecuteContext.h
+++ b/Engine/Plugins/Runtime/RigVM/Source/RigVM/Public/RigVMCore/RigVMExecuteContext.h
@@ -538,6 +538,22 @@ struct FRigVMExecuteContext : public FRigVMExecutePin
 	 */
 	UE_API bool SerializeFromMismatchedTag(const FPropertyTag& Tag, FStructuredArchive::FSlot Slot);
 
+	/**
+	 * Completes the execution of this run.  
+	 */
+	void MarkExecutionCompleted()
+	{
+		NumExecutions++;
+	}
+
+	/**
+	 * Resets the number of executions
+	 */
+	void ResetNumExecutions()
+	{
+		NumExecutions = 0;
+	}
+
 protected:
 
 	virtual void Reset()
@@ -636,6 +652,71 @@ struct FRigVMExternalVariableRuntimeData
 	uint8* Memory = nullptr;
 };
 
+class FRigVMNativizedVMVariables
+{
+};
+
+struct FRigVMExecuteCallstackEntry
+{
+	FRigVMExecuteCallstackEntry() = default;
+
+	FRigVMExecuteCallstackEntry(int32 InInstructionIndex, const FRigVMCallableInfo* InCallable)
+		: InstructionIndex(InInstructionIndex)
+		, Callable(InCallable)
+	{
+	}
+	
+	int32 InstructionIndex = INDEX_NONE;
+	const FRigVMCallableInfo* Callable = nullptr;
+};
+
+struct FRigVMExecuteCallstack
+{
+	FRigVMExecuteCallstack() = default;
+
+	void Reset()
+	{
+		Entries.Reset();
+	}
+
+	bool IsValidIndex(int32 InIndex) const
+	{
+		return Entries.IsValidIndex(InIndex);
+	}
+
+	int32 Num() const
+	{
+		return Entries.Num();
+	}
+
+	bool IsEmpty() const
+	{
+		return Entries.IsEmpty();
+	}
+
+	const FRigVMExecuteCallstackEntry& Last() const
+	{
+		return Entries.Last();
+	}
+
+	const FRigVMExecuteCallstackEntry& operator[](int32 InIndex) const
+	{
+		return Entries[InIndex];
+	}
+
+	int32 Push(const FRigVMExecuteCallstackEntry& InEntry)
+	{
+		return Entries.Add(InEntry);
+	}
+
+	FRigVMExecuteCallstackEntry Pop()
+	{
+		return Entries.Pop();
+	}
+
+	TArray<FRigVMExecuteCallstackEntry> Entries;
+};
+
 /**
  * The execute context is used for mutable nodes to
  * indicate execution order.
@@ -803,7 +884,7 @@ public:
 		return Hash;
 	}
 
-	const TArray<TTuple<int32, const FRigVMCallableInfo*>>& GetCallstack() const
+	const FRigVMExecuteCallstack& GetCallstack() const
 	{
 		return Callstack;
 	}
@@ -818,6 +899,10 @@ public:
 		return HashCombine(GetTypeHash(GetNumExecutions()), GetCallstackHash());
 	}
 
+	UE_API bool StepForward();
+	UE_API bool StepInto();
+	UE_API bool StepOut();
+
 	bool IsValidArrayIndex(int32& InOutIndex, const FScriptArrayHelper& InArrayHelper) const
 	{
 		return IsValidArrayIndex(InOutIndex, InArrayHelper.Num());
@@ -901,6 +986,14 @@ public:
 		return 0;
 	}
 
+	void ResetNumExecutions()
+	{
+		if(PublicDataScope.IsValid())
+		{
+			GetPublicData<>().ResetNumExecutions();
+		}
+	}
+
 	UE_API int32 GetInstructionIndex() const;
 	UE_API const URigVM* GetVM() const;
 	UE_API URigVM* GetVM();
@@ -943,7 +1036,7 @@ public:
 	// gets its own memory.
 	TArray<uint16> SliceOffsetsPerCallable;
 
-	TArray<TTuple<int32, const FRigVMCallableInfo*>> Callstack;
+	FRigVMExecuteCallstack Callstack;
 	uint32 CallstackHash = 0;
 
 	const FRigVMDispatchFactory* Factory = nullptr;
@@ -990,6 +1083,7 @@ public:
 
 	TArray<FRigVMInstructionSetExecuteState> LazyBranchExecuteState;
 	TArray<FRigVMExternalVariableRuntimeData> ExternalVariableRuntimeData;
+	TSharedPtr<FRigVMNativizedVMVariables> NativizedVMVariables;
 
 	UE_API const FRigVMOperand* GetDebugOperandForOperand(const FRigVMOperand& InOperand) const;
 	UE_API void MarkOperandForDebugging(const FRigVMOperand& InOperand, bool bEnableDebugging = true);
@@ -1009,6 +1103,7 @@ public:
 	FRigVMDebugInfo* DebugInfo = nullptr;
 
 	TWeakPtr<FRigVMDebugInfo> DebugInfoWeak;
+	FRigVMDebugInfo* DebugInfoStrong = nullptr;
 
 	FRigVMInstructionVisitInfo* InstructionVisitInfo = nullptr;
 
@@ -1022,8 +1117,22 @@ public:
 		return InstructionVisitInfo;
 	}
 
-	void SetDebugInfo(const TSharedPtr<FRigVMDebugInfo>& InSharedDebugInfo)
+	void SetDebugInfo(const TSharedPtr<FRigVMDebugInfo>& InSharedDebugInfo, bool bForce)
 	{
+		if (!bForce)
+		{
+			if (DebugInfoWeak.IsValid())
+			{
+				if (const TSharedPtr<FRigVMDebugInfo>& DebugInfoShared = DebugInfoWeak.Pin())
+				{
+					// don't change actively debugged instances.
+					if (DebugInfoShared->IsActive())
+					{
+						return;
+					}
+				}
+			}
+		}
 		DebugInfoWeak = InSharedDebugInfo.ToWeakPtr();
 	}
 
diff --git a/Engine/Plugins/Runtime/RigVM/Source/RigVM/Public/RigVMCore/RigVMFunction.h b/Engine/Plugins/Runtime/RigVM/Source/RigVM/Public/RigVMCore/RigVMFunction.h
--- a/Engine/Plugins/Runtime/RigVM/Source/RigVM/Public/RigVMCore/RigVMFunction.h
+++ b/Engine/Plugins/Runtime/RigVM/Source/RigVM/Public/RigVMCore/RigVMFunction.h
@@ -192,6 +192,7 @@ struct FRigVMFunction
 	RIGVM_API FString GetModuleName() const;
 	RIGVM_API FString GetModuleRelativeHeaderPath() const;
 	const TArray<FRigVMFunctionArgument>& GetArguments() const { return Arguments; }
+	RIGVM_API int32 GetArgumentIndex(const FName& InArgumentName) const;
 	RIGVM_API const TArray<TRigVMTypeIndex>& GetArgumentTypeIndices() const;
 	RIGVM_API const TArray<TRigVMTypeIndex>& GetArgumentTypeIndices_NoLock(FRigVMRegistryHandle& InRegistry) const;
 	RIGVM_API void GetArgumentTypeIndices_NoLock(FRigVMRegistryHandle& InRegistry, TArray<TRigVMTypeIndex>& OutTypeIndices) const;
@@ -203,4 +204,9 @@ struct FRigVMFunction
 	RIGVM_API bool SupportsExecuteContextStruct_NoLock(const UScriptStruct* InExecuteContextStruct, const FRigVMRegistryHandle& InRegistry) const;
 	RIGVM_API const FName& GetArgumentNameForOperandIndex(int32 InOperandIndex, int32 InTotalOperands, FRigVMRegistryHandle& InRegistry) const;
 	RIGVM_API void Serialize_NoLock(FArchive& Ar, FRigVMRegistryHandle& InRegistry);
+
+#if WITH_EDITOR
+	RIGVM_API bool IsFixedSizeArray(const FName& InArgumentName) const;
+	RIGVM_API bool IsLazyValue(const FName& InArgumentName) const;
+#endif
 };
diff --git a/Engine/Plugins/Runtime/RigVM/Source/RigVM/Public/RigVMCore/RigVMMemoryStorage.h b/Engine/Plugins/Runtime/RigVM/Source/RigVM/Public/RigVMCore/RigVMMemoryStorage.h
--- a/Engine/Plugins/Runtime/RigVM/Source/RigVM/Public/RigVMCore/RigVMMemoryStorage.h
+++ b/Engine/Plugins/Runtime/RigVM/Source/RigVM/Public/RigVMCore/RigVMMemoryStorage.h
@@ -75,16 +75,6 @@ struct FRigVMBranchInfo
 		return InstructionIndex != INDEX_NONE;
 	}
 
-	bool IsOutputBranch() const
-	{
-		return ArgumentIndex == INDEX_NONE;
-	}
-
-	bool IsLazyEvaluationBranch() const
-	{
-		return !IsOutputBranch();
-	}
-	
 	RIGVM_API void Serialize(FArchive& Ar);
 	friend FArchive& operator<<(FArchive& Ar, FRigVMBranchInfo& P)
 	{
@@ -104,7 +94,7 @@ struct FRigVMBranchInfo
 
 	static FString GetFixedArrayLabel(const FString& InArrayName, const FString& InElementName)
 	{
-		return RigVMStringUtils::JoinStrings(InArrayName, InElementName, TEXT("_"));
+		return RigVMStringUtils::JoinStrings(InArrayName, InElementName.Replace(TEXT("_"), TEXT("")), TEXT("_"));
 	}
 
 	static FName GetFixedArrayLabel(const FName& InArrayName, const FName& InElementName)
@@ -145,6 +135,14 @@ public:
 		, FunctionPtr(nullptr)
 	{}
 
+	FRigVMLazyBranch(int32 InBranchIndex, const TFunction<ERigVMExecuteResult()>& InFunctionPtr)
+		: VM(nullptr)
+		, BranchInfo()
+		, FunctionPtr(InFunctionPtr)
+	{
+		BranchInfo.Index = InBranchIndex;
+	}
+
 	RIGVM_API ERigVMExecuteResult Execute(FRigVMExtendedExecuteContext& Context);
 	RIGVM_API ERigVMExecuteResult ExecuteIfRequired(FRigVMExtendedExecuteContext& Context, uint32 InSliceHash);
 
@@ -198,6 +196,10 @@ public:
 		, ConstValue(InConstValue)
 	{}
 
+	virtual ~TRigVMLazyValue()
+	{
+	}
+
 	const ComputedType& Get(const FRigVMExecuteContext& ExecuteContext) const
 	{
 		if(MemoryHandle)
@@ -1047,7 +1049,6 @@ private:
 	friend struct TRigVMLazyValueBase;
 	friend class URigVMMemoryStorage;
 	friend struct FRigVMMemoryStorageStruct;
-	friend struct FRigVMMemoryStorageStruct;
 };
 
 /**
diff --git a/Engine/Plugins/Runtime/RigVM/Source/RigVM/Public/RigVMCore/RigVMMemoryStorageStruct.h b/Engine/Plugins/Runtime/RigVM/Source/RigVM/Public/RigVMCore/RigVMMemoryStorageStruct.h
--- a/Engine/Plugins/Runtime/RigVM/Source/RigVM/Public/RigVMCore/RigVMMemoryStorageStruct.h
+++ b/Engine/Plugins/Runtime/RigVM/Source/RigVM/Public/RigVMCore/RigVMMemoryStorageStruct.h
@@ -20,6 +20,8 @@ struct FRigVMMemoryStorageStruct : public FInstancedPropertyBag
 	UE_API FRigVMMemoryStorageStruct(ERigVMMemoryType InMemoryType, const TArray<FRigVMPropertyDescription>& InPropertyDescriptions, const TArray<FRigVMPropertyPathDescription>& InPropertyPaths = TArray<FRigVMPropertyPathDescription>());
 	UE_API ~FRigVMMemoryStorageStruct();
 
+	UE_API void InitializeFromNativeStruct(const UScriptStruct* InStruct);
+
 	UE_API bool Serialize(FArchive& Ar);
 	UE_API void AddStructReferencedObjects(FReferenceCollector& Collector);
 	UE_API void GetUserDefinedDependencies(TArray<const UObject*>& OutDependencies) const;
@@ -389,10 +391,7 @@ struct FRigVMMemoryStorageStruct : public FInstancedPropertyBag
 	/**
 	* Returns the PropertyBag struct using a standard function signature, so it can be used at IStructureDataProvider
 	*/
-	inline const UStruct * GetStruct() const
-	{
-		return GetPropertyBagStruct();
-	}
+	UE_API  const UStruct * GetStruct() const;
 
 	// Helper function to generate a property bag desc from a FRigVMPropertyDescription
 	static UE_API FPropertyBagPropertyDesc GeneratePropertyBagDescriptor(const FRigVMPropertyDescription& RigVMDescriptor);
