RigVM Nativization System

diff --git a/Engine/Plugins/Runtime/RigVM/Source/RigVM/Public/RigVMCore/RigVMNativized.h b/Engine/Plugins/Runtime/RigVM/Source/RigVM/Public/RigVMCore/RigVMNativized.h
--- a/Engine/Plugins/Runtime/RigVM/Source/RigVM/Public/RigVMCore/RigVMNativized.h
+++ b/Engine/Plugins/Runtime/RigVM/Source/RigVM/Public/RigVMCore/RigVMNativized.h
@@ -17,6 +17,7 @@
 #include "RigVMCore/RigVMExecuteContext.h"
 #include "RigVMCore/RigVMExternalVariable.h"
 #include "RigVMCore/RigVMMemoryCommon.h"
+#include "RigVMCore/RigVMMemoryStorage.h"
 #include "RigVMCore/RigVMTraits.h"
 #include "Templates/EnableIf.h"
 #include "Templates/Models.h"
@@ -35,6 +36,72 @@ class FArchive;
 class UObject;
 struct FRigVMInstructionArray;
 
+#define RIGVMNATIVIZED_ARRAYTYPE_BODY(NativeType) \
+operator TArray<NativeType>&() { return Array; } \
+operator const TArray<NativeType>&() const { return Array; } \
+void Reset() { Array.Reset(); } \
+void Empty() { Array.Empty(); } \
+void Reserve(int32 InSize) { Array.Reserve(InSize); } \
+void SetNum(int32 InSize) { Array.SetNum(InSize); } \
+void SetNumUninitialized(int32 InSize) { Array.SetNumUninitialized(InSize); } \
+void SetNumZeroed(int32 InSize) { Array.SetNumZeroed(InSize); } \
+bool IsValidIndex(int32 InIndex) const { return Array.IsValidIndex(InIndex); } \
+int32 Num() const { return Array.Num(); } \
+bool IsEmpty() const { return Array.IsEmpty(); } \
+const NativeType& operator[](int32 InIndex) const { return Array[InIndex]; } \
+NativeType& operator[](int32 InIndex) { return Array[InIndex]; } \
+TArray<NativeType>::RangedForIteratorType      begin()       { return Array.begin(); } \
+TArray<NativeType>::RangedForConstIteratorType begin() const { return Array.begin(); } \
+TArray<NativeType>::RangedForIteratorType      end()         { return Array.end();   } \
+TArray<NativeType>::RangedForConstIteratorType end() const   { return Array.end();   } \
+const NativeType* GetData() const { return Array.GetData(); }
+
+struct FRigVMNativizedContext
+{
+	UE_API FRigVMNativizedContext(FRigVMExtendedExecuteContext& InExecuteContext, FRigVMExecuteContext& InPublicContext);
+	FRigVMExtendedExecuteContext& ExecuteContext;
+	FRigVMExecuteContext& PublicContextBase;
+};
+
+template<typename ComputedType>
+class TRigVMNativizedLazyValue : public TRigVMLazyValue<ComputedType>
+{
+public:
+	TRigVMNativizedLazyValue(const ComputedType& Value, const FProperty* Property)
+	: TRigVMLazyValue<ComputedType>(Value)
+	, NativizedMemoryHandle((uint8*)&Value, Property, nullptr, nullptr)
+	{
+		TRigVMLazyValueBase::MemoryHandle = &NativizedMemoryHandle;
+	}
+
+	TRigVMNativizedLazyValue(const ComputedType& Value, const FProperty* Property, int32 InBranchIndex, FRigVMNativizedContext& Context, const TFunction<ERigVMExecuteResult()>& ComputeCallback)
+	: TRigVMLazyValue<ComputedType>(Value)
+	, LazyBranch(InBranchIndex, ComputeCallback)
+	, NativizedMemoryHandle((uint8*)&Value, Property, nullptr, &LazyBranch)
+	{
+		TRigVMLazyValueBase::MemoryHandle = &NativizedMemoryHandle;
+	}
+
+	virtual ~TRigVMNativizedLazyValue() override
+	{
+	}
+
+	operator const ComputedType& () const
+	{
+		return NativizedMemoryHandle.GetDataLazily<ComputedType>();
+	}
+
+	operator ComputedType& ()
+	{
+		return NativizedMemoryHandle.GetDataLazily<ComputedType>();
+	}
+
+private:
+
+	FRigVMLazyBranch LazyBranch;
+	FRigVMMemoryHandle NativizedMemoryHandle;
+};
+
 UCLASS(MinimalAPI, BlueprintType, Abstract)
 class URigVMNativized : public URigVM
 {
@@ -43,35 +110,66 @@ class URigVMNativized : public URigVM
 public:
 
 	UE_API URigVMNativized();
-	UE_API virtual ~URigVMNativized();
+	UE_API virtual ~URigVMNativized() override;
+
+	UE_API static UClass* FindClassForHash(uint32 InVMHash);
+	UE_API static void RegisterClassForHash(uint32 InVMHash, UClass* InClass);
+	UE_API static TMap<uint32, UClass*>& GetNativizedClassMap();
 
 	// URigVM interface
 	UE_API virtual void Serialize(FArchive& Ar) override;
 	virtual void PostLoad() override { return; }
-	UE_API virtual void Reset(bool IsIgnoringArchetypeRef = false) override;
+	UE_API virtual void Reset(FRigVMExtendedExecuteContext& Context) override;
 	virtual bool IsNativized() const override { return true; }
 	virtual void Empty(FRigVMExtendedExecuteContext& Context) override { return; }
-	virtual void CopyFrom(URigVM* InVM, bool bDeferCopy = false, bool bReferenceLiteralMemory = false, bool bReferenceByteCode = false, bool bCopyExternalVariables = false, bool bCopyDynamicRegisters = false) override { return; }
 	UE_API virtual bool Initialize(FRigVMExtendedExecuteContext& Context) override;
 	UE_API virtual ERigVMExecuteResult ExecuteVM(FRigVMExtendedExecuteContext& Context, const FName& InEntryName = NAME_None) override;
 	virtual int32 AddRigVMFunction(UScriptStruct* InRigVMStruct, const FName& InMethodName) override { return INDEX_NONE; }
 	virtual FString GetRigVMFunctionName(int32 InFunctionIndex) const override { return FString(); }
-	virtual FRigVMMemoryStorageStruct* GetMemoryByType(FRigVMExtendedExecuteContext& Context, ERigVMMemoryType InMemoryType) override { return nullptr; }
-	virtual void ClearMemory() override { return; }
+	UE_API virtual FRigVMMemoryStorageStruct* GetMemoryByType(FRigVMExtendedExecuteContext& Context, ERigVMMemoryType InMemoryType) override;
+	UE_API virtual const FRigVMMemoryStorageStruct* GetMemoryByType(const FRigVMExtendedExecuteContext& Context, ERigVMMemoryType InMemoryType) const override;
+	virtual void ClearMemory(FRigVMExtendedExecuteContext& Context) override { return; }
 	UE_API virtual const FRigVMInstructionArray& GetInstructions() override;
+	UE_API virtual const FRigVMByteCode& GetByteCode() const override;
+	UE_API virtual const TArray<const FRigVMFunction*>& GetFunctions() const override;
+	UE_API virtual const TArray<FName>& GetFunctionNames() const override;
+	UE_API virtual const TArray<const FRigVMDispatchFactory*>& GetFactories() const override;
 	virtual bool ContainsEntry(const FName& InEntryName) const override { return GetEntryNames().Contains(InEntryName); }
 	virtual int32 FindEntry(const FName& InEntryName) const override { return GetEntryNames().Find(InEntryName); }
 	virtual const TArray<FName>& GetEntryNames() const override { static const TArray<FName> EmptyEntries; return EmptyEntries; }
-	virtual void SetInstructionIndex(FRigVMExtendedExecuteContext& Context, int32 InInstructionIndex) override
+
+	UE_DEPRECATED(5.8, "Please, use with FRigVMNativizedContext param")
+	[[nodiscard]] virtual bool SetInstructionIndex(FRigVMExtendedExecuteContext& Context, FRigVMExecuteContext& PublicData, int32 InInstructionIndex) const override
 	{
-		Super::SetInstructionIndex(Context, InInstructionIndex);
+		if (!Super::SetInstructionIndex(Context, PublicData, InInstructionIndex))
+		{
+			return false;
+		}
 #if WITH_EDITOR
+		
 		if (FRigVMInstructionVisitInfo* InstructionVisitInfo = Context.GetRigVMInstructionVisitInfo())
 		{
 			InstructionVisitInfo->SetInstructionVisitedDuringLastRun(InInstructionIndex);
 			InstructionVisitInfo->AddInstructionIndexToVisitOrder(InInstructionIndex);
 		}
 #endif
+		return true;
+	}
+
+	[[nodiscard]] bool SetInstructionIndex(FRigVMNativizedContext& Context, int32 InInstructionIndex) const
+	{
+		if (!Super::SetInstructionIndex(Context.ExecuteContext, Context.PublicContextBase, InInstructionIndex))
+		{
+			return false;
+		}
+#if WITH_EDITOR
+		if (FRigVMInstructionVisitInfo* InstructionVisitInfo = Context.ExecuteContext.GetRigVMInstructionVisitInfo())
+		{
+			InstructionVisitInfo->SetInstructionVisitedDuringLastRun(InInstructionIndex);
+			InstructionVisitInfo->AddInstructionIndexToVisitOrder(InInstructionIndex);
+		}
+#endif
+		return true;
 	}
 
 	UE_DEPRECATED(5.3, "Please, use Empty with Context param")
@@ -89,27 +187,46 @@ public:
 	UE_DEPRECATED(5.3, "Please, use SetInstructionIndex with Context param")
 	virtual void SetInstructionIndex(uint16 InInstructionIndex) {}
 
+	UE_API virtual URigVM* GetSourceVM() const;
 #if WITH_EDITOR
-	UE_API void SetByteCode(const FRigVMByteCode& InByteCode);
+	UE_API void SetSourceVM(URigVM* InSourceVM);
 #endif
-	
+
+	template<typename T>
+	T* SetupNativizedVariables(FRigVMExtendedExecuteContext& Context)
+	{
+		Context.NativizedVMVariables = MakeShared<T>();
+		return (T*)Context.NativizedVMVariables.Get();
+	}
+	template<typename T>
+	T* GetNativizedVariables(FRigVMExtendedExecuteContext& Context)
+	{
+		check(Context.NativizedVMVariables.IsValid());
+		return (T*)Context.NativizedVMVariables.Get();
+	}
+
+	virtual void SetExternalVariablesInstanceData(FRigVMExtendedExecuteContext& Context, const TArray<FRigVMExternalVariable>& InExternalVariables, bool bAllowNullMemory = false) override {};
+
 protected:
 
 	template<typename ExecuteContextType = FRigVMExecuteContext>
-	ExecuteContextType& UpdateContext(FRigVMExtendedExecuteContext& Context, int32 InNumberInstructions, int32 InNumberCallables, const FName& InEntryName)
+	ExecuteContextType& GetPublicContext(FRigVMExtendedExecuteContext& Context, int32 InNumberInstructions, int32 InNumberCallables, const FName& InEntryName)
 	{
-		UpdateExternalVariables(Context);
-	
 		Context.ResetExecutionState();
 		Context.SetVM(this);
 		Context.SliceOffsetsPerInstruction.SetNumZeroed(InNumberInstructions);
 		Context.SliceOffsetsPerCallable.SetNumZeroed(InNumberCallables);
 		Context.GetPublicData<ExecuteContextType>().EventName = InEntryName;
+		SetupInstructionTracking(Context, InNumberInstructions, InNumberCallables, InEntryName);
 		return Context.GetPublicData<ExecuteContextType>();
 	}
 
-	virtual void UpdateExternalVariables(FRigVMExtendedExecuteContext& Context) {};
-	
+	UE_API const FRigVMFunction* FindDispatchFunction(const FString& InIdentifier) const;
+
+#if WITH_EDITOR
+	mutable TWeakObjectPtr<URigVM> WeakSourceVM;
+#endif
+
 	class FErrorPipe : public FOutputDevice
 	{
 	public:
@@ -128,35 +245,35 @@ protected:
 		}
 	};
 
-	void BeginLoopSlice(FRigVMExtendedExecuteContext& Context, int32 InCount, int32 InRelativeIndex = 0)
+	void BeginLoopSlice(FRigVMNativizedContext& Context, int32 InCount, int32 InRelativeIndex = 0)
 	{
-		Context.BeginLoopSlice(InCount, InRelativeIndex);
+		Context.ExecuteContext.BeginLoopSlice(InCount, InRelativeIndex);
 	}
 
-	void EndLoopSlice(FRigVMExtendedExecuteContext& Context)
+	void EndLoopSlice(FRigVMNativizedContext& Context)
 	{
-		Context.EndLoopSlice();
+		Context.ExecuteContext.EndLoopSlice();
 	}
 
-	void BeginCallableSlice(FRigVMExtendedExecuteContext& Context, int32 InCallableIndex)
+	void BeginCallableSlice(FRigVMNativizedContext& Context, int32 InCallableIndex)
 	{
-		Context.BeginCallableSlice(InCallableIndex);
+		Context.ExecuteContext.BeginCallableSlice(InCallableIndex);
 	}
 
-	void EndCallableSlice(FRigVMExtendedExecuteContext& Context)
+	void EndCallableSlice(FRigVMNativizedContext& Context)
 	{
-		Context.EndCallableSlice();
+		Context.ExecuteContext.EndCallableSlice();
 	}
 
-	bool IsValidArraySize(const FRigVMExtendedExecuteContext& Context, int32 InArraySize) const
+	bool IsValidArraySize(const FRigVMNativizedContext& Context, int32 InArraySize) const
 	{
-		return Context.IsValidArraySize(InArraySize);
+		return Context.ExecuteContext.IsValidArraySize(InArraySize);
 	}
 
 	template<typename T>
-	bool IsValidArrayIndex(const FRigVMExtendedExecuteContext& Context, int32& InOutIndex, const TArray<T>& InArray) const
+	bool IsValidArrayIndex(const FRigVMNativizedContext& Context, int32& InOutIndex, const TArray<T>& InArray) const
 	{
-		return Context.IsValidArrayIndex(InOutIndex, InArray.Num());
+		return Context.ExecuteContext.IsValidArrayIndex(InOutIndex, InArray.Num());
 	}
 
 	template<typename T>
@@ -205,11 +322,21 @@ protected:
 	}
 
 	template<typename T>
-	T& GetExternalVariableRef(FRigVMExtendedExecuteContext& Context, const FName& InExternalVariableName, const FName& InExpectedType) const
+	T& GetExternalVariableRef(const FName& InExternalVariableName, const FName& InExpectedType, const TArray<FRigVMExternalVariable>& InExternalVariables, TMap<FName,int32>& InLookup) const
 	{
-		for(const FRigVMExternalVariable& ExternalVariable : GetExternalVariables(Context))
+		if (InLookup.IsEmpty())
 		{
-			if(ExternalVariable.GetName() == InExternalVariableName && ExternalVariable.GetExtendedCPPType() == InExpectedType)
+			for(int32 Index = 0; Index < InExternalVariables.Num(); Index++)
+			{
+				InLookup.Add(InExternalVariables[Index].GetName(), Index);
+			}
+		}
+
+		const int32* Index = InLookup.Find(InExternalVariableName);
+		if (InExternalVariables.IsValidIndex(*Index))
+		{
+			const FRigVMExternalVariable& ExternalVariable = InExternalVariables[*Index];
+			if(ExternalVariable.GetExtendedCPPType() == InExpectedType)
 			{
 				if(ExternalVariable.IsValid(false))
 				{
@@ -221,10 +348,16 @@ protected:
 		return EmptyValue;
 	}
 
-	template<typename T>
-	T& GetOperandSlice(FRigVMExtendedExecuteContext& Context, TArray<T>& InOutArray, const T* InDefaultValue = nullptr)
+	template<typename TargetT, typename SourceT>
+	void Copy(FRigVMNativizedContext& Context, const SourceT& InSource, TargetT& InTarget)
 	{
-		const int32 SliceIndex = Context.GetSlice().GetIndex();
+		InTarget = InSource;
+	}
+
+	template<typename ArrayType>
+	typename ArrayType::ElementType& GetOperandSlice(FRigVMNativizedContext& Context, ArrayType& InOutArray, const typename ArrayType::ElementType* InDefaultValue = nullptr)
+	{
+		const int32 SliceIndex = Context.ExecuteContext.GetSlice().GetIndex();
 		if(InOutArray.Num() <= SliceIndex)
 		{
 			const int32 FirstInsertedItem = InOutArray.AddDefaulted(1 + SliceIndex - InOutArray.Num());
@@ -239,110 +372,6 @@ protected:
 		return InOutArray[SliceIndex];
 	}
 
-	template<typename T>
-	static bool ArrayOp_Iterator(const TArray<T>& InArray, T& OutElement, int32 InIndex, int32& OutCount, float& OutRatio)
-	{
-		OutCount = InArray.Num();
-		const bool bContinue = InIndex >=0 && InIndex < OutCount;
-
-		if((OutCount <= 0) || !bContinue)
-		{
-			OutRatio = 0.f;
-		}
-		else
-		{
-			OutRatio = float(InIndex) / float(OutCount - 1);
-			OutElement = InArray[InIndex];
-		}
-		
-		return bContinue;
-	}
-
-	static bool ArrayOp_Iterator(const TArray<double>& InArray, float& OutElement, int32 InIndex, int32& OutCount, float& OutRatio)
-	{
-		double Element = (double)OutElement;
-		const bool bResult = ArrayOp_Iterator<double>(InArray, Element, InIndex, OutCount, OutRatio);
-		OutElement = (float)Element;
-		return bResult;
-	}
-
-	static bool ArrayOp_Iterator(const TArray<float>& InArray, double& OutElement, int32 InIndex, int32& OutCount, float& OutRatio)
-	{
-		float Element = (float)OutElement;
-		const bool bResult = ArrayOp_Iterator<float>(InArray, Element, InIndex, OutCount, OutRatio);
-		OutElement = (double)Element;
-		return bResult;
-	}
-
-	template<typename T>
-	static bool ArrayOp_Find(const TArray<T>& InArray, const T& InElement, int32& OutIndex)
-	{
-		OutIndex = InArray.Find(InElement);
-		return OutIndex != INDEX_NONE;
-	}
-
-	static bool ArrayOp_Find(const TArray<float>& InArray, const double& InElement, int32& OutIndex)
-	{
-		return ArrayOp_Find<float>(InArray, (float)InElement, OutIndex);
-	}
-
-	static bool ArrayOp_Find(const TArray<double>& InArray, const float& InElement, int32& OutIndex)
-	{
-		return ArrayOp_Find<double>(InArray, (double)InElement, OutIndex);
-	}
-
-	template<typename T>
-	static void ArrayOp_Union(TArray<T>& InOutA, const TArray<T>& InB)
-	{
-		for(const T& Element : InB)
-		{
-			InOutA.AddUnique(Element);
-		}
-	}
-
-	template<typename T>
-	static void ArrayOp_Difference(TArray<T>& InOutA, const TArray<T>& InB)
-	{
-		TArray<T> Temp;
-		Swap(InOutA, Temp);
-
-		for(const T& Element : Temp)
-		{
-			if(!InB.Contains(Element))
-			{
-				InOutA.Add(Element);
-			}
-		}
-		for(const T& Element : InB)
-		{
-			if(!Temp.Contains(Element))
-			{
-				InOutA.Add(Element);
-			}
-		}
-	}
-
-	template<typename T>
-	static void ArrayOp_Intersection(TArray<T>& InOutA, const TArray<T>& InB)
-	{
-		TArray<T> Temp;
-		Swap(InOutA, Temp);
-
-		for(const T& Element : Temp)
-		{
-			if(InB.Contains(Element))
-			{
-				InOutA.Add(Element);
-			}
-		}
-	}
-
-	template<typename T>
-	static void ArrayOp_Reverse(TArray<T>& InOutArray)
-	{
-		Algo::Reverse(InOutArray);
-	}
-
 public:
 
 	template <
@@ -436,97 +465,85 @@ public:
 	}
 
 protected:
-	
-	class FTransformSetter
+
+	class FDefaultValueImportErrorPipe : public FOutputDevice
 	{
 	public:
-		FTransformSetter(FTransform& InOutTransform)
-			: Transform(InOutTransform)
-		{
-		}
 
-		void SetTranslationX(const FTransform::FReal InValue) const { SetTranslationComponent(0, InValue); }
-		void SetTranslationY(const FTransform::FReal InValue) const { SetTranslationComponent(1, InValue); }
-		void SetTranslationZ(const FTransform::FReal InValue) const { SetTranslationComponent(2, InValue); }
-		void SetTranslationComponent(const uint8 InComponent, const FTransform::FReal InValue) const
-		{
-			FVector Translation = Transform.GetTranslation();
-			Translation.Component(InComponent) = InValue;
-			Transform.SetTranslation(Translation);
-		}
-		
-		void SetScaleX(const FTransform::FReal InValue) const { SetScaleComponent(0, InValue); }
-		void SetScaleY(const FTransform::FReal InValue) const { SetScaleComponent(1, InValue); }
-		void SetScaleZ(const FTransform::FReal InValue) const { SetScaleComponent(2, InValue); }
-		void SetScaleComponent(const uint8 InComponent, const FTransform::FReal InValue) const
-		{
-			FVector Scale = Transform.GetScale3D();
-			Scale.Component(InComponent) = InValue;
-			Transform.SetScale3D(Scale);
-		}
-		
-		void SetRotationX(const FTransform::FReal InValue) const { SetRotationComponent(0, InValue); }
-		void SetRotationY(const FTransform::FReal InValue) const { SetRotationComponent(1, InValue); }
-		void SetRotationZ(const FTransform::FReal InValue) const { SetRotationComponent(2, InValue); }
-		void SetRotationW(const FTransform::FReal InValue) const { SetRotationComponent(3, InValue); }
-		void SetRotationComponent(const uint8 InComponent, const FTransform::FReal InValue) const
+		FString Owner;
+		int32 NumErrors;
+
+		FDefaultValueImportErrorPipe(const FString& InOwner)
+			: FOutputDevice()
+			, Owner(InOwner)
+			, NumErrors(0)
 		{
-			FQuat Rotation = Transform.GetRotation();
-			switch(InComponent)
-			{
-				case 0:
-				{
-					Rotation.X = InValue;
-					break;
-				}
-				case 1:
-				{
-					Rotation.Y = InValue;
-					break;
-				}
-				case 2:
-				{
-					Rotation.Z = InValue;
-					break;
-				}
-				case 3:
-				default:
-				{
-					Rotation.W = InValue;
-					break;
-				}
-			}
-			Rotation.Normalize();
-			Transform.SetRotation(Rotation);
 		}
 
-	private:
-		FTransform& Transform;
+		UE_API virtual void Serialize(const TCHAR* V, ELogVerbosity::Type Verbosity, const class FName& Category) override;
 	};
 
-	class FMatrixSetter
+	struct FRigVMNativizedCallstackGuard
 	{
 	public:
-		FMatrixSetter(FMatrix& InOutMatrix)
-			: Matrix(InOutMatrix)
-		{
-		}
-
-		void SetPlane(int32 InPlaneIndex, const FPlane& InPlane)
+		FRigVMNativizedCallstackGuard(FRigVMNativizedContext& InContext, const FRigVMCallableInfo* InCallable = nullptr)
+			: Context(InContext.ExecuteContext)
+			, PublicData(InContext.PublicContextBase)
+			, PreviousInstructionIndex(INDEX_NONE)
+			, PreviousNumSlices(InContext.ExecuteContext.Slices.Num())
+			, PreviousCallstackHash(Context.CallstackHash)
+			, bIsCallable(false)
 		{
-			Matrix.M[InPlaneIndex][0] = InPlane.X;
-			Matrix.M[InPlaneIndex][1] = InPlane.Y;
-			Matrix.M[InPlaneIndex][2] = InPlane.Z;
-			Matrix.M[InPlaneIndex][3] = InPlane.W;
+			bIsCallable = InCallable != nullptr;
+			if (bIsCallable)
+			{
+				PreviousInstructionIndex = Context.GetInstructionIndex();
+				PublicData.InstructionIndex = InCallable->FirstInstruction;
+				Context.Callstack.Push({PreviousInstructionIndex, InCallable});
+				Context.CallstackHash = HashCombine(PreviousInstructionIndex, InCallable->Index, Context.CallstackHash);
+				Context.BeginCallableSlice(InCallable->Index);
+#if WITH_EDITOR
+				if (FRigVMInstructionVisitInfo* RigVMInstructionVisitInfo = Context.GetRigVMInstructionVisitInfo())
+				{
+					RigVMInstructionVisitInfo->SetCallableVisitedDuringLastRun(InCallable->Index);
+					RigVMInstructionVisitInfo->AddCallableIndexToVisitOrder(InCallable->Index);
+				}
+#endif
+			}
 		}
 
-		void SetComponent(int32 InPlaneIndex, int32 InComponentIndex, double InValue)
+		~FRigVMNativizedCallstackGuard()
 		{
-			Matrix.M[InPlaneIndex][InComponentIndex] = InValue;
+			if (bIsCallable)
+			{
+				Context.EndCallableSlice();
+				Context.CallstackHash = PreviousCallstackHash;
+				Context.Callstack.Pop();
+				PublicData.InstructionIndex = PreviousInstructionIndex;
+			}
+			else
+			{
+				while (Context.Slices.Num() > PreviousNumSlices)
+				{
+					if (Context.Slices.Last().IsLoop())
+					{
+						Context.EndLoopSlice();
+					}
+					else if (Context.Slices.Last().IsCallable())
+					{
+						Context.EndCallableSlice();
+					}
+				}
+			}
 		}
-
+		
 	private:
-		FMatrix& Matrix;
+		FRigVMExtendedExecuteContext& Context;
+		FRigVMExecuteContext& PublicData;
+		int32 PreviousInstructionIndex;
+		int32 PreviousNumSlices;
+		uint32 PreviousCallstackHash;
+		bool bIsCallable;
 	};
 
 	template <typename T,
@@ -542,22 +559,61 @@ protected:
 	{
 		return false;
 	}
+};
 
-	int32 TemporaryArrayIndex;
+template<>
+inline void URigVMNativized::Copy(FRigVMNativizedContext& Context, const double& InSource, float& InTarget)
+{
+	InTarget = static_cast<float>(InSource);
+}
 
-	UE_API const FRigVMMemoryHandle& GetLazyMemoryHandle(int32 InIndex, uint8* InMemory, const FProperty* InProperty, TFunction<ERigVMExecuteResult()> InLambda);
-	
-	template<typename T>
-	TRigVMLazyValue<T> GetLazyValue(int32 InIndex, T& InMemory, const FProperty* InProperty, TFunction<ERigVMExecuteResult()> InLambda)
+template<>
+inline void URigVMNativized::Copy(FRigVMNativizedContext& Context, const float& InSource, double& InTarget)
+{
+	InTarget = static_cast<double>(InSource);
+}
+
+template<>
+inline void URigVMNativized::Copy(FRigVMNativizedContext& Context, const TArray<double>& InSource, TArray<float>& InTarget)
+{
+	InTarget.Reset();
+	InTarget.Reserve(InSource.Num());
+	for (const double& SourceValue : InSource)
 	{
-		const FRigVMMemoryHandle& Handle = GetLazyMemoryHandle(InIndex, (uint8*)&InMemory, InProperty, InLambda);
-		return Handle.GetDataLazily<T>(INDEX_NONE);
+		InTarget.Add(static_cast<float>(SourceValue));
 	}
+}
 
-	UE_API void AllocateLazyMemoryHandles(int32 InCount);
+template<>
+inline void URigVMNativized::Copy(FRigVMNativizedContext& Context, const TArray<float>& InSource, TArray<double>& InTarget)
+{
+	InTarget.Reset();
+	InTarget.Reserve(InSource.Num());
+	for (const float& SourceValue : InSource)
+	{
+		InTarget.Add(static_cast<double>(SourceValue));
+	}
+}
 
-	TArray<FRigVMLazyBranch> LazyMemoryBranches;
-	TArray<FRigVMMemoryHandle> LazyMemoryHandles;
+struct FPlaneRef
+{
+public:
+	FPlaneRef(FMatrix& InMatrix, int32 InColumn)
+		: Matrix(InMatrix)
+		, Column(InColumn)
+		, X(Matrix.M[0][InColumn])
+		, Y(Matrix.M[1][InColumn])
+		, Z(Matrix.M[2][InColumn])
+		, W(Matrix.M[3][InColumn])
+	{
+	}
+	
+	FMatrix& Matrix;
+	int32 Column;
+	FMatrix::FReal& X;
+	FMatrix::FReal& Y;
+	FMatrix::FReal& Z;
+	FMatrix::FReal& W;
 };
 
 #undef UE_API
diff --git a/Engine/Plugins/Runtime/RigVM/Source/RigVM/Public/RigVMCore/RigVMObjectArchive.h b/Engine/Plugins/Runtime/RigVM/Source/RigVM/Public/RigVMCore/RigVMObjectArchive.h
--- a/Engine/Plugins/Runtime/RigVM/Source/RigVM/Public/RigVMCore/RigVMObjectArchive.h
+++ b/Engine/Plugins/Runtime/RigVM/Source/RigVM/Public/RigVMCore/RigVMObjectArchive.h
@@ -24,12 +24,15 @@ public:
 	UE_API void Empty();
 	UE_API bool IsEmpty() const;
 	UE_API void Compress();
-	UE_API void Decompress();
+	UE_API bool Decompress();
 	bool IsCompressed() const
 	{
 		return bIsCompressed;
 	}
 
+	UE_API FString ToString();
+	UE_API bool SetFromString(const FString& InString);
+
 	RIGVM_API friend FArchive& operator<<(FArchive& Ar, FRigVMObjectArchive& Data);
 
 	RIGVM_API void SetVersionsFromArchive(const FArchive& Ar);
diff --git a/Engine/Plugins/Runtime/RigVM/Source/RigVM/Public/RigVMCore/RigVMPropertyPath.h b/Engine/Plugins/Runtime/RigVM/Source/RigVM/Public/RigVMCore/RigVMPropertyPath.h
--- a/Engine/Plugins/Runtime/RigVM/Source/RigVM/Public/RigVMCore/RigVMPropertyPath.h
+++ b/Engine/Plugins/Runtime/RigVM/Source/RigVM/Public/RigVMCore/RigVMPropertyPath.h
@@ -118,6 +118,9 @@ public:
 	// Returns the property path as a string (the sanitized SegmentPath)
 	const FString& ToString() const { return Path; }
 
+	// Returns a descriptive name
+	UE_API FString GetName() const;
+
 	// Returns the number of segments
 	int32 Num() const { return Segments.Num(); }
 
@@ -178,6 +181,9 @@ public:
 		);
 	}
 
+	// Return the owner's native type. the owner can be a property or a struct.
+	UE_API FString GetHeadCPPType() const;
+
 	// Returns the property of the last segment
 	UE_API const FProperty* GetTailProperty() const;
 
diff --git a/Engine/Plugins/Runtime/RigVM/Source/RigVM/Public/RigVMHost.h b/Engine/Plugins/Runtime/RigVM/Source/RigVM/Public/RigVMHost.h
--- a/Engine/Plugins/Runtime/RigVM/Source/RigVM/Public/RigVMHost.h
+++ b/Engine/Plugins/Runtime/RigVM/Source/RigVM/Public/RigVMHost.h
@@ -357,7 +357,8 @@ public:
 
 	UE_API virtual void PostInitInstanceIfRequired();
 	UE_API void SwapVMToNativizedIfRequired(UClass* InNativizedClass = nullptr);
-	static UE_API bool AreNativizedVMsDisabled() ;
+	UE_API bool AreNativizedVMsEnabled() const;
+	UE_API bool CanSwapVMToNativized() const;
 
 #if WITH_EDITORONLY_DATA
 	static UE_API void DeclareConstructClasses(TArray<FTopLevelAssetPath>& OutConstructClasses, const UClass* SpecificSubclass);
@@ -404,6 +405,10 @@ public:
 	
 	TScriptInterface<IRigVMRuntimeAssetInterface> GetGeneratedByAsset() const { return GeneratedBy; }
 
+#if WITH_EDITOR
+	UE_API void EnableNativizedVM(bool bEnable = true);
+#endif
+	
 protected:
 
 	UE_API virtual void PostInitInstance(URigVMHost* InCDO);
@@ -426,8 +431,14 @@ protected:
 #if WITH_EDITOR
 	/** true if the instance is being debugged in an asset editor */
 	bool bIsBeingDebugged;
+
+	/** true if this instance is allowed to use nativized VMs */
+	bool bEnableNativizedVMs;
 #endif
 
+	mutable TOptional<uint32> LastNativizedVMHash;
+	mutable TOptional<TWeakObjectPtr<UClass>> WeakNativizedVMClass;
+
 #if RIGVM_TRACE_ENABLED
 	/** true if this instance is playing back a rewind debug trace */
 	bool bIsPlayingRewindDebugTrace;
