RigVM Nativization System

diff --git a/Engine/Plugins/Runtime/RigVM/Source/RigVMDeveloper/Private/RigVMCompiler/RigVMCodeGenerator.cpp b/Engine/Plugins/Runtime/RigVM/Source/RigVMDeveloper/Private/RigVMCompiler/RigVMCodeGenerator.cpp
--- a/Engine/Plugins/Runtime/RigVM/Source/RigVMDeveloper/Private/RigVMCompiler/RigVMCodeGenerator.cpp
+++ b/Engine/Plugins/Runtime/RigVM/Source/RigVMDeveloper/Private/RigVMCompiler/RigVMCodeGenerator.cpp
@@ -1,2493 +1,261 @@
 // Copyright Epic Games, Inc. All Rights Reserved.
 
 #include "RigVMCompiler/RigVMCodeGenerator.h"
-#include "RigVMCore/RigVMExecuteContext.h"
-#include "RigVMCore/RigVMStruct.h"
+#include "HAL/FileManagerGeneric.h"
+#include "HAL/IConsoleManager.h"
+#include "Misc/FileHelper.h"
+#include "Misc/ScopedSlowTask.h"
 #include "RigVMDeveloperModule.h"
-#include "Algo/Count.h"
-#include "RigVMModel/Nodes/RigVMDispatchNode.h"
-#include "RigVMStringUtils.h"
-
-static constexpr TCHAR RigVM_CommaSeparator[] = TEXT(", ");
-static constexpr TCHAR RigVM_NewLineFormat[] = TEXT("\r\n");
-static constexpr TCHAR RigVM_IncludeBracketFormat[] = TEXT("#include <{0}>");
-static constexpr TCHAR RigVM_IncludeQuoteFormat[] = TEXT("#include \"{0}.h\"");
-static constexpr TCHAR RigVM_DispatchKeyFormat[] = TEXT("{0}_{1}");
-static constexpr TCHAR RigVM_DispatchDeclarationFormat[] = TEXT("\tbool {0}({1})\r\n\t{\r\n\t\tstatic const FRigVMFunction* Dispatch = FRigVMRegistry::Get().FindFunction(TEXT(\"{2}\"));\r\n\t\tif(Dispatch == nullptr) return false;");
-static constexpr TCHAR RigVM_UPropertyDeclareFormat[] = TEXT("\tUPROPERTY()\r\n\t{0} {1};");
-static constexpr TCHAR RigVM_UPropertyMemberFormat[] = TEXT("\tstatic const FProperty* {0}_Ptr;");
-static constexpr TCHAR RigVM_UPropertyMember2Format[] = TEXT("const FProperty* U{0}::{1}_Ptr = nullptr;");
-static constexpr TCHAR RigVM_UPropertyDefineFormat[] = TEXT("\tif({1}_Ptr == nullptr)\r\n\t{\r\n\t\t{1}_Ptr = StaticClass()->FindPropertyByName(GET_MEMBER_NAME_CHECKED(U{0}, {1}));\r\n\t}");
-static constexpr TCHAR RigVM_InvokeDispatchPtrFormat[] = TEXT("\t\t(*Dispatch->FunctionPtr)({0}, FRigVMPredicateBranchArray()); // todo: predicates are not implemented yet");
-static constexpr TCHAR RigVM_InvokeDispatchFormat[] = TEXT("\t{0}({1});");
-static constexpr TCHAR RigVM_WrappedArrayTypeFormat[] = TEXT("struct {0}_API {1}\r\n{\r\n\tTArray<{2}> Array;\r\n};");
-static constexpr TCHAR RigVM_WrappedTypeNameFormat[] = TEXT("{0}Array_{1}");
-static constexpr TCHAR RigVM_DeclareExternalVariableFormat[] = TEXT("\t{0}* {1} = nullptr;");
-static constexpr TCHAR RigVM_UpdateExternalVariableFormat[] = TEXT("\t{0} = &GetExternalVariableRef<{1}>(Context, TEXT(\"{2}\"), TEXT(\"{1}\"));");
-static constexpr TCHAR RigVM_MemberPropertyFormat[] = TEXT("\t{0} {1} = {2};");
-static constexpr TCHAR RigVM_MemberPropertyFormatNoDefault[] = TEXT("\t{0} {1};");
-static constexpr TCHAR RigVM_DeclareEntryNameFormat[] = TEXT("\tstatic const FLazyName EntryName_{0};");
-static constexpr TCHAR RigVM_DefineEntryNameFormat[] = TEXT("const FName U{0}::EntryName_{1} = TEXT(\"{2}\");");
-static constexpr TCHAR RigVM_DeclareBlockNameFormat[] = TEXT("\tstatic const FLazyName BlockName_{0};");
-static constexpr TCHAR RigVM_DefineBlockNameFormat[] = TEXT("const FName U{0}::BlockName_{1} = TEXT(\"{2}\");");
-static constexpr TCHAR RigVM_DefineConstFormatNoDefault[] = TEXT("\tstatic const {0} {1};");
-static constexpr TCHAR RigVM_StructConstantArrayArrayValue[] = TEXT("URigVMNativized::GetStructArrayArrayConstant<{0}>(TEXT(\"{1}\"))");
-static constexpr TCHAR RigVM_StructConstantArrayValue[] = TEXT("URigVMNativized::GetStructArrayConstant<{0}>(TEXT(\"{1}\"))");
-static constexpr TCHAR RigVM_StructConstantValue[] = TEXT("URigVMNativized::GetStructConstant<{0}>(TEXT(\"{1}\"))");
-static constexpr TCHAR RigVM_DefineConstFormat[] = TEXT("\tstatic const {0} {1} = {2};");
-static constexpr TCHAR RigVM_NameNoneFormat[] = TEXT("FName(NAME_None)");
-static constexpr TCHAR RigVM_EmptyStringFormat[] = TEXT("FString()");
-static constexpr TCHAR RigVM_SingleStringFormat[] = TEXT("%s");
-static constexpr TCHAR RigVM_TextFormat[] = TEXT("TEXT({0})");
-static constexpr TCHAR RigVM_QuotedTextFormat[] = TEXT("TEXT(\"{0}\")");
-static constexpr TCHAR RigVM_CurlyBracesFormat[] = TEXT("{{0}}");
-static constexpr TCHAR RigVM_BracesFormat[] = TEXT("({0})");
-static constexpr TCHAR RigVM_TemplateOneArgFormat[] = TEXT("<{0}>");
-static constexpr TCHAR RigVM_CallExternOpFormat[] = TEXT("\t{0}::Static{1}({2});");
-static constexpr TCHAR RigVM_ZeroOpIntFormat[] = TEXT("\t{0} = 0;");
-static constexpr TCHAR RigVM_ZeroOpNameFormat[] = TEXT("\t{0} = NAME_None;");
-static constexpr TCHAR RigVM_BoolFalseOpFormat[] = TEXT("\t{0} = false;");
-static constexpr TCHAR RigVM_BoolTrueFormat[] = TEXT("\t{0} = true;");
-static constexpr TCHAR RigVM_CopyUnrelatedArraysFormat[] = TEXT("\tCopyUnrelatedArrays<{0}, {1}>({2}, {3});");
-static constexpr TCHAR RigVM_CopyOpMethodFormat[] = TEXT("\t{0}{1}{2});");
-static constexpr TCHAR RigVM_CopyOpAssignFormat[] = TEXT("\t{0} = {1}{2};");
-static constexpr TCHAR RigVM_IncrementOpFormat[] = TEXT("\t{0}++;");
-static constexpr TCHAR RigVM_DecrementOpFormat[] = TEXT("\t{0}--;");
-static constexpr TCHAR RigVM_EqualsOpFormat[] = TEXT("\t{0} = {1} == {2};");
-static constexpr TCHAR RigVM_InvokeEntryFormat[] = TEXT("\tif (InEntryName == EntryName_{0}) return ExecuteEntry_{0}(Context, PublicContext);");
-static constexpr TCHAR RigVM_InvokeEntryByNameFormat[] = TEXT("\tERigVMExecuteResult EntryResult = InvokeEntryByName(Context, InEntryName{0});\r\n\tSetInstructionIndex(Context, 0);\r\n\tStopProfiling(Context);\r\n\treturn EntryResult;");
-static constexpr TCHAR RigVM_InvokeEntryByNameFormat2[] = TEXT("\tif(!InvokeEntryByName(Context, {0}{1})) return ERigVMExecuteResult::Failed;");
-static constexpr TCHAR RigVM_CanExecuteEntryFormat[] = TEXT("\tif(!CanExecuteEntry(Context, InEntryName, false)) { return ERigVMExecuteResult::Failed; }");
-static constexpr TCHAR RigVM_EntryExecuteGuardFormat[] = TEXT("\tFEntryExecuteGuard EntryExecuteGuard(EntriesBeingExecuted, FindEntry(InEntryName));");
-static constexpr TCHAR RigVM_PublicContextGuardFormat[] = TEXT("\tTGuardValue<{0}> PublicContextGuard(Context.GetPublicData<{0}>(), PublicContext);");
-static constexpr TCHAR RigVM_EntryNameFormat[] = TEXT("EntryName_{0}");
-static constexpr TCHAR RigVM_SetExecuteContextStructFormat[] = TEXT("\tContext.SetContextPublicDataStruct({0}::StaticStruct());");
-static constexpr TCHAR RigVM_UpdateContextFormat[] = TEXT("\t{0}& PublicContext = UpdateContext<{0}>(Context, {1}, InEntryName);");
-static constexpr TCHAR RigVM_TrueFormat[] = TEXT("true");  
-static constexpr TCHAR RigVM_FalseFormat[] = TEXT("false");
-static constexpr TCHAR RigVM_SingleUnderscoreFormat[] = TEXT("_");
-static constexpr TCHAR RigVM_DoubleUnderscoreFormat[] = TEXT("__");
-static constexpr TCHAR RigVM_BoolPropertyPrefix[] = TEXT("b");
-static constexpr TCHAR RigVM_EnumTypeSuffixFormat[] = TEXT("::Type");
-static constexpr TCHAR RigVM_IsValidArraySizeFormat[] = TEXT("IsValidArraySize(Context, {0})");
-static constexpr TCHAR RigVM_IsValidArrayIndexFormat[] = TEXT("IsValidArrayIndex<{0}>(Context, TemporaryArrayIndex, {1})");
-static constexpr TCHAR RigVM_TemporaryArrayIndexFormat[] = TEXT("\tTemporaryArrayIndex = {0};");
-static constexpr TCHAR RigVM_StartProfilingFormat[] = TEXT("\tStartProfiling(Context);");
-static constexpr TCHAR RigVM_ExecuteReachedExitFormat[] = TEXT("\tBroadcastExecutionReachedExit(Context);");
-static constexpr TCHAR RigVM_InstructionLabelFormat[] = TEXT("\tInstruction{0}Label:");
-static constexpr TCHAR RigVM_SetInstructionIndexFormat[] = TEXT("\tSetInstructionIndex(Context, {0});");
-static constexpr TCHAR RigVM_ContextRefParamFormat[] = TEXT("FRigVMExtendedExecuteContext& Context");
-static constexpr TCHAR RigVM_ContextFormat[] = TEXT("Context");
-static constexpr TCHAR RigVM_ContextPublicFormat[] = TEXT("PublicContext");
-static constexpr TCHAR RigVM_ContextPublicParameterFormat[] = TEXT("{0}& PublicContext");
-static constexpr TCHAR RigVM_NotEqualsOpFormat[] = TEXT("\t{0} = {1} != {2};");
-static constexpr TCHAR RigVM_JumpOpFormat[] = TEXT("\tgoto Instruction{0}Label;");
-static constexpr TCHAR RigVM_JumpIfOpFormat[] = TEXT("\tif ({0} == {1}) { goto Instruction{2}Label; }");
-static constexpr TCHAR RigVM_JumpToBranchFormat[] = TEXT("\tif ({0} == BlockName_{1}) { goto Instruction{2}Label; }");
-static constexpr TCHAR RigVM_BeginBlockOpFormat[] = TEXT("\tBeginSlice(Context, {0}, {1});");
-static constexpr TCHAR RigVM_EndBlockOpFormat[] = TEXT("\tEndSlice(Context);");
-static constexpr TCHAR RigVM_ReturnFailedFormat[] = TEXT("\treturn ERigVMExecuteResult::Failed;");
-static constexpr TCHAR RigVM_ReturnSucceededFormat[] = TEXT("\treturn ERigVMExecuteResult::Succeeded;");
-static constexpr TCHAR RigVM_CopyrightFormat[] = TEXT("// Copyright Epic Games, Inc. All Rights Reserved.");
-static constexpr TCHAR RigVM_AutoGeneratedFormat[] = TEXT("// THIS FILE HAS BEEN AUTO-GENERATED. PLEASE DO NOT MANUALLY EDIT THIS FILE FURTHER.");
-static constexpr TCHAR RigVM_PragmaOnceFormat[] = TEXT("#pragma once");
-static constexpr TCHAR RigVM_GeneratedIncludeFormat[] = TEXT("#include \"{0}.generated.h\"");
-static constexpr TCHAR RigVM_UClassDefinitionFormat[] = TEXT("UCLASS()\r\nclass {0}_API U{1} : public URigVMNativized\r\n{\r\n\tGENERATED_BODY()\r\npublic:\r\n\tU{1}() {}\r\n\tvirtual ~U{1}() override {}\r\n");
-static constexpr TCHAR RigVM_ProtectedFormat[] = TEXT("protected:");
-static constexpr TCHAR RigVM_GetVMHashFormat[] = TEXT("\tvirtual uint32 GetVMHash() const override { return {0}; }");
-static constexpr TCHAR RigVM_GetEntryNamesFormat[] = TEXT("\tvirtual const TArray<FName>& GetEntryNames() const override\r\n\t{\r\n\t\tstatic const TArray<FName> StaticEntryNames = { {0} };\r\n\t\treturn StaticEntryNames;\r\n\t}");
-static constexpr TCHAR RigVM_DeclareUpdateExternalVariablesFormat[] = TEXT("\tvirtual void UpdateExternalVariables(FRigVMExtendedExecuteContext& Context) override;");
-static constexpr TCHAR RigVM_DeclareInvokeEntryByNameFormat[] = TEXT("\tERigVMExecuteResult InvokeEntryByName(FRigVMExtendedExecuteContext& Context, const FName& InEntryName{0});");
-static constexpr TCHAR RigVM_DeclareInitializeFormat[] = TEXT("\tvirtual bool Initialize(FRigVMExtendedExecuteContext& Context) override;");
-static constexpr TCHAR RigVM_DefineInitializeFormat[] = TEXT("bool U{0}::Initialize(FRigVMExtendedExecuteContext& Context)\r\n{");
-static constexpr TCHAR RigVM_DeclareExecuteFormat[] = TEXT("\tvirtual ERigVMExecuteResult ExecuteVM(FRigVMExtendedExecuteContext& Context, const FName& InEntryName = NAME_None) override;");
-static constexpr TCHAR RigVM_DefineUpdateExternalVariablesFormat[] = TEXT("void U{0}::UpdateExternalVariables(FRigVMExtendedExecuteContext& Context)\r\n{");
-static constexpr TCHAR RigVM_DefineInvokeEntryByNameFormat[] = TEXT("ERigVMExecuteResult U{0}::InvokeEntryByName(FRigVMExtendedExecuteContext& Context, const FName& InEntryName{1})\r\n{");
-static constexpr TCHAR RigVM_DefineExecuteFormat[] = TEXT("ERigVMExecuteResult U{0}::ExecuteVM(FRigVMExtendedExecuteContext& Context, const FName& InEntryName)\r\n{");
-static constexpr TCHAR RigVM_DeclareExecuteEntryFormat[] = TEXT("\tERigVMExecuteResult ExecuteEntry_{0}(FRigVMExtendedExecuteContext& Context, {1});");
-static constexpr TCHAR RigVM_DefineExecuteEntryFormat[] = TEXT("ERigVMExecuteResult U{0}::ExecuteEntry_{1}(FRigVMExtendedExecuteContext& Context, {2})\r\n{");
-static constexpr TCHAR RigVM_DeclareExecuteGroupFormat[] = TEXT("\tERigVMExecuteResult ExecuteGroup_{0}_{1}({2});");
-static constexpr TCHAR RigVM_DefineExecuteGroupFormat[] = TEXT("ERigVMExecuteResult U{0}::ExecuteGroup_{1}_{2}({3})\r\n{");
-static constexpr TCHAR RigVM_InvokeExecuteGroupFormat[] = TEXT("\tif(ExecuteGroup_{0}_{1}({2}) != ERigVMExecuteResult::Succeeded) return ERigVMExecuteResult::Failed;");
-static constexpr TCHAR RigVM_RigVMCoreIncludeFormat[] = TEXT("RigVMCore/RigVMCore.h");
-static constexpr TCHAR RigVM_RigVMModuleIncludeFormat[] = TEXT("RigVMModule.h");
-static constexpr TCHAR RigVM_RigVMCoreLibraryFormat[] = TEXT("RigVM");
-static constexpr TCHAR RigVM_JoinFilePathFormat[] = TEXT("{0}/{1}");
-static constexpr TCHAR RigVM_GetOperandSliceFormat[] = TEXT("GetOperandSlice<{0}>(Context, {1},&{1}_Const){2}");
-static constexpr TCHAR RigVM_ExternalVariableFormat[] = TEXT("(*External_{0})");
-static constexpr TCHAR RigVM_JoinSegmentPathFormat[] = TEXT("{0}.{1}");
-static constexpr TCHAR RigVM_GetArrayElementSafeFormat[] = TEXT("GetArrayElementSafe<{0}>({1}, {2})");
-static constexpr TCHAR RigVM_InvokeEntryOpFormat[] = TEXT("\tif(InvokeEntryByName(Context, EntryName_{0}) != ERigVMExecuteResult::Succeeded) return ERigVMExecuteResult::Failed;");
-static constexpr TCHAR RigVM_LazyEvalValueName[] =  TEXT("LazyValue_{0}_{1}");
-static constexpr TCHAR RigVM_LazyEvalLambdaDefine[] =  TEXT("\tconst TRigVMLazyValue<{2}> LazyValue_{0}_{1} = GetLazyValue<{2}>({0}, {3}, {4}_Ptr,\r\n\t\t[&]() -> ERigVMExecuteResult\r\n\t\t{");
-static constexpr TCHAR RigVM_LazyEvalLambdaReturn[] =  TEXT("\t\t\treturn ERigVMExecuteResult::Succeeded;\r\n\t\t}\r\n\t);");
-static constexpr TCHAR RigVM_LazyMemoryHandleInitFormat[] = TEXT("\tAllocateLazyMemoryHandles({0});");
-static constexpr TCHAR RigVM_SetupInstructionTrackingFormat[] = TEXT("\tSetupInstructionTracking(Context, {0});");
-
-FString FRigVMCodeGenerator::DumpIncludes(bool bLog)
-{
-	FStringArray Lines;
-	for(const FString& Include : Includes)
-	{
-		Lines.Add(Format(RigVM_IncludeBracketFormat, Include));
-	}
-	return DumpLines(Lines, bLog);
-}
-
-FString FRigVMCodeGenerator::DumpExternalVariables(const FRigVMExtendedExecuteContext& Context, bool bForHeader, bool bLog)
-{
-	FStringArray Lines;
-
-	if(bForHeader)
-	{
-		Lines.Emplace();
-	}
-	
-	for(int32 ExternalVariableIndex = 0; ExternalVariableIndex < VM->GetExternalVariableDefs().Num(); ExternalVariableIndex++)
-	{
-		const FRigVMOperand ExternalVarOperand(ERigVMMemoryType::External, ExternalVariableIndex, INDEX_NONE);
-		const FRigVMExternalVariableDef& ExternalVariable = VM->GetExternalVariableDefs()[ExternalVariableIndex]; //-V758
-		const FString ExternalVarCPPType = ExternalVariable.GetExtendedCPPType().ToString();
-		FString OperandName = *GetOperandName(Context, ExternalVarOperand, false);
-		if(OperandName.StartsWith(TEXT("(*")) && OperandName.EndsWith(TEXT(")")))
-		{
-			OperandName = OperandName.Mid(2, OperandName.Len() - 3);
-		}
-
-		if(bForHeader)
-		{
-			Lines.Add(Format(RigVM_DeclareExternalVariableFormat, *ExternalVarCPPType, *OperandName));
-		}
-		else
-		{
-			Lines.Add(Format(RigVM_UpdateExternalVariableFormat, *OperandName, *ExternalVarCPPType, *ExternalVariable.GetName().ToString(), *ExternalVarCPPType));
-		}
-	}
-
-	return DumpLines(Lines, bLog);
-}
-
-FString FRigVMCodeGenerator::DumpEntries(bool bForHeader, bool bLog)
-{
-	FStringArray Lines;
-
-	const FRigVMByteCode& ByteCode = VM->GetByteCode();
-	for(int32 EntryIndex = 0; EntryIndex < ByteCode.NumEntries(); EntryIndex++)
-	{
-		const FRigVMByteCodeEntry& Entry = ByteCode.GetEntry(EntryIndex);
-		const FString EntryName = Entry.GetSanitizedName();
-		if(bForHeader)
-		{
-			Lines.Add(Format(RigVM_DeclareEntryNameFormat, *EntryName));
-		}
-		else
-		{
-			Lines.Add(Format(RigVM_DefineEntryNameFormat, *ClassName, *EntryName, *Entry.Name.ToString()));
-		}
-	}
-
-	return DumpLines(Lines, bLog);
-}
-
-FString FRigVMCodeGenerator::DumpBlockNames(bool bForHeader, bool bLog)
-{
-	FStringArray Lines;
-
-	const FRigVMByteCode& ByteCode = VM->GetByteCode();
-	for(const FRigVMBranchInfo& BranchInfo : ByteCode.BranchInfos)
-	{
-		if(!BranchInfo.IsOutputBranch())
-		{
-			continue;
-		}
-		
-		const FString BlockName = BranchInfo.Label.ToString();
-		if(bForHeader)
-		{
-			Lines.AddUnique(Format(RigVM_DeclareBlockNameFormat, *BlockName));
-		}
-		else
-		{
-			Lines.AddUnique(Format(RigVM_DefineBlockNameFormat, *ClassName, *BlockName, *BlockName));
-		}
-	}
-
-	return DumpLines(Lines, bLog);
-}
-
-FString FRigVMCodeGenerator::DumpProperties(const FRigVMExtendedExecuteContext& Context, bool bForHeader, int32 InInstructionGroup, bool bLog)
-{
-	if(bForHeader)
-	{
-		// for headers we show all properties
-		check(InInstructionGroup == INDEX_NONE);
-	}
-
-	FStringArray Lines;
-	for(int32 Index = 0; Index < Properties.Num(); Index++)
-	{
-		const FPropertyInfo& PropertyInfo = Properties[Index];
-		
-		// in headers we only dump the work / sliced properties,
-		// and for source files we only dump the non-sliced (and initialized sliced)
-		if(PropertyInfo.PropertyType != ERigVMNativizedPropertyType::Sliced &&
-			bForHeader != (PropertyInfo.PropertyType == ERigVMNativizedPropertyType::Work &&
-								PropertyInfo.Groups.Num() > 1))
-		{
-			continue;
-		}
-
-		const FRigVMPropertyDescription& Property = PropertyInfo.Description;
-		check(Property.IsValid());
-
-		if(PropertyInfo.PropertyType == ERigVMNativizedPropertyType::Literal ||
-			(!bForHeader && PropertyInfo.PropertyType == ERigVMNativizedPropertyType::Sliced))
-		{
-			FRigVMOperand Operand;
-			if (PropertyInfo.PropertyType == ERigVMNativizedPropertyType::Literal)
-			{
-				Operand = FRigVMOperand(ERigVMMemoryType::Literal, PropertyInfo.MemoryPropertyIndex, INDEX_NONE);
-			}
-			else if (PropertyInfo.PropertyType == ERigVMNativizedPropertyType::Sliced)
-			{
-				Operand = FRigVMOperand(ERigVMMemoryType::Work, PropertyInfo.MemoryPropertyIndex, INDEX_NONE);
-			}
-			FString OperandName = GetOperandName(Context, Operand, false);
-			FString CPPType = GetOperandCPPType(Context, Operand);
-
-			FString BaseCPPType = CPPType;
-			bool bIsArray = RigVMTypeUtils::IsArrayType(CPPType);
-			bool bIsDoubleArray = false;
-			if (bIsArray)
-			{
-				BaseCPPType = RigVMTypeUtils::BaseTypeFromArrayType(CPPType);
-				bIsDoubleArray = RigVMTypeUtils::IsArrayType(BaseCPPType);
-				if (bIsDoubleArray)
-				{
-					BaseCPPType = RigVMTypeUtils::BaseTypeFromArrayType(BaseCPPType);
-				}
-			}
-			
-			FString DefaultValue = Property.DefaultValue;
-			if (!bForHeader && PropertyInfo.PropertyType == ERigVMNativizedPropertyType::Sliced)
-			{
-				// The const definition of a slice should have the element type
-				OperandName += TEXT("_Const");
-				if (bIsDoubleArray)
-				{
-					CPPType = RigVMTypeUtils::ArrayTypeFromBaseType(BaseCPPType);
-					bIsArray = true;
-					bIsDoubleArray = false;
-				}
-				else
-				{
-					CPPType = BaseCPPType;
-					bIsArray = RigVMTypeUtils::IsArrayType(CPPType);
-					if (bIsArray)
-					{
-						BaseCPPType = RigVMTypeUtils::BaseTypeFromArrayType(BaseCPPType);
-					}
-				}
-				DefaultValue = DefaultValue.LeftChop(1);
-				DefaultValue = DefaultValue.RightChop(1);
-			}
-			DefaultValue = SanitizeValue(DefaultValue, CPPType, PropertyInfo.Description.CPPTypeObject);
-
-			if (UScriptStruct* ScriptStruct = Cast<UScriptStruct>(Property.CPPTypeObject))
-			{
-				if(DefaultValue.IsEmpty())
-				{
-					Lines.Add(Format(
-						RigVM_DefineConstFormatNoDefault,
-						*CPPType,
-						*OperandName
-					));
-				}
-				else
-				{
-					Lines.Add(Format(
-						RigVM_DefineConstFormat,
-						*CPPType,
-						*OperandName,
-						*DefaultValue
-					));
-				}
-			}
-			else if (const UEnum* Enum = Cast<UEnum>(Property.CPPTypeObject))
-			{
-				BaseCPPType = Enum->GetName();
-				if (Enum->GetCppForm() == UEnum::ECppForm::Namespaced)
-				{
-					BaseCPPType += RigVM_EnumTypeSuffixFormat;
-				}
-				
-				CPPType = BaseCPPType;
-				if (bIsArray)
-				{
-					CPPType = RigVMTypeUtils::ArrayTypeFromBaseType(BaseCPPType);
-				}
-				if (bIsDoubleArray)
-				{
-					CPPType = RigVMTypeUtils::ArrayTypeFromBaseType(CPPType);
-				}
-
-				if(DefaultValue.IsEmpty())
-				{
-					Lines.Add(Format(
-						RigVM_DefineConstFormatNoDefault,
-						*CPPType,
-						*OperandName
-					));
-				}
-				else
-				{
-					Lines.Add(Format(
-						RigVM_DefineConstFormat,
-						*CPPType,
-						*OperandName,
-						*DefaultValue
-					));
-				}
-			}
-			else
-			{
-				bool bUseConstExpr = true;
-
-				if (bIsArray)
-				{
-					bUseConstExpr = false;
-				}
-
-				if (CPPType == RigVMTypeUtils::FNameType || CPPType == RigVMTypeUtils::FStringType)
-				{
-					bUseConstExpr = false;
-				}
-
-				if(DefaultValue.IsEmpty())
-				{
-					Lines.Add(Format(
-						RigVM_DefineConstFormatNoDefault, 
-						*CPPType,
-						*OperandName
-					));
-				}
-				else
-				{
-					Lines.Add(Format(
-						RigVM_DefineConstFormat, 
-						*CPPType,
-						*OperandName,
-						*DefaultValue
-					));
-				}
-			}
-		}
-		else // work and slice look the same in the file
-		{
-			FRigVMOperand Operand(ERigVMMemoryType::Work, PropertyInfo.MemoryPropertyIndex, INDEX_NONE);
-			FString OperandName = GetOperandName(Context, Operand, false);
-			FString CPPType = GetOperandCPPType(Context, Operand);
-
-			const FString MappedType = GetMappedType(Property.CPPType);
-
-			if (bForHeader && PropertyInfo.PropertyType == ERigVMNativizedPropertyType::Sliced)
-			{
-				const FString Line = Format(RigVM_MemberPropertyFormatNoDefault, *MappedType, *SanitizeName(Property.Name.ToString(), Property.CPPType));
-				Lines.Add(Line);
-			}
-			else
-			{
-				const FString DefaultValue = SanitizeValue(PropertyInfo.Description.DefaultValue, Property.CPPType, Property.CPPTypeObject);
-				const FString Line = Format(RigVM_MemberPropertyFormat, *MappedType, *SanitizeName(Property.Name.ToString(), Property.CPPType), DefaultValue);
-				Lines.Add(Line);
-			}
-		}
-	}
-
-	return DumpLines(Lines, bLog);
-}
-
-FString FRigVMCodeGenerator::DumpDispatches(bool bLog)
-{
-	FStringArray Lines;
-
-	if(!Dispatches.IsEmpty())
-	{
-		for(const TPair<FString, FRigVMDispatchInfo>& Pair : Dispatches)
-		{
-			const FRigVMDispatchInfo& Info = Pair.Value;
-			const FString FunctionName = Info.Function->GetName();
-			const FRigVMDispatchFactory* Factory = Info.Function->Factory;
-
-			TArray<FString> InputArguments, OutputArguments, MemoryHandles;
-			InputArguments.Add(RigVM_ContextRefParamFormat);
-			OutputArguments.Add(RigVM_ContextFormat);
-
-			for(const FRigVMFunctionArgument& Argument : Info.Function->Arguments)
-			{
-				const TRigVMTypeIndex TypeIndex = FRigVMRegistry::Get().GetTypeIndexFromCPPType(Argument.Type);
-				const FString DispatchArgument = RequiredUProperties.FindChecked(TypeIndex).Get<1>();
-				const ERigVMPinDirection Direction = Factory->GetTemplate()->FindArgument(Argument.Name)->GetDirection();
-				const FString ConstPrefix = (Direction == ERigVMPinDirection::Visible || Direction == ERigVMPinDirection::Input) ? TEXT("const ") : FString();
-
-				if(Factory->IsLazyInputArgument(Argument.Name))
-				{
-					InputArguments.Add(Format(TEXT("{0}TRigVMLazyValue<{1}>& {2}"), ConstPrefix, Argument.Type, Argument.Name));
-					MemoryHandles.Add(Format(TEXT("{0}.GetMemoryHandle()"), Argument.Name));
-				}
-				else
-				{
-					InputArguments.Add(Format(TEXT("{0}{1}& {2}"), ConstPrefix, Argument.Type, Argument.Name));
-					MemoryHandles.Add(Format(TEXT("{(uint8*)&{0}, {1}_Ptr, nullptr}"), Argument.Name, *DispatchArgument));
-				}
-			}
-
-			OutputArguments.Add(TEXT("MemoryHandles"));
-			
-			Lines.Emplace();
-			Lines.Add(Format(RigVM_DispatchDeclarationFormat, Info.Name, FString::Join(InputArguments, RigVM_CommaSeparator), *FunctionName));
-			Lines.Add(Format(TEXT("\t\tTArray<FRigVMMemoryHandle> MemoryHandles = {\r\n\t\t\t{0}\r\n\t\t};"), FString::Join(MemoryHandles, TEXT(",\r\n\t\t\t"))));
-			Lines.Add(Format(RigVM_InvokeDispatchPtrFormat, FString::Join(OutputArguments, RigVM_CommaSeparator)));
-			Lines.Add(TEXT("\t\treturn true;\r\n\t}"));
-		}
-	}
-	
-	return DumpLines(Lines, bLog); 
-}
-
-FString FRigVMCodeGenerator::DumpRequiredUProperties(bool bLog)
-{
-	FStringArray Lines;
-
-	if(!RequiredUProperties.IsEmpty())
-	{
-		for(auto Pair : RequiredUProperties)
-		{
-			Lines.Emplace();
-			Lines.Add(Format(RigVM_UPropertyDeclareFormat, Pair.Value.Get<0>(), Pair.Value.Get<1>()));
-		}
-
-		Lines.Emplace();
-
-		for(auto Pair : RequiredUProperties)
-		{
-			Lines.Add(Format(RigVM_UPropertyMemberFormat, Pair.Value.Get<1>()));
-		}
-	}
-	
-	return DumpLines(Lines, bLog); 
-}
-
-
-FString FRigVMCodeGenerator::DumpInitialize(bool bLog)
-{
-	FStringArray Lines;
-	Lines.Add(Format(RigVM_LazyMemoryHandleInitFormat, VM->GetByteCode().BranchInfos.Num()));
-
-	for(auto Pair : RequiredUProperties)
-	{
-		Lines.Add(Format(RigVM_UPropertyDefineFormat, *ClassName, Pair.Value.Get<1>()));
-	}
-	
-	// we'll add workstate initialization here later
-	Lines.Add(TEXT("\treturn true;"));
-	return DumpLines(Lines, bLog); 
-}
-
-FString FRigVMCodeGenerator::DumpInstructions(const FRigVMExtendedExecuteContext& Context, int32 InInstructionGroup, bool bLog)
-{
-	const FRigVMByteCode& ByteCode = VM->GetByteCode();
-	const FRigVMInstructionArray Instructions = ByteCode.GetInstructions();
-
-	FStringArray Lines;
-
-	const FInstructionGroup& Group = GetGroup(InInstructionGroup);
-
-	if(Group.Entry.IsEmpty())
-	{
-		if(InInstructionGroup == INDEX_NONE)
-		{
-			Lines.Add(Format(RigVM_SetExecuteContextStructFormat, ExecuteContextType));
-			Lines.Add(Format(RigVM_UpdateContextFormat, ExecuteContextType, Instructions.Num()));
-			Lines.Add(Format(RigVM_InvokeEntryByNameFormat, *GetEntryParameters()));
-
-			return DumpLines(Lines, bLog);
-		}
-		else
-		{
-			Lines.Add(FString(RigVM_CanExecuteEntryFormat));
-			Lines.Emplace();
-			Lines.Add(FString(RigVM_EntryExecuteGuardFormat));
-			Lines.Add(Format(RigVM_PublicContextGuardFormat, ExecuteContextType));
-			Lines.Emplace();
-			
-			for(int32 EntryIndex = 0; EntryIndex < ByteCode.NumEntries(); EntryIndex++)
-			{
-				const FRigVMByteCodeEntry& Entry = ByteCode.GetEntry(EntryIndex);
-				const FString EntryName = Entry.GetSanitizedName();
-				Lines.Add(Format(RigVM_InvokeEntryFormat, *EntryName));
-			}
-		}
-	}
-
-	if(Group.Entry.IsEmpty() && InstructionGroups.Num() > 0)
-	{
-		Lines.Add(RigVM_ReturnFailedFormat);
-		return DumpLines(Lines, bLog);
-	}
-
-	if(Group.ChildGroups.IsEmpty())
-	{
-		TArray<int32> InstructionIndices = GetInstructionIndicesFromRange(Group.First, Group.Last);
-
-		// dump all lambdas and remove those instructions from instructions to process
-		for(const FRigVMBranchInfo& BranchInfo : ByteCode.BranchInfos)
-		{
-			if(BranchInfo.IsOutputBranch())
-			{
-				continue;
-			}
-			
-			if(BranchInfo.FirstInstruction < Group.First ||
-				BranchInfo.LastInstruction > Group.Last)
-			{
-				continue;
-			}
-
-			InstructionIndices.RemoveAll([BranchInfo](int32 Index) {
-				return FMath::IsWithinInclusive<int32>(Index, BranchInfo.FirstInstruction, BranchInfo.LastInstruction);
-			});
-
-			// find the operand this lazy eval lambda belongs to
-			const FRigVMOperandArray Operands = ByteCode.GetOperandsForOp(Instructions[BranchInfo.InstructionIndex]);
-			check(Operands.IsValidIndex(BranchInfo.ArgumentIndex));
-			const FRigVMOperand Operand = Operands[BranchInfo.ArgumentIndex];
-			const FString OperandName = GetOperandName(Context, Operand, false, true);
-			const FString OperandCPPType = GetOperandCPPType(Context, Operand);
-			const TRigVMTypeIndex OperandTypeIndex = FRigVMRegistry::Get().GetTypeIndexFromCPPType(OperandCPPType);
-			const FString PropertyName = RequiredUProperties.FindChecked(OperandTypeIndex).Get<1>();
-
-			// dump the instructions for the lambda wrapped with the lambda definition
-			Lines.Add(Format(RigVM_LazyEvalLambdaDefine, BranchInfo.Index, BranchInfo.Label.ToString(), *OperandCPPType, *OperandName, *PropertyName));
-			Lines.Add(DumpInstructions(Context, TEXT("\t\t"), BranchInfo.FirstInstruction, BranchInfo.LastInstruction, Group, false));
-			Lines.Add(RigVM_LazyEvalLambdaReturn);
-			Lines.Emplace();
-
-			OverriddenOperatorNames.Add(OperandName, Format(RigVM_LazyEvalValueName, BranchInfo.Index, BranchInfo.Label.ToString()));
-		}
-		
-		// dump the remaining instruction indices
-		Lines.Add(DumpInstructions(Context, FString(), InstructionIndices, Group, false));
-	}
-	else
-	{
-		// we have child groups - we need to invoke those
-		for(int32 ChildGroupIndex : Group.ChildGroups)
-		{
-			const FInstructionGroup& ChildGroup = InstructionGroups[ChildGroupIndex];
-
-			FString Parameters;
-			TArray<FString> ParameterArray = {RigVM_ContextPublicFormat};
-			Parameters = FString::Join(ParameterArray, RigVM_CommaSeparator);
-			Lines.Add(Format(RigVM_InvokeExecuteGroupFormat, *ChildGroup.Entry, ChildGroupIndex, *Parameters));
-		}
-	}
-	
-	Lines.Emplace();
-	if(Group.Depth <= 0)
-	{
-		Lines.Add(RigVM_ExecuteReachedExitFormat);
-	}
-	Lines.Add(RigVM_ReturnSucceededFormat);
-
-	return DumpLines(Lines, bLog);
-}
-
-FString FRigVMCodeGenerator::DumpInstructions(const FRigVMExtendedExecuteContext& Context, const FString& InPrefix, int32 InFirstInstruction, int32 InLastInstruction, const FInstructionGroup& InGroup, bool bLog)
-{
-	return DumpInstructions(Context, InPrefix, GetInstructionIndicesFromRange(InFirstInstruction, InLastInstruction), InGroup, bLog);
-}
-
-FString FRigVMCodeGenerator::DumpInstructions(const FRigVMExtendedExecuteContext& Context, const FString& InPrefix, const TArray<int32> InInstructionIndices, const FInstructionGroup& InGroup, bool bLog)
-{
-	const FRigVMByteCode& ByteCode = VM->GetByteCode();
-	const TArray<FName>& Functions = VM->GetFunctionNames();
-	const FRigVMInstructionArray Instructions = ByteCode.GetInstructions();
-
-	FString Prefix = InPrefix;
-	
-	TArray<FString> Lines;
-	for(int32 InstructionIndex : InInstructionIndices)
-	{
-		// inject a label if required
-		if (InGroup.RequiredLabels.Contains(InstructionIndex))
-		{
-			// check if the last line was a jump to this label
-			if(!Lines.IsEmpty() && Lines.Last().Contains(Format(RigVM_JumpOpFormat, InstructionIndex)))
-			{
-				Lines.Pop();
-			}
-			else
-			{
-				Lines.Add(Prefix + Format(RigVM_InstructionLabelFormat, InstructionIndex));
-			}
-		}
-
-		while(!Lines.IsEmpty() && Lines.Last().Contains(TEXT("SetInstructionIndex")))
-		{
-			Lines.Pop();
-		}
-		Lines.Add(Prefix + Format(RigVM_SetInstructionIndexFormat, InstructionIndex));
-
-		const FRigVMInstruction& Instruction = Instructions[InstructionIndex];
-		switch(Instruction.OpCode)
-		{
-			case ERigVMOpCode::Execute:
-			{
-				const FRigVMExecuteOp& Op = ByteCode.GetOpAt<FRigVMExecuteOp>(Instruction);
-				FRigVMOperandArray Operands = ByteCode.GetOperandsForCallableOp(Instruction);
-
-				const FRigVMFunction* Function = FRigVMRegistry::Get().FindFunction(*Functions[Op.CallableIndex].ToString());
-				check(Function);
-					
-				FStringArray Arguments;
-				if(Function->Struct)
-				{
-					Arguments.Add(RigVM_ContextPublicFormat);
-				}
-				else if (Function->Factory)
-				{
-					Arguments.Add(RigVM_ContextFormat);
-				}
-				for(int32 OperandIndex = 0; OperandIndex < Operands.Num(); OperandIndex++)
-				{
-					const FRigVMOperand& Operand = Operands[OperandIndex];
-					bool bSliced = false;
-					const FRigVMPropertyDescription& Property = GetPropertyDescForOperand(Context, Operand);
-					if (RigVMTypeUtils::IsArrayType(Property.CPPType))
-					{
-						const FRigVMFunctionArgument& FunctionArgument = Function->GetArguments()[OperandIndex];
-						bSliced = FunctionArgument.Type != Property.CPPType;
-					}
-					Arguments.Add(GetOperandName(Context, Operand, bSliced));
-				}
-
-				const FString JoinedArguments = FString::Join(Arguments, RigVM_CommaSeparator);
-				if(Function->Struct)
-				{
-					Lines.Add(Prefix + Format(RigVM_CallExternOpFormat, *Function->Struct->GetStructCPPName(), *Function->GetMethodName().ToString(), *JoinedArguments));
-				}
-				else if(Function->Factory)
-				{
-					FString DispatchName;
-					for(const TPair<FString, FRigVMDispatchInfo>& Pair : Dispatches)
-					{
-						if(Pair.Value.Function == Function)
-						{
-							DispatchName = Pair.Value.Name;
-							break;
-						}
-					}
-					check(!DispatchName.IsEmpty());
-					Lines.Add(Prefix + Format(RigVM_InvokeDispatchFormat, *DispatchName, *JoinedArguments));
-				}
-				else
-				{
-					checkNoEntry();
-				}
-				break;
-			}
-			case ERigVMOpCode::Zero:
-			{
-				const FRigVMUnaryOp& Op = ByteCode.GetOpAt<FRigVMUnaryOp>(Instruction);
-				const FProperty* Property = VM->GetDefaultWorkMemory().GetProperty(Op.Arg.GetRegisterIndex());
-
-				if(Property->IsA<FIntProperty>())
-				{
-					Lines.Add(Prefix + Format(RigVM_ZeroOpIntFormat, *GetOperandName(Context, Op.Arg, false)));
-				}
-				else if(Property->IsA<FNameProperty>())
-				{
-					Lines.Add(Prefix + Format(RigVM_ZeroOpNameFormat, *GetOperandName(Context, Op.Arg, false)));
-				}
-				else
-				{
-					checkNoEntry();
-				}
-				break;
-			}
-			case ERigVMOpCode::BoolFalse:
-			{
-				const FRigVMUnaryOp& Op = ByteCode.GetOpAt<FRigVMUnaryOp>(Instruction);
-				Lines.Add(Prefix + Format(RigVM_BoolFalseOpFormat, *GetOperandName(Context, Op.Arg, false)));
-				break;
-			}
-			case ERigVMOpCode::BoolTrue:
-			{
-				const FRigVMUnaryOp& Op = ByteCode.GetOpAt<FRigVMUnaryOp>(Instruction);
-				Lines.Add(Prefix + Format(RigVM_BoolTrueFormat, *GetOperandName(Context, Op.Arg, false)));
-				break;
-			}
-			case ERigVMOpCode::Copy:
-			{
-				const FRigVMBinaryOp& Op = ByteCode.GetOpAt<FRigVMBinaryOp>(Instruction);
-				const FString TargetOperand = GetOperandName(Context, Op.ArgB, false, false);
-				const FString SourceOperand = GetOperandName(Context, Op.ArgA, false, true);
-				const FString TargetCPPType = GetOperandCPPType(Context, Op.ArgB);
-				const FString SourceCPPType = GetOperandCPPType(Context, Op.ArgA);
-
-				if(RigVMTypeUtils::IsArrayType(TargetCPPType) &&
-					RigVMTypeUtils::IsArrayType(SourceCPPType) &&
-					TargetCPPType != SourceCPPType)
-				{
-					const FString TargetBaseCPPType = RigVMTypeUtils::BaseTypeFromArrayType(TargetCPPType);
-					const FString SourceBaseCPPType = RigVMTypeUtils::BaseTypeFromArrayType(SourceCPPType);
-					Lines.Add(Prefix + Format(RigVM_CopyUnrelatedArraysFormat, *TargetBaseCPPType, *SourceBaseCPPType, *TargetOperand, *SourceOperand));
-				}
-				else
-				{
-					const FString CastPrefix = TargetCPPType != SourceCPPType ? Format(RigVM_BracesFormat, *TargetCPPType) : FString();  
-					
-					if(TargetOperand.EndsWith(TEXT("(")) ||
-						TargetOperand.EndsWith(TEXT(",")) ||
-						TargetOperand.EndsWith(TEXT(", ")))
-					{
-						Lines.Add(Prefix + Format(RigVM_CopyOpMethodFormat, *TargetOperand, *CastPrefix, *SourceOperand));
-					}
-					else
-					{
-						Lines.Add(Prefix + Format(RigVM_CopyOpAssignFormat, *TargetOperand, *CastPrefix, *SourceOperand));
-					}
-				}
-				break;
-			}
-			case ERigVMOpCode::Increment:
-			{
-				const FRigVMUnaryOp& Op = ByteCode.GetOpAt<FRigVMUnaryOp>(Instruction);
-				Lines.Add(Prefix + Format(RigVM_IncrementOpFormat, *GetOperandName(Context, Op.Arg, false)));
-				break;
-			}
-			case ERigVMOpCode::Decrement:
-			{
-				const FRigVMUnaryOp& Op = ByteCode.GetOpAt<FRigVMUnaryOp>(Instruction);
-				Lines.Add(Prefix + Format(RigVM_DecrementOpFormat, *GetOperandName(Context, Op.Arg, false)));
-				break;
-			}
-			case ERigVMOpCode::Equals:
-			{
-				const FRigVMComparisonOp& Op = ByteCode.GetOpAt<FRigVMComparisonOp>(Instruction);
-				Lines.Add(Prefix + Format(RigVM_EqualsOpFormat, *GetOperandName(Context, Op.Result, false), *GetOperandName(Context, Op.B, false), *GetOperandName(Context, Op.B, false)));
-				break;
-			}
-			case ERigVMOpCode::NotEquals:
-			{
-				const FRigVMComparisonOp& Op = ByteCode.GetOpAt<FRigVMComparisonOp>(Instruction);
-				Lines.Add(Prefix + Format(RigVM_NotEqualsOpFormat, *GetOperandName(Context, Op.Result, false), *GetOperandName(Context, Op.B, false), *GetOperandName(Context, Op.B, false)));
-				break;
-			}
-			case ERigVMOpCode::JumpAbsolute:
-			{
-				const FRigVMJumpOp& Op = ByteCode.GetOpAt<FRigVMJumpOp>(Instruction);
-				Lines.Add(Prefix + Format(RigVM_JumpOpFormat, Op.InstructionIndex));
-				break;
-			}
-			case ERigVMOpCode::JumpForward:
-			{
-				const FRigVMJumpOp& Op = ByteCode.GetOpAt<FRigVMJumpOp>(Instruction);
-				Lines.Add(Prefix + Format(RigVM_JumpOpFormat, InstructionIndex + Op.InstructionIndex));
-				break;
-			}
-			case ERigVMOpCode::JumpBackward:
-			{
-				const FRigVMJumpOp& Op = ByteCode.GetOpAt<FRigVMJumpOp>(Instruction);
-				Lines.Add(Prefix + Format(RigVM_JumpOpFormat, InstructionIndex - Op.InstructionIndex));
-				break;
-			}
-			case ERigVMOpCode::JumpAbsoluteIf:
-			{
-				const FRigVMJumpIfOp& Op = ByteCode.GetOpAt<FRigVMJumpIfOp>(Instruction);
-				const FString& Condition = Op.Condition ? RigVM_TrueFormat : RigVM_FalseFormat;
-				Lines.Add(Prefix + Format(RigVM_JumpIfOpFormat, *GetOperandName(Context, Op.Arg, false), *Condition, Op.InstructionIndex));
-				break;
-			}
-			case ERigVMOpCode::JumpForwardIf:
-			{
-				const FRigVMJumpIfOp& Op = ByteCode.GetOpAt<FRigVMJumpIfOp>(Instruction);
-				const FString& Condition = Op.Condition ? RigVM_TrueFormat : RigVM_FalseFormat;
-				Lines.Add(Prefix + Format(RigVM_JumpIfOpFormat, *GetOperandName(Context, Op.Arg, false), *Condition, InstructionIndex + Op.InstructionIndex));
-				break;
-			}
-			case ERigVMOpCode::JumpBackwardIf:
-			{
-				const FRigVMJumpIfOp& Op = ByteCode.GetOpAt<FRigVMJumpIfOp>(Instruction);
-				const FString& Condition = Op.Condition ? RigVM_TrueFormat : RigVM_FalseFormat;
-				Lines.Add(Prefix + Format(RigVM_JumpIfOpFormat, *GetOperandName(Context, Op.Arg, false), *Condition, InstructionIndex - Op.InstructionIndex));
-				break;
-			}
-			case ERigVMOpCode::Exit:
-			{
-				if(InstructionIndex != InGroup.Last)
-				{
-					Lines.Add(Prefix + Format(RigVM_JumpOpFormat, Instructions.Num()));
-				}
-				break;
-			}
-			case ERigVMOpCode::BeginBlock:
-			{
-				const FRigVMBinaryOp& Op = ByteCode.GetOpAt<FRigVMBinaryOp>(Instruction);
-				Lines.Add(Prefix + Format(RigVM_BeginBlockOpFormat, *GetOperandName(Context, Op.ArgA, false), *GetOperandName(Context, Op.ArgB, false)));
-				break;
-			}
-			case ERigVMOpCode::EndBlock:
-			{
-				Lines.Add(Prefix + RigVM_EndBlockOpFormat);
-				break;
-			}
-			case ERigVMOpCode::InvokeEntry:
-			{
-				const FRigVMInvokeEntryOp& Op = ByteCode.GetOpAt<FRigVMInvokeEntryOp>(Instruction);
-				const FString EntryName = Op.EntryName.ToString();
-				Lines.Add(Prefix + Format(RigVM_InvokeEntryByNameFormat2, *EntryName, *GetEntryParameters()));
-				break;
-			}
-			case ERigVMOpCode::JumpToBranch:
-			{
-				const FRigVMJumpToBranchOp& Op = ByteCode.GetOpAt<FRigVMJumpToBranchOp>(Instruction);
-				const TArray<FRigVMBranchInfo>& Branches = ByteCode.BranchInfos;
-				for(int32 BranchIndex = Op.FirstBranchInfoIndex; BranchIndex < Branches.Num(); BranchIndex++)
-				{
-					const FRigVMBranchInfo& Branch = Branches[BranchIndex];
-					if(Branch.InstructionIndex != InstructionIndex)
-					{
-						break;
-					}
-					Lines.Add(Prefix + Format(RigVM_JumpToBranchFormat, *GetOperandName(Context, Op.Arg, false), Branches[BranchIndex].Label.ToString(), Branch.FirstInstruction));
-				}
-				break;
-			}
-			case ERigVMOpCode::RunInstructions:
-			{
-				const FRigVMRunInstructionsOp& Op = ByteCode.GetOpAt<FRigVMRunInstructionsOp>(Instruction);
-				// todo
-				break;
-			}
-			case ERigVMOpCode::Invalid:
-			case ERigVMOpCode::ChangeType:
-			default:
-			{
-				// we expect to cover all op types
-				checkNoEntry();
-			}
-		}
-	}
-
-	return DumpLines(Lines, bLog);
-}
-
-TArray<int32> FRigVMCodeGenerator::GetInstructionIndicesFromRange(int32 First, int32 Last)
-{
-	if(First > Last || First == INDEX_NONE || Last == INDEX_NONE)
-	{
-		return {};
-	}
-	if(First == Last)
-	{
-		return {First};
-	}
-	TArray<int32> Indices;
-	Indices.Reserve(Last - First + 1);
-	for(int32 Index = First; Index <= Last; Index++)
-	{
-		Indices.Add(Index);
-	}
-	return Indices;
-}
-
-
-FString FRigVMCodeGenerator::DumpHeader(const FRigVMExtendedExecuteContext& Context, bool bLog)
-{
-	const FRigVMByteCode& ByteCode = VM->GetByteCode();
-
-	FStringArray FormattedEntries;
-	for(int32 EntryIndex = 0; EntryIndex < ByteCode.NumEntries(); EntryIndex++)
-	{
-		const FRigVMByteCodeEntry& Entry = ByteCode.GetEntry(EntryIndex);
-		FormattedEntries.Add(Format(RigVM_EntryNameFormat, *Entry.GetSanitizedName()));
-	}
-
-	FStringArray Lines;
-	Lines.Add(RigVM_CopyrightFormat);
-	Lines.Emplace();
-	Lines.Add(RigVM_AutoGeneratedFormat);
-	Lines.Emplace();
-	Lines.Add(RigVM_PragmaOnceFormat);
-	Lines.Emplace();
-	Lines.Add(DumpIncludes());
-	Lines.Add(Format(RigVM_GeneratedIncludeFormat, *ClassName));
-	Lines.Emplace();
-	
-	Lines.Add(Format(RigVM_UClassDefinitionFormat, *ModuleName.ToUpper(), *ClassName));
-	Lines.Add(Format(RigVM_GetVMHashFormat, FString::Printf(TEXT("%u"), VM->GetVMHash())));
-	Lines.Add(Format(RigVM_GetEntryNamesFormat, *FString::Join(FormattedEntries, RigVM_CommaSeparator)));
-	Lines.Emplace();
-	Lines.Add(FString(RigVM_DeclareInitializeFormat));
-	Lines.Add(FString(RigVM_DeclareExecuteFormat));
-	Lines.Emplace();
-	Lines.Add(RigVM_ProtectedFormat);
-	if(!VM->GetExternalVariableDefs().IsEmpty())
-	{
-		Lines.Add(RigVM_DeclareUpdateExternalVariablesFormat);
-	}
-
-	for(int32 GroupIndex = 0; GroupIndex < InstructionGroups.Num(); GroupIndex++)
-	{
-		const FInstructionGroup& Group = InstructionGroups[GroupIndex];
-
-		FString Parameters;
-		TArray<FString> ParameterArray = {Format(RigVM_ContextPublicParameterFormat, ExecuteContextType)};
-		Parameters = FString::Join(ParameterArray, RigVM_CommaSeparator);
-		
-		if(Group.Depth == 0)
-		{
-			Lines.Add(Format(RigVM_DeclareExecuteEntryFormat, *Group.Entry, *Parameters));
-		}
-		else
-		{
-			Lines.Add(Format(RigVM_DeclareExecuteGroupFormat, *Group.Entry, GroupIndex, *Parameters));
-		}
-
-		if(GroupIndex == InstructionGroups.Num() - 1)
-		{
-			if(!Parameters.IsEmpty())
-			{
-				Parameters = RigVM_CommaSeparator + Parameters;
-			}
-			Lines.Add(Format(RigVM_DeclareInvokeEntryByNameFormat, *Parameters));
-		}
-	}
-	
-	Lines.Emplace();
-	Lines.Add(DumpEntries(true));
-	const FString BlockNames = DumpBlockNames(true);
-	if(!BlockNames.IsEmpty())
-	{
-		Lines.Add(BlockNames);
-	}
-
-	Lines.Emplace();
-	Lines.Add(DumpProperties(Context, true, INDEX_NONE));
-	Lines.Add(DumpDispatches(true));
-	if(!VM->GetExternalVariableDefs().IsEmpty())
-	{
-		Lines.Add(DumpExternalVariables(Context, true));
-	}
-	if(!RequiredUProperties.IsEmpty())
-	{
-		Lines.Add(DumpRequiredUProperties());
-	}
-	Lines.Add(TEXT("};"));
-	Lines.Emplace();
-	
-	return DumpLines(Lines, bLog);
-}
-
-FString FRigVMCodeGenerator::DumpSource(const FRigVMExtendedExecuteContext& Context, bool bLog)
-{
-	const FRigVMByteCode& ByteCode = VM->GetByteCode();
-
-	FStringArray Lines;
-	Lines.Add(RigVM_CopyrightFormat);
-	Lines.Emplace();
-	Lines.Add(RigVM_AutoGeneratedFormat);
-	Lines.Emplace();
-	Lines.Add(Format(RigVM_IncludeQuoteFormat, *ClassName));
-	Lines.Emplace();
-	Lines.Add(DumpEntries(false));
-	const FString BlockNames = DumpBlockNames(false);
-	if(!BlockNames.IsEmpty())
-	{
-		Lines.Add(BlockNames);
-	}
-	for(auto Pair : RequiredUProperties)
-	{
-		Lines.Add(Format(RigVM_UPropertyMember2Format, *ClassName, Pair.Value.Get<1>()));
-	}
-	Lines.Emplace();
-	Lines.Add(Format(RigVM_DefineInitializeFormat, *ClassName));
-	const FString InitializeContent = DumpInitialize();
-	if(!InitializeContent.IsEmpty())
-	{
-		Lines.Add(InitializeContent);
-	}
-	Lines.Add(TEXT("}"));
-	Lines.Emplace();
-	Lines.Add(Format(RigVM_DefineExecuteFormat, *ClassName));
-	Lines.Add(RigVM_StartProfilingFormat);
-	Lines.Add(Format(RigVM_SetupInstructionTrackingFormat, VM->GetByteCode().GetNumInstructions()));
-	Lines.Add(DumpInstructions(Context, INDEX_NONE));
-	Lines.Add(TEXT("}"));
-
-	if(!VM->GetExternalVariableDefs().IsEmpty())
-	{
-		Lines.Emplace();
-		Lines.Add(Format(RigVM_DefineUpdateExternalVariablesFormat, *ClassName));
-		Lines.Add(DumpExternalVariables(Context, false));
-		Lines.Add(TEXT("}"));
-	}
-
-	for(int32 GroupIndex = 0; GroupIndex < InstructionGroups.Num(); GroupIndex++)
-	{
-		const FInstructionGroup& Group = InstructionGroups[GroupIndex];
-
-		FString Parameters;
-		TArray<FString> ParameterArray = {Format(RigVM_ContextPublicParameterFormat, ExecuteContextType)};
-		Parameters = FString::Join(ParameterArray, RigVM_CommaSeparator);
-
-		Lines.Emplace();
-		if(Group.Depth == 0)
-		{
-			Lines.Add(Format(RigVM_DefineExecuteEntryFormat, *ClassName, *Group.Entry, *Parameters));
-		}
-		else
-		{
-			Lines.Add(Format(RigVM_DefineExecuteGroupFormat, *ClassName, *Group.Entry, GroupIndex, *Parameters));
-		}
-
-		const FString DumpedProperties = DumpProperties(Context, false, GroupIndex);
-		if(!DumpedProperties.IsEmpty())
-		{
-			Lines.Add(DumpedProperties);
-			Lines.Emplace();
-		}
-		Lines.Add(DumpInstructions(Context, GroupIndex));
-		Lines.Add(TEXT("}"));
-
-		if(GroupIndex == InstructionGroups.Num() - 1)
-		{
-			if(!Parameters.IsEmpty())
-			{
-				Parameters = RigVM_CommaSeparator + Parameters;
-			}
-
-			Lines.Emplace();
-			Lines.Add(Format(RigVM_DefineInvokeEntryByNameFormat, *ClassName, *Parameters));
-			Lines.Add(DumpInstructions(Context, -2));
-			Lines.Add(TEXT("}"));
-		}
-	}
-	
-	return DumpLines(Lines, bLog);
-}
-
-FString FRigVMCodeGenerator::DumpLines(const TArray<FString>& InLines, bool bLog)
-{
-	if (bLog)
-	{
-		for(const FString& Line : InLines)
-		{
-			UE_LOG(LogRigVMDeveloper, Display, RigVM_SingleStringFormat, *Line);
-		}
-	}
-	return FString::Join(InLines, RigVM_NewLineFormat);
-}
-
-void FRigVMCodeGenerator::Reset()
-{
-	ClassName.Reset();
-	ModuleName.Reset();
-	Libraries.Reset();
-	Includes.Reset();
-	Dispatches.Reset();
-	RequiredUProperties.Reset();
-	InstructionGroups.Reset();
-	MappedCPPTypes.Reset();
-	Properties.Reset();
-	PropertyNameToIndex.Reset();
-	OverriddenOperatorNames.Reset();
-}
-
-void FRigVMCodeGenerator::ParseVM(const FString& InClassName, const FString& InModuleName,
-	URigVMGraph* InModelToNativize, URigVM* InVMToNativize, FRigVMExtendedExecuteContext& InVMContext,
-	const UScriptStruct* PublicContextStruct, TMap<FString,FRigVMOperand> InPinToOperandMap,
-	int32 InMaxInstructionsPerFunction)
-{
-	check(InVMToNativize);
-
-	Reset();
-
-	Model = TStrongObjectPtr<URigVMGraph>(InModelToNativize);
-	VM = TStrongObjectPtr<URigVM>(InVMToNativize);
-	PinToOperandMap = InPinToOperandMap;
-	MaxInstructionsPerFunction = InMaxInstructionsPerFunction;
-	ClassName = InClassName;
-	ModuleName = InModuleName;
-	ExecuteContextType = PublicContextStruct->GetStructCPPName();
-
-	// create an inverted map to lookup pins from operands
-	OperandToPinMap.Reset();
-	for(const TPair<FString,FRigVMOperand>& Pair : PinToOperandMap)
-	{
-		OperandToPinMap.FindOrAdd(Pair.Value) = Pair.Key;
-	}
-
-	// default includes
-	ParseInclude(URigVM::StaticClass());
-
-	Properties.Reserve(
-		InVMToNativize->GetDefaultLiteralMemory().Num() +
-		InVMToNativize->GetDefaultWorkMemory().Num());
-	ParseMemory(InVMContext, &InVMToNativize->GetDefaultLiteralMemory());
-	ParseMemory(InVMContext, &InVMToNativize->GetDefaultWorkMemory());
-
-	ParseRequiredUProperties(InVMContext);
-	ParseInstructionGroups(InVMContext);
-}
-
-void FRigVMCodeGenerator::ParseInclude(UStruct* InDependency, const FName& InMethodName)
-{
-	if (InDependency->IsChildOf(FRigVMStruct::StaticStruct()))
-	{
-		if (const FRigVMFunction* Function = FRigVMRegistry::Get().FindFunction(Cast<UScriptStruct>(InDependency), *InMethodName.ToString()))
-		{
-			FString FunctionModuleName = Function->GetModuleName();
-			if (FunctionModuleName.Contains(TEXT("/")))
-			{
-				FunctionModuleName.Split(TEXT("/"), nullptr, &FunctionModuleName, ESearchCase::IgnoreCase, ESearchDir::FromEnd);
-			}
-			Libraries.AddUnique(FunctionModuleName);
-
-			static const FString PrivatePrefix = TEXT("Private/");
-			static const FString PublicPrefix = TEXT("Public/");
-
-			const FString RelativeHeaderPath = Function->GetModuleRelativeHeaderPath();
-			if(RelativeHeaderPath.StartsWith(PrivatePrefix))
-			{
-				Includes.Add(RelativeHeaderPath.Mid(PrivatePrefix.Len()));
-			}
-			else if(RelativeHeaderPath.StartsWith(PublicPrefix))
-			{
-				Includes.Add(RelativeHeaderPath.Mid(PublicPrefix.Len()));
-			}
-			else
-			{
-				Includes.AddUnique(Format(RigVM_JoinFilePathFormat, *FunctionModuleName, *RelativeHeaderPath));
-			}
-			return;
-		}
-	}
-
-	if (InDependency == URigVM::StaticClass() ||
-		InDependency == FRigVMByteCode::StaticStruct() || 
-		InDependency->IsChildOf(FRigVMExecutePin::StaticStruct()) ||
-		InDependency->IsChildOf(FRigVMBaseOp::StaticStruct()) ||
-		InDependency->IsChildOf(URigVMMemoryStorage::StaticClass()))
-	{
-		Libraries.AddUnique(RigVM_RigVMCoreLibraryFormat);
-		Includes.AddUnique(RigVM_RigVMCoreIncludeFormat);
-		Includes.AddUnique(RigVM_RigVMModuleIncludeFormat);
-	}
-}
-
-void FRigVMCodeGenerator::ParseRequiredUProperties(const FRigVMExtendedExecuteContext& Context)
-{
-	const TArray<FName>& Functions = VM->GetFunctionNames();
-	const FRigVMByteCode& ByteCode = VM->GetByteCode();
-	const FRigVMInstructionArray Instructions = ByteCode.GetInstructions();
-
-	auto AddRequiredUProperty = [this](const FString& InTypeName)
-	{
-		const TRigVMTypeIndex TypeIndex = FRigVMRegistry::Get().GetTypeIndexFromCPPType(InTypeName);
-		check(TypeIndex != INDEX_NONE);
-		if(!RequiredUProperties.Contains(TypeIndex))
-		{
-			const FRigVMTemplateArgumentType& Type = FRigVMRegistry::Get().GetType(TypeIndex);
-			FString TypeLabel = Type.GetBaseCPPType();
-			if(Type.IsArray())
-			{
-				TypeLabel += TEXT("_Array");
-			}
-			TypeLabel = TEXT("Property_") + TypeLabel; 
-			RequiredUProperties.Add(TypeIndex, {InTypeName, TypeLabel});
-		}
-	};
-
-	for(int32 InstructionIndex = 0; InstructionIndex < Instructions.Num(); InstructionIndex++)
-	{
-		const FRigVMInstruction& Instruction = Instructions[InstructionIndex];
-		
-		if(Instruction.OpCode == ERigVMOpCode::Execute)
-		{
-			const FRigVMExecuteOp& Op = ByteCode.GetOpAt<FRigVMExecuteOp>(Instruction);
-			const FRigVMFunction* Function = FRigVMRegistry::Get().FindFunction(*Functions[Op.CallableIndex].ToString());
-			if(const FRigVMDispatchFactory* Factory = Function->Factory)
-			{
-				const FString FunctionName = Function->GetName();
-				const int32 PermutationIndex = Factory->GetTemplate()->FindPermutation(Function); 
-				check(PermutationIndex != INDEX_NONE);
-
-				for(const FRigVMFunctionArgument& Argument : Function->Arguments)
-				{
-					AddRequiredUProperty(Argument.Type);
-				}
-				
-				const FString DispatchKey = Format(RigVM_DispatchKeyFormat, Factory->GetFactoryName().ToString(), PermutationIndex);
-				if (Dispatches.Contains(DispatchKey))
-				{
-					continue;
-				}
-
-				FRigVMDispatchContext DispatchContext;
-				if(const URigVMDispatchNode* DispatchNode = Cast<URigVMDispatchNode>(ByteCode.GetSubjectForInstruction(InstructionIndex)))
-				{
-					DispatchContext = DispatchNode->GetDispatchContext();
-				}
-
-				Dispatches.Add(DispatchKey, {DispatchKey, Function, DispatchContext});
-			}
-		}
-	}
-
-	// also parse all lazy branches
-	for(const FRigVMBranchInfo& BranchInfo : ByteCode.BranchInfos)
-	{
-		// skip output block branches
-		if(BranchInfo.IsOutputBranch())
-		{
-			continue;
-		}
-
-		const FRigVMOperandArray Operands = ByteCode.GetOperandsForOp(Instructions[BranchInfo.InstructionIndex]);
-		check(Operands.IsValidIndex(BranchInfo.ArgumentIndex));
-		const FRigVMOperand Operand = Operands[BranchInfo.ArgumentIndex];
-		const FString CPPType = GetOperandCPPType(Context, Operand);
-		AddRequiredUProperty(CPPType);
-	}
-}
-
-void FRigVMCodeGenerator::ParseMemory(const FRigVMExtendedExecuteContext& Context, FRigVMMemoryStorageStruct* InMemory)
-{
-	if (InMemory == nullptr)
-	{
-		return;
-	}
-	for (const FProperty* Property : InMemory->GetProperties())
-	{
-		ParseProperty(Context, InMemory->GetMemoryType(), Property, InMemory);
-	}
-}
-
-void FRigVMCodeGenerator::ParseProperty(const FRigVMExtendedExecuteContext& Context, ERigVMMemoryType InMemoryType, const FProperty* InProperty, FRigVMMemoryStorageStruct* InMemory)
-{
-	const int32 PropertyIndex = InMemory->GetPropertyIndex(InProperty);
-	const FRigVMOperand Operand(InMemoryType, PropertyIndex);
-	const FRigVMPropertyDescription& PropertyDescription = GetPropertyForOperand(Context, Operand);
-	
-	FPropertyInfo Info;
-	Info.MemoryPropertyIndex = PropertyIndex;
-	Info.Description = PropertyDescription;
-
-	check(InMemoryType == ERigVMMemoryType::Literal || InMemoryType == ERigVMMemoryType::Work);
-	if(InMemoryType == ERigVMMemoryType::Literal)
-	{
-		Info.PropertyType = ERigVMNativizedPropertyType::Literal;
-	}
-	else
-	{
-		Info.PropertyType = ERigVMNativizedPropertyType::Work;
-
-		// work state may be a hidden / sliced property.
-		if(RigVMTypeUtils::IsArrayType(PropertyDescription.CPPType))
-		{
-			if(const FString* PinPath = OperandToPinMap.Find(Operand))
-			{
-				if(const URigVMPin* Pin = Model->FindPin(*PinPath))
-				{
-					if(Pin->GetDirection() == ERigVMPinDirection::Hidden)
-					{
-						Info.PropertyType = ERigVMNativizedPropertyType::Sliced;
-					}
-				}
-			}
-		}
-	}
-
-	const int32 LookupIndex = Properties.Add(Info);
-	PropertyNameToIndex.Add(Info.Description.Name, LookupIndex);
-}
-
-void FRigVMCodeGenerator::ParseInstructionGroups(const FRigVMExtendedExecuteContext& Context)
-{
-	const TArray<FName>& Functions = VM->GetFunctionNames();
-	const FRigVMByteCode& ByteCode = VM->GetByteCode();
-	const FRigVMInstructionArray Instructions = ByteCode.GetInstructions();
-
-	for(int32 EntryIndex = 0; EntryIndex < ByteCode.NumEntries(); EntryIndex++)
-	{
-		const FRigVMByteCodeEntry& Entry = ByteCode.GetEntry(EntryIndex);
-
-		FInstructionGroup Group;
-		Group.Entry = Entry.GetSanitizedName();
-		Group.Depth = 0;
-		Group.First = Entry.InstructionIndex;
-		Group.Last = Instructions.Num() - 1;
-
-		for(int32 InstructionIndex = Group.First+1; InstructionIndex < Instructions.Num(); InstructionIndex++)
-		{
-			if(Instructions[InstructionIndex].OpCode == ERigVMOpCode::Exit)
-			{
-				Group.Last = InstructionIndex;
-				break;
-			}
-		}
-
-		InstructionGroups.Add(Group);
-	}
-
-	for(int32 GroupIndex = 0; GroupIndex < InstructionGroups.Num(); GroupIndex++)
-	{
-		// copy the group here since it will be invalid soon
-		FInstructionGroup Group = InstructionGroups[GroupIndex];
-		if(Group.Last - Group.First + 1 > MaxInstructionsPerFunction)
-		{
-			// find all of the constraints / jumps
-			TArray<TTuple<int32,int32>> Constraints;
-			for(int32 InstructionIndex = Group.First; InstructionIndex <= Group.Last; InstructionIndex++)
-			{
-				auto AddConstraint = [&Constraints, InstructionIndex](int32 OtherInstructionIndex)
-				{
-					Constraints.Add(TTuple<int32,int32>(
-						FMath::Min(InstructionIndex, OtherInstructionIndex),
-						FMath::Max(InstructionIndex, OtherInstructionIndex)));
-				};
-				
-				const FRigVMInstruction& Instruction = Instructions[InstructionIndex];
-				switch(Instruction.OpCode)
-				{
-					case ERigVMOpCode::JumpAbsolute:
-					{
-						const FRigVMJumpOp& Op = ByteCode.GetOpAt<FRigVMJumpOp>(Instruction);
-						AddConstraint(Op.InstructionIndex);
-						break;
-					}
-					case ERigVMOpCode::JumpForward:
-					{
-						const FRigVMJumpOp& Op = ByteCode.GetOpAt<FRigVMJumpOp>(Instruction);
-						AddConstraint(InstructionIndex + Op.InstructionIndex);
-						break;
-					}
-					case ERigVMOpCode::JumpBackward:
-					{
-						const FRigVMJumpOp& Op = ByteCode.GetOpAt<FRigVMJumpOp>(Instruction);
-						AddConstraint(InstructionIndex - Op.InstructionIndex);
-						break;
-					}
-					case ERigVMOpCode::JumpAbsoluteIf:
-					{
-						const FRigVMJumpIfOp& Op = ByteCode.GetOpAt<FRigVMJumpIfOp>(Instruction);
-						AddConstraint(InstructionIndex);
-						break;
-					}
-					case ERigVMOpCode::JumpForwardIf:
-					{
-						const FRigVMJumpIfOp& Op = ByteCode.GetOpAt<FRigVMJumpIfOp>(Instruction);
-						AddConstraint(InstructionIndex + Op.InstructionIndex);
-						break;
-					}
-					case ERigVMOpCode::JumpBackwardIf:
-					{
-						const FRigVMJumpIfOp& Op = ByteCode.GetOpAt<FRigVMJumpIfOp>(Instruction);
-						AddConstraint(InstructionIndex - Op.InstructionIndex);
-						break;
-					}
-					case ERigVMOpCode::JumpToBranch:
-					{
-						const FRigVMJumpToBranchOp& Op = ByteCode.GetOpAt<FRigVMJumpToBranchOp>(Instruction);
-						for(int32 BranchIndex = Op.FirstBranchInfoIndex; BranchIndex < ByteCode.BranchInfos.Num(); BranchIndex++)
-						{
-							const FRigVMBranchInfo& Branch = ByteCode.BranchInfos[BranchIndex];
-							if(Branch.InstructionIndex != InstructionIndex)
-							{
-								break;
-							}
-							AddConstraint(Branch.FirstInstruction);
-						}
-						break;
-					}
-					case ERigVMOpCode::RunInstructions:
-					{
-						const FRigVMRunInstructionsOp& Op = ByteCode.GetOpAt<FRigVMRunInstructionsOp>(Instruction);
-						// todo
-						break;
-					}
-					default:
-					{
-						break;
-					}
-				}
-			}
-
-			const int32 First = Group.First;
-			const int32 Last = Group.Last;
-
-			// also constrain the group split by the branches used
-			for(const FRigVMBranchInfo& Branch : ByteCode.BranchInfos)
-			{
-				if(Branch.FirstInstruction >= First && Branch.LastInstruction <= Last)
-				{
-					Constraints.Add(TTuple<int32,int32>(
-						FMath::Min(Branch.FirstInstruction, Branch.LastInstruction),
-						FMath::Max(Branch.FirstInstruction, Branch.LastInstruction)));
-				}
-			}
-			
-			int32 Middle = (Group.Last + Group.First) / 2;
-			int32 MiddleHead = Middle;
-			int32 MiddleTail = Middle;
-			bool bMoveTowardsHead = true;
-
-			auto IsWithinConstraint = [&bMoveTowardsHead, &MiddleHead, &MiddleTail](const TTuple<int32,int32>& Constraint) -> bool
-			{
-				int32& Middle = bMoveTowardsHead ? MiddleHead : MiddleTail;
-				return FMath::IsWithinInclusive(Middle, Constraint.Get<0>(), Constraint.Get<1>());
-			};
-
-			// while the middle instruction is within a constraint - move it  
-			while(Constraints.FindByPredicate(IsWithinConstraint))
-			{
-				MiddleHead--;
-			}
-			bMoveTowardsHead = false;
-			while(Constraints.FindByPredicate(IsWithinConstraint))
-			{
-				MiddleTail++;
-			}
-
-			Middle = ((MiddleHead - First) > (Last - MiddleTail)) ? MiddleHead : MiddleTail;
-			if(Middle == First || Middle == Last)
-			{
-				continue;
-			}
-
-			FInstructionGroup HeadGroup = Group;
-			HeadGroup.ParentGroup = GroupIndex;
-			HeadGroup.ChildGroups.Reset();
-			HeadGroup.First = First;
-			HeadGroup.Last = Middle;
-			HeadGroup.Depth++;
-
-			FInstructionGroup TailGroup = Group;
-			TailGroup.ParentGroup = GroupIndex;
-			TailGroup.ChildGroups.Reset();
-			TailGroup.First = Middle+1;
-			TailGroup.Last = Last;
-			TailGroup.Depth++;
-			
-			const int32 HeadGroupIndex = InstructionGroups.Add(HeadGroup);
-			const int32 TailGroupIndex = InstructionGroups.Add(TailGroup);
-			InstructionGroups[GroupIndex].ChildGroups.Add(HeadGroupIndex);
-			InstructionGroups[GroupIndex].ChildGroups.Add(TailGroupIndex);
-		}
-	}
-
-	for(int32 GroupIndex = 0; GroupIndex < InstructionGroups.Num(); GroupIndex++)
-	{
-		FInstructionGroup& Group = InstructionGroups[GroupIndex];
-		
-		// determine all the necessary labels and optional arguments
-		for(int32 InstructionIndex = Group.First; InstructionIndex <= Group.Last; InstructionIndex++)
-		{
-			const FRigVMInstruction& Instruction = Instructions[InstructionIndex];
-			switch(Instruction.OpCode)
-			{
-				case ERigVMOpCode::Execute:
-				{
-					const FRigVMExecuteOp& Op = ByteCode.GetOpAt<FRigVMExecuteOp>(Instruction);
-					const FRigVMFunction* Function = FRigVMRegistry::Get().FindFunction(*Functions[Op.CallableIndex].ToString());
-					check(Function);
-
-					// make sure to include the required header
-					if(Function->Struct)
-					{
-						ParseInclude(Function->Struct, Function->GetMethodName());
-					}
-					if(Function->Factory)
-					{
-						ParseInclude(Function->Factory->GetScriptStruct(), Function->GetMethodName());
-					}
-
-					break;
-				}
-				case ERigVMOpCode::JumpAbsolute:
-				{
-					const FRigVMJumpOp& Op = ByteCode.GetOpAt<FRigVMJumpOp>(Instruction);
-					Group.RequiredLabels.AddUnique(Op.InstructionIndex);
-					break;
-				}
-				case ERigVMOpCode::JumpForward:
-				{
-					const FRigVMJumpOp& Op = ByteCode.GetOpAt<FRigVMJumpOp>(Instruction);
-					Group.RequiredLabels.AddUnique(InstructionIndex + Op.InstructionIndex);
-					break;
-				}
-				case ERigVMOpCode::JumpBackward:
-				{
-					const FRigVMJumpOp& Op = ByteCode.GetOpAt<FRigVMJumpOp>(Instruction);
-					Group.RequiredLabels.AddUnique(InstructionIndex - Op.InstructionIndex);
-					break;
-				}
-				case ERigVMOpCode::JumpAbsoluteIf:
-				{
-					const FRigVMJumpIfOp& Op = ByteCode.GetOpAt<FRigVMJumpIfOp>(Instruction);
-					Group.RequiredLabels.AddUnique(Op.InstructionIndex);
-					break;
-				}
-				case ERigVMOpCode::JumpForwardIf:
-				{
-					const FRigVMJumpIfOp& Op = ByteCode.GetOpAt<FRigVMJumpIfOp>(Instruction);
-					Group.RequiredLabels.AddUnique(InstructionIndex + Op.InstructionIndex);
-					break;
-				}
-				case ERigVMOpCode::JumpBackwardIf:
-				{
-					const FRigVMJumpIfOp& Op = ByteCode.GetOpAt<FRigVMJumpIfOp>(Instruction);
-					Group.RequiredLabels.AddUnique(InstructionIndex - Op.InstructionIndex);
-					break;
-				}
-				case ERigVMOpCode::JumpToBranch:
-				{
-					const FRigVMJumpToBranchOp& Op = ByteCode.GetOpAt<FRigVMJumpToBranchOp>(Instruction);
-					const TArray<FRigVMBranchInfo>& Branches = ByteCode.BranchInfos;
-					for(int32 BranchIndex = Op.FirstBranchInfoIndex; BranchIndex < Branches.Num(); BranchIndex++)
-					{
-						// only add labels for output blocks - so blocks that belong
-						// to a jump to branch instruction
-						if(Branches[BranchIndex].IsOutputBranch())
-						{
-							Group.RequiredLabels.AddUnique(Branches[BranchIndex].FirstInstruction);
-						}
-					}
-					break;
-				}
-				case ERigVMOpCode::RunInstructions:
-				{
-					const FRigVMRunInstructionsOp& Op = ByteCode.GetOpAt<FRigVMRunInstructionsOp>(Instruction);
-					// todo
-					break;
-				}
-				case ERigVMOpCode::SetupTraits:
-				{
-					const FRigVMSetupTraitsOp& Op = ByteCode.GetOpAt<FRigVMSetupTraitsOp>(Instruction);
-					// todo
-					break;
-				}
-				default:
-				{
-					break;
-				}
-			}
-
-			// update the usage table to know which property is used where
-			FRigVMOperandArray Operands = ByteCode.GetOperandsForOp(Instruction);
-			for(const FRigVMOperand& Operand : Operands)
-			{
-				if(Operand.GetMemoryType() == ERigVMMemoryType::Literal ||
-					Operand.GetMemoryType() == ERigVMMemoryType::Work)
-				{
-					const int32 PropertyIndex = GetPropertyIndex(Context, Operand);
-					Properties[PropertyIndex].Groups.AddUnique(GroupIndex);
-				}
-			}
-		}
-	}
-
-	// now that we know the groups that properties belong to, we need to filter them.
-	// we only want to declare properties once - and pass them as parameters to the enclosed
-	// groups. parent groups which only pass them to one sub group can be removed from the
-	// property's group table. the top level group which still contains the property is the
-	// one to declare it.
-	const TArray<FInstructionGroup>& Groups = InstructionGroups;
-	for(int32 PropertyIndex = 0; PropertyIndex < Properties.Num(); PropertyIndex++)
-	{
-		// work / sliced properties are always defined at the top
-		FPropertyInfo& Property = Properties[PropertyIndex];
-
-		if(Property.PropertyType == ERigVMNativizedPropertyType::Sliced)
-		{
-			Property.Groups.Reset();
-			continue;
-		}
-		
-		int32 GroupCount;
-		do
-		{
-			GroupCount = Property.Groups.Num();
-			TArray<int32> FilteredGroups = Property.Groups;
-			FilteredGroups.RemoveAll([Groups, Property](int32 Group) -> bool
-			{
-				return !Groups[Group].ChildGroups.IsEmpty();
-			});
-			Property.Groups = FilteredGroups;
-			
-		} while(GroupCount != Property.Groups.Num());
-	}
-}
-
-FString FRigVMCodeGenerator::GetOperandName(const FRigVMExtendedExecuteContext& Context, const FRigVMOperand& InOperand, bool bPerSlice, bool bAsInput) const
-{
-	const FRigVMPropertyDescription& Property = GetPropertyDescForOperand(Context, InOperand);
-	const FRigVMPropertyPathDescription& PropertyPath = GetPropertyPathForOperand(Context, InOperand);
-
-	FString OperandName;
-	if (Property.IsValid())
-	{
-		OperandName = Property.Name.ToString();
-	}
-
-	if (InOperand.GetMemoryType() != ERigVMMemoryType::External)
-	{
-		OperandName = SanitizeName(OperandName, Property.CPPType);
-	}
-
-	if(const FString* OverrideName = OverriddenOperatorNames.Find(OperandName))
-	{
-		return *OverrideName;
-	}
-
-	// if we are an array on work memory this indicates that we'll be sliced.
-	if (bPerSlice && InOperand.GetMemoryType() == ERigVMMemoryType::Work && RigVMTypeUtils::IsArrayType(Property.CPPType))
-	{
-		const FString MappedType = GetMappedType(Property.CPPType);
-		const FString MappedTypeSuffix = GetMappedTypeSuffix(Property.CPPType);
-		const FString BaseCPPType = RigVMTypeUtils::IsArrayType(MappedType) ? RigVMTypeUtils::BaseTypeFromArrayType(MappedType) : MappedType;
-		OperandName = Format(RigVM_GetOperandSliceFormat, *BaseCPPType, *OperandName, *MappedTypeSuffix);
-	}
-
-	if (InOperand.GetMemoryType() == ERigVMMemoryType::External)
-	{
-		OperandName = Format(RigVM_ExternalVariableFormat, *OperandName);
-	}
-
-	if (PropertyPath.IsValid())
-	{
-		check(Property.IsValid());
-
-		FString RemainingSegmentPath = PropertyPath.SegmentPath;
-		FString SegmentPath = OperandName;
-
-		static TMap<UScriptStruct*, TMap<FString, TTuple<FString, FString>>> MappedProperties;
-		if (MappedProperties.IsEmpty())
-		{
-			TMap<FString, TTuple<FString, FString>>& TransformMap = MappedProperties.Add(TBaseStructure<FTransform>::Get());
-			TransformMap.Add(TEXT("Translation"), TTuple<FString, FString>(TEXT("{0}.GetTranslation()"), TEXT("{0}.SetTranslation(")));
-			TransformMap.Add(TEXT("Translation.X"), TTuple<FString, FString>(TEXT("{0}.GetTranslation().X"), TEXT("FTransformSetter({0}).SetTranslationX(")));
-			TransformMap.Add(TEXT("Translation.Y"), TTuple<FString, FString>(TEXT("{0}.GetTranslation().Y"), TEXT("FTransformSetter({0}).SetTranslationY(")));
-			TransformMap.Add(TEXT("Translation.Z"), TTuple<FString, FString>(TEXT("{0}.GetTranslation().Z"), TEXT("FTransformSetter({0}).SetTranslationZ(")));
-			TransformMap.Add(TEXT("Rotation"), TTuple<FString, FString>(TEXT("{0}.GetRotation()"), TEXT("{0}.SetRotation(")));
-			TransformMap.Add(TEXT("Rotation.X"), TTuple<FString, FString>(TEXT("{0}.GetRotation().X"), TEXT("FTransformSetter({0}).SetRotationX(")));
-			TransformMap.Add(TEXT("Rotation.Y"), TTuple<FString, FString>(TEXT("{0}.GetRotation().Y"), TEXT("FTransformSetter({0}).SetRotationY(")));
-			TransformMap.Add(TEXT("Rotation.Z"), TTuple<FString, FString>(TEXT("{0}.GetRotation().Z"), TEXT("FTransformSetter({0}).SetRotationZ(")));
-			TransformMap.Add(TEXT("Rotation.W"), TTuple<FString, FString>(TEXT("{0}.GetRotation().W"), TEXT("FTransformSetter({0}).SetRotationW(")));
-			TransformMap.Add(TEXT("Scale3D"), TTuple<FString, FString>(TEXT("{0}.GetScale3D()"), TEXT("{0}.SetScale3D(")));
-			TransformMap.Add(TEXT("Scale3D.X"), TTuple<FString, FString>(TEXT("{0}.GetScale3D().X"), TEXT("FTransformSetter({0}).SetScaleX(")));
-			TransformMap.Add(TEXT("Scale3D.Y"), TTuple<FString, FString>(TEXT("{0}.GetScale3D().Y"), TEXT("FTransformSetter({0}).SetScaleY(")));
-			TransformMap.Add(TEXT("Scale3D.Z"), TTuple<FString, FString>(TEXT("{0}.GetScale3D().Z"), TEXT("FTransformSetter({0}).SetScaleZ(")));
-
-			TMap<FString, TTuple<FString, FString>>& MatrixMap = MappedProperties.Add(TBaseStructure<FMatrix>::Get());
-			MatrixMap.Add(TEXT("XPlane"), TTuple<FString, FString>(TEXT("(*(FPlane*){0}.M[0])"), TEXT("FMatrixSetter({0}).SetPlane(0, ")));
-			MatrixMap.Add(TEXT("XPlane.X"), TTuple<FString, FString>(TEXT("{0}.M[0][0]"), TEXT("FMatrixSetter({0}).SetComponent(0, 0, ")));
-			MatrixMap.Add(TEXT("XPlane.Y"), TTuple<FString, FString>(TEXT("{0}.M[0][1]"), TEXT("FMatrixSetter({0}).SetComponent(0, 1, ")));
-			MatrixMap.Add(TEXT("XPlane.Z"), TTuple<FString, FString>(TEXT("{0}.M[0][2]"), TEXT("FMatrixSetter({0}).SetComponent(0, 2, ")));
-			MatrixMap.Add(TEXT("XPlane.W"), TTuple<FString, FString>(TEXT("{0}.M[0][3]"), TEXT("FMatrixSetter({0}).SetComponent(0, 3, ")));
-			MatrixMap.Add(TEXT("YPlane"), TTuple<FString, FString>(TEXT("(*(FPlane*){0}.M[1])"), TEXT("FMatrixSetter({0}).SetPlane(1, ")));
-			MatrixMap.Add(TEXT("YPlane.X"), TTuple<FString, FString>(TEXT("{0}.M[1][0]"), TEXT("FMatrixSetter({0}).SetComponent(1, 0, ")));
-			MatrixMap.Add(TEXT("YPlane.Y"), TTuple<FString, FString>(TEXT("{0}.M[1][1]"), TEXT("FMatrixSetter({0}).SetComponent(1, 1, ")));
-			MatrixMap.Add(TEXT("YPlane.Z"), TTuple<FString, FString>(TEXT("{0}.M[1][2]"), TEXT("FMatrixSetter({0}).SetComponent(1, 2, ")));
-			MatrixMap.Add(TEXT("YPlane.W"), TTuple<FString, FString>(TEXT("{0}.M[1][3]"), TEXT("FMatrixSetter({0}).SetComponent(1, 3, ")));
-			MatrixMap.Add(TEXT("ZPlane"), TTuple<FString, FString>(TEXT("(*(FPlane*){0}.M[2])"), TEXT("FMatrixSetter({0}).SetPlane(2, ")));
-			MatrixMap.Add(TEXT("ZPlane.X"), TTuple<FString, FString>(TEXT("{0}.M[2][0]"), TEXT("FMatrixSetter({0}).SetComponent(2, 0, ")));
-			MatrixMap.Add(TEXT("ZPlane.Y"), TTuple<FString, FString>(TEXT("{0}.M[2][1]"), TEXT("FMatrixSetter({0}).SetComponent(2, 1, ")));
-			MatrixMap.Add(TEXT("ZPlane.Z"), TTuple<FString, FString>(TEXT("{0}.M[2][2]"), TEXT("FMatrixSetter({0}).SetComponent(2, 2, ")));
-			MatrixMap.Add(TEXT("ZPlane.W"), TTuple<FString, FString>(TEXT("{0}.M[2][3]"), TEXT("FMatrixSetter({0}).SetComponent(2, 3, ")));
-			MatrixMap.Add(TEXT("WPlane"), TTuple<FString, FString>(TEXT("(*(FPlane*){0}.M[3])"), TEXT("FMatrixSetter({0}).SetPlane(3, ")));
-			MatrixMap.Add(TEXT("WPlane.X"), TTuple<FString, FString>(TEXT("{0}.M[3][0]"), TEXT("FMatrixSetter({0}).SetComponent(3, 0, ")));
-			MatrixMap.Add(TEXT("WPlane.Y"), TTuple<FString, FString>(TEXT("{0}.M[3][1]"), TEXT("FMatrixSetter({0}).SetComponent(3, 1, ")));
-			MatrixMap.Add(TEXT("WPlane.Z"), TTuple<FString, FString>(TEXT("{0}.M[3][2]"), TEXT("FMatrixSetter({0}).SetComponent(3, 2, ")));
-			MatrixMap.Add(TEXT("WPlane.W"), TTuple<FString, FString>(TEXT("{0}.M[3][3]"), TEXT("FMatrixSetter({0}).SetComponent(3, 3, ")));
-		}
-
-		const FProperty* CurrentProperty = Property.Property;
-
-		while(!RemainingSegmentPath.IsEmpty() && (CurrentProperty != nullptr))
-		{
-			FString Left, Right;
-			if(!RigVMStringUtils::SplitPinPathAtStart(RemainingSegmentPath, Left, Right))
-			{
-				Left = RemainingSegmentPath;
-				Right.Reset();
-			}
-			
-			if (const FStructProperty* StructProperty = CastField<FStructProperty>(CurrentProperty))
-			{
-				CurrentProperty = StructProperty->Struct->FindPropertyByName(*Left);
-				
-				if (const TMap<FString, TTuple<FString, FString>>* PropertyMap = MappedProperties.Find(StructProperty->Struct))
-				{
-					if (const TTuple<FString, FString>* NewSegmentName = PropertyMap->Find(RemainingSegmentPath))
-					{
-						FStringFormatOrderedArguments FormatArguments;
-						FormatArguments.Add(SegmentPath);
-
-						SegmentPath = FString::Format(
-							bAsInput ? *NewSegmentName->Get<0>() : *NewSegmentName->Get<1>(),
-							FormatArguments);
-						break;
-					}
-				}
-			}
-			else if(const FArrayProperty* ArrayProperty = CastField<FArrayProperty>(CurrentProperty))
-			{
-				CurrentProperty = ArrayProperty->Inner;
-
-				Left = Left.TrimChar(TEXT('['));
-				Left = Left.TrimChar(TEXT(']'));
-
-				FString CPPType = RigVMTypeUtils::GetCPPTypeFromProperty(CurrentProperty);
-				if(CPPType == RigVMTypeUtils::UInt8Type && CastField<FBoolProperty>(CurrentProperty))
-				{
-					CPPType = RigVMTypeUtils::BoolType;
-				}
-
-				SegmentPath = Format(RigVM_GetArrayElementSafeFormat, *CPPType, *SegmentPath, *Left);
-				RemainingSegmentPath = Right;
-				continue;
-			}
-			else
-			{
-				CurrentProperty = nullptr;
-			}
-
-			SegmentPath = Format(RigVM_JoinSegmentPathFormat, *SegmentPath, *Left);
-			RemainingSegmentPath = Right;
-		}
-
-		return SegmentPath;
-	}
-
-	return OperandName;
-}
-
-FString FRigVMCodeGenerator::GetOperandCPPType(const FRigVMExtendedExecuteContext& Context, const FRigVMOperand& InOperand) const
-{
-	const FRigVMPropertyDescription& Property = GetPropertyDescForOperand(Context, InOperand);
-	const FRigVMPropertyPathDescription& PropertyPath = GetPropertyPathForOperand(Context, InOperand);
-
-	if (PropertyPath.IsValid())
-	{
-		check(Property.IsValid());
-
-		const FRigVMPropertyPath Path(Property.Property, PropertyPath.SegmentPath);
-		const FProperty* TailProperty = Path.GetTailProperty();
-		check(TailProperty);
-
-		return RigVMTypeUtils::GetCPPTypeFromProperty(TailProperty);
-	}
-	
-	if (Property.IsValid())
-	{
-		return Property.CPPType;
-	}
-
-	return FString(); 
-}
-
-FString FRigVMCodeGenerator::GetOperandCPPBaseType(const FRigVMExtendedExecuteContext& Context, const FRigVMOperand& InOperand) const
-{
-	FString CPPType = GetOperandCPPType(Context, InOperand);
-	while(RigVMTypeUtils::IsArrayType(CPPType))
-	{
-		CPPType = RigVMTypeUtils::BaseTypeFromArrayType(CPPType);
-	}
-	return CPPType;
-}
-
-
-FString FRigVMCodeGenerator::SanitizeName(const FString& InName, const FString& CPPType)
-{
-	FString Name = InName;
-	while(Name.Contains(RigVM_DoubleUnderscoreFormat))
-	{
-		Name.ReplaceInline(RigVM_DoubleUnderscoreFormat, RigVM_SingleUnderscoreFormat);
-	}
-
-	if(CPPType == RigVMTypeUtils::BoolType && !Name.StartsWith(RigVM_BoolPropertyPrefix, ESearchCase::CaseSensitive))
-	{
-		Name = RigVM_BoolPropertyPrefix + Name;
-	}
-
-	return Name;
-}
-
-FString FRigVMCodeGenerator::SanitizeValue(const FString& InValue, const FString& InCPPType, const UObject* InCPPTypeObject)
-{
-	FString DefaultValue = InValue;
-	FString CPPType = InCPPType;
-	FString BaseCPPType = CPPType;
-	bool bIsArray = RigVMTypeUtils::IsArrayType(CPPType);
-	bool bIsDoubleArray = false;
-	if (bIsArray)
-	{
-		BaseCPPType = RigVMTypeUtils::BaseTypeFromArrayType(CPPType);
-		bIsDoubleArray = RigVMTypeUtils::IsArrayType(BaseCPPType);
-		if (bIsDoubleArray)
+#include "AssetRegistry/IAssetRegistry.h"
+#include "AssetRegistry/AssetRegistryModule.h"
+#include "Modules/ModuleManager.h"
+
+FAutoConsoleCommand FCmdRigVMGenerateCode
+(
+	TEXT("RigVM.GenerateNativizedCode"),
+	TEXT("Generates code for a given rigvm based asset."),
+	FConsoleCommandWithArgsDelegate::CreateLambda([](const TArray<FString>& Args)
+	{
+		auto PrintHelp = []()
+		{
+			UE_LOG(LogRigVMDeveloper, Warning, TEXT("Valid Arguments:\n"
+				"asset=pathname - the asset to parse\n"
+				"language=CPlusPlus - the language to generate (CPlusPlus or Verse)\n"
+				"module=MyModule - the module name to use for the exports\n"
+				"output=MyFolder - the [optional] folder to export the files to\n"
+				"json=PathToMyJson - the [optional] filepath to save the json file to\n"
+				"write=0 - the [optional] settings to enable writing out the files."
+			));
+		};
+		
+		if (Args.Num() < 1)
 		{
-			BaseCPPType = RigVMTypeUtils::BaseTypeFromArrayType(BaseCPPType);			
+			PrintHelp();
+			return;
 		}
-	}
-	
-	if (const UScriptStruct* ScriptStruct = Cast<UScriptStruct>(InCPPTypeObject))
-	{
-		DefaultValue = DefaultValue.ReplaceCharWithEscapedChar();
 
-		if(!DefaultValue.IsEmpty())
+		ERigVMCodeLanguage::Type Language = ERigVMCodeLanguage::CPlusPlus;
+		
+		FString AssetPathName;
+		FRigVMCodeConversionSettings Settings;
+		Settings.TargetModule = TEXT("MyModule"); 
+	 	Settings.OutputFolder = TEXT("/Temp");
+		Settings.bWriteFiles = false;
+
+		for (const FString& Arg : Args)
 		{
-			if(const TStructConstGenerator* StructConstGenerator = GetStructConstGenerators().Find(*BaseCPPType))
+			if(Arg.Len() > 6 && Arg.StartsWith(TEXT("asset="), ESearchCase::IgnoreCase))
+			{
+				AssetPathName = Arg.Mid(6); // remove asset= from the start. 
+			}
+			else if(Arg.Len() > 7 && Arg.StartsWith(TEXT("module="), ESearchCase::IgnoreCase))
+			{
+				Settings.TargetModule = Arg.Mid(7); // remove asset= from the start. 
+			}
+			else if(Arg.Len() > 9 && Arg.StartsWith(TEXT("language="), ESearchCase::IgnoreCase))
 			{
-				if (bIsDoubleArray)
+				const FString LanguageString = Arg.Mid(9); // remove language= from the start.
+				if (LanguageString.Equals(TEXT("CPlusPlus"), ESearchCase::IgnoreCase))
 				{
-					FStringArray ArrayDefaultValues = RigVMStringUtils::SplitDefaultValue(DefaultValue);
-					for(int32 ArrayDefaultValueIndex = 0; ArrayDefaultValueIndex < ArrayDefaultValues.Num(); ArrayDefaultValueIndex++)
-					{
-						FStringArray DefaultValues = RigVMStringUtils::SplitDefaultValue(ArrayDefaultValues[ArrayDefaultValueIndex]);
-						for(int32 DefaultValueIndex = 0; DefaultValueIndex < DefaultValues.Num(); DefaultValueIndex++)
-						{
-							DefaultValues[DefaultValueIndex] = (*StructConstGenerator)(DefaultValues[DefaultValueIndex]); 
-						}
-						ArrayDefaultValues[ArrayDefaultValueIndex] = Format(RigVM_CurlyBracesFormat, FString::Join(DefaultValues, RigVM_CommaSeparator));
-					}
-					DefaultValue = Format(RigVM_CurlyBracesFormat, FString::Join(ArrayDefaultValues, RigVM_CommaSeparator));
+					Language = ERigVMCodeLanguage::CPlusPlus;
 				}
-				else if(bIsArray)
+				else if (LanguageString.Equals(TEXT("Verse"), ESearchCase::IgnoreCase))
 				{
-					FStringArray DefaultValues = RigVMStringUtils::SplitDefaultValue(DefaultValue);
-					for(int32 DefaultValueIndex = 0; DefaultValueIndex < DefaultValues.Num(); DefaultValueIndex++)
-					{
-						DefaultValues[DefaultValueIndex] = (*StructConstGenerator)(DefaultValues[DefaultValueIndex]); 
-					}
-					DefaultValue = Format(RigVM_CurlyBracesFormat, FString::Join(DefaultValues, RigVM_CommaSeparator));
+					Language = ERigVMCodeLanguage::Verse;
 				}
 				else
 				{
-					DefaultValue = (*StructConstGenerator)(DefaultValue);
+					PrintHelp();
+					return;
 				}
 			}
-			else
+			else if(Arg.Len() > 7 && Arg.StartsWith(TEXT("output="), ESearchCase::IgnoreCase))
 			{
-				if (bIsDoubleArray)
+				Settings.OutputFolder = Arg.Mid(7); // remove output= from the start.
+				if (!FPaths::DirectoryExists(Settings.OutputFolder))
 				{
-					DefaultValue = Format(RigVM_StructConstantArrayArrayValue, *BaseCPPType, *DefaultValue);
-				}
-				else if(bIsArray)
-				{
-					DefaultValue = Format(RigVM_StructConstantArrayValue, *BaseCPPType, *DefaultValue);
-				}
-				else
-				{
-					DefaultValue = Format(RigVM_StructConstantValue, *BaseCPPType, *DefaultValue);
+					UE_LOG(LogRigVMDeveloper, Warning, TEXT("Output Folder '%s' does not exist"), *Settings.OutputFolder);
+					return;
 				}
 			}
-		}
-	}
-	else if (const UEnum* Enum = Cast<UEnum>(InCPPTypeObject))
-	{
-		BaseCPPType = Enum->GetName();
-		if (Enum->GetCppForm() == UEnum::ECppForm::Namespaced)
-		{
-			BaseCPPType += RigVM_EnumTypeSuffixFormat;
-		}
-		
-		CPPType = BaseCPPType;
-		if (bIsArray)
-		{
-			CPPType = RigVMTypeUtils::ArrayTypeFromBaseType(BaseCPPType);
-		}
-
-		FStringArray DefaultValues;
-		if (bIsArray)
-		{
-			DefaultValues = RigVMStringUtils::SplitDefaultValue(DefaultValue);
-		}
-		else
-		{
-			DefaultValues.Add(DefaultValue);
-		}
-
-		for(int32 DefaultValueIndex = 0; DefaultValueIndex < DefaultValues.Num(); DefaultValueIndex++)
-		{
-			if (DefaultValues[DefaultValueIndex].IsNumeric())
+			else if(Arg.Len() > 5 && Arg.StartsWith(TEXT("json="), ESearchCase::IgnoreCase))
 			{
-				const int64 EnumIndex = FCString::Atoi64(*DefaultValues[DefaultValueIndex]);
-				const FString EnumName = Enum->GetNameStringByValue(EnumIndex);
-				DefaultValues[DefaultValueIndex] = Enum->GenerateFullEnumName(*EnumName);
+				Settings.JsonFilePath = Arg.Mid(5); // remove json= from the start.
 			}
-			else if (!UEnum::IsFullEnumName(*DefaultValues[DefaultValueIndex]))
+			else if(Arg.Len() > 6 && Arg.StartsWith(TEXT("write="), ESearchCase::IgnoreCase))
 			{
-				DefaultValues[DefaultValueIndex] = Enum->GenerateFullEnumName(*DefaultValues[DefaultValueIndex]);
+				Settings.bWriteFiles =
+					Arg.Mid(6).Equals(TEXT("1")) ||
+					Arg.Mid(6).Equals(TEXT("true"), ESearchCase::IgnoreCase) ||
+					Arg.Mid(6).Equals(TEXT("on"), ESearchCase::IgnoreCase); 
 			}
 		}
-		
-		if (bIsArray)
-		{
-			DefaultValue = Format(RigVM_CurlyBracesFormat, *FString::Join(DefaultValues, TEXT(",")));
-		}
-		else
+
+		if (AssetPathName.IsEmpty())
 		{
-			DefaultValue = DefaultValues[0];
+			PrintHelp();
+			return;
 		}
-	}
-	else
-	{
-		if (bIsArray)
+
+		FScopedSlowTask SlowTask(3.f /* Number of steps */, FText());
+		SlowTask.MakeDialog(true, false);
+		SlowTask.EnterProgressFrame(0.f, NSLOCTEXT("RigVMCodeGenerator", "LoadingAsset", "Loading Asset..."));
+
+		FAssetRegistryModule& AssetRegistryModule = FModuleManager::LoadModuleChecked<FAssetRegistryModule>("AssetRegistry");
+		FAssetData AssetData = AssetRegistryModule.Get().GetAssetByObjectPath(FSoftObjectPath(AssetPathName));
+		if (!AssetData.IsValid())
 		{
-			if (DefaultValue.StartsWith(TEXT("(")) && DefaultValue.EndsWith(TEXT(")")))
-			{
-				DefaultValue = Format(RigVM_CurlyBracesFormat, *DefaultValue.Mid(1, DefaultValue.Len() - 2));
-			}
+			UE_LOG(LogRigVMDeveloper, Error, TEXT("The AssetData '%s' could not be found in the Asset Registry."), *AssetPathName);
+			return;
 		}
 
-		if (BaseCPPType == RigVMTypeUtils::BoolType)
+		if (SlowTask.ShouldCancel())
 		{
-			DefaultValue.ToLowerInline();
+			return;
 		}
+		SlowTask.EnterProgressFrame(1.f, NSLOCTEXT("RigVMCodeGenerator", "ParsingAsset", "Parsing Asset..."));
 
-		if (CPPType == RigVMTypeUtils::FNameType || CPPType == RigVMTypeUtils::FStringType)
+		UObject* AssetObject = AssetData.GetAsset();
+		IRigVMEditorAssetInterface* EditorAsset = Cast<IRigVMEditorAssetInterface>(AssetObject);
+		if (!EditorAsset)
 		{
-			if(DefaultValue.IsEmpty())
-			{
-				if(CPPType == RigVMTypeUtils::FNameType)
-				{
-					DefaultValue = RigVM_NameNoneFormat;
-				}
-				else
-				{
-					DefaultValue = RigVM_EmptyStringFormat;
-				}
-			}
-			else
+			if (IRigVMRuntimeAssetInterface* RuntimeAsset = Cast<IRigVMRuntimeAssetInterface>(AssetObject))
 			{
-				if (DefaultValue.StartsWith(TEXT("\"")) && DefaultValue.EndsWith(TEXT("\"")))
-				{
-					DefaultValue = Format(RigVM_TextFormat, *DefaultValue);
-				}
-				else
-				{
-					DefaultValue = Format(RigVM_QuotedTextFormat, *DefaultValue);
-				}
+				EditorAsset = Cast<IRigVMEditorAssetInterface>(RuntimeAsset->GetEditorOnlyData());
 			}
 		}
-	}
-
-	return DefaultValue;
-}
-
-FRigVMPropertyDescription FRigVMCodeGenerator::GetPropertyDescForOperand(const FRigVMExtendedExecuteContext& Context, const FRigVMOperand& InOperand) const
-{
-	CheckOperand(Context, InOperand);
-
-	const FProperty* Property = nullptr;
-
-	if (InOperand.GetMemoryType() == ERigVMMemoryType::External)
-	{
-		const TArray<FRigVMExternalVariableDef> ExternalVariables = VM->GetExternalVariableDefs();
-		const FRigVMExternalVariableDef& ExternalVariable = ExternalVariables[InOperand.GetRegisterIndex()];
-		Property = ExternalVariable.GetProperty();
-	}
-	else
-	{
-		if (const FRigVMMemoryStorageStruct* MemoryStorage = VM->GetDefaultMemoryByType(InOperand.GetMemoryType()))
+		if (!EditorAsset)
 		{
-			Property = MemoryStorage->GetProperty(InOperand.GetRegisterIndex());
+			UE_LOG(LogRigVMDeveloper, Error, TEXT("The Asset '%s' is not a RigVM editor asset."), *AssetPathName);
+			return;
 		}
-	}
-
-	return (Property) ? FRigVMPropertyDescription(Property, FString()) : FRigVMPropertyDescription();
-}
 
-FRigVMPropertyDescription FRigVMCodeGenerator::GetPropertyForOperand(const FRigVMExtendedExecuteContext& Context, const FRigVMOperand& InOperand) const
-{
-	CheckOperand(Context, InOperand);
-
-	const FProperty* Property = nullptr;
-	const uint8* Memory = nullptr;
-
-	if (InOperand.GetMemoryType() == ERigVMMemoryType::External)
-	{
-		const TArray<FRigVMExternalVariable> ExternalVariables = VM->GetExternalVariables(Context);
-		const FRigVMExternalVariable& ExternalVariable = ExternalVariables[InOperand.GetRegisterIndex()];
-		Property = ExternalVariable.GetProperty();
-		Memory = ExternalVariable.GetMemory();
-	}
-	else
-	{
-		if (const FRigVMMemoryStorageStruct* MemoryStorage = VM->GetDefaultMemoryByType(InOperand.GetMemoryType()))
-		{
-			Property = MemoryStorage->GetProperty(InOperand.GetRegisterIndex());
-			if (Property)
-			{
-				Memory = Property->ContainerPtrToValuePtr<uint8>(MemoryStorage->GetContainerPtr());
-			}
-		}
-	}
+		const TSharedPtr<FRigVMCodeConverter> Converter = MakeShared<FRigVMCodeConverter>(EditorAsset, Settings);
 
-	if (Property)
-	{
-		FString DefaultValue;
-		
-		if (const UClass* OwningClass = Property->GetOwnerClass())
+		if (SlowTask.ShouldCancel())
 		{
-			if (const UObject* CDO = OwningClass->GetDefaultObject())
-			{
-				Memory = Property->ContainerPtrToValuePtr<uint8>(CDO);
-			}
+			return;
 		}
+		SlowTask.EnterProgressFrame(1.f, NSLOCTEXT("RigVMCodeGenerator", "GeneratingCode", "Generating Code..."));
 
-		if (Memory)
+		if (!Settings.JsonFilePath.IsEmpty())
 		{
-			DefaultValue = FRigVMStruct::ExportToFullyQualifiedText(Property, Memory, true);
+			const std::string path(StringCast<ANSICHAR>(*Settings.JsonFilePath).Get());
+			std::ofstream json_output_file(path);
+			json_output_file << Converter->GetJson().dump(4);
+			UE_LOG(LogRigVMDeveloper, Display, TEXT("Json data -> '%s'"), *Settings.JsonFilePath)
 		}
-		
-		return FRigVMPropertyDescription(Property, DefaultValue);
-	}
-
-	return FRigVMPropertyDescription();
-}
-
-const FRigVMPropertyPathDescription& FRigVMCodeGenerator::GetPropertyPathForOperand(const FRigVMExtendedExecuteContext& Context, const FRigVMOperand& InOperand) const
-{
-	CheckOperand(Context, InOperand);
-
-	const int32 RegisterOffsetIndex = InOperand.GetRegisterOffset();
-	if (RegisterOffsetIndex != INDEX_NONE)
-	{
-		if (InOperand.GetMemoryType() == ERigVMMemoryType::External)
+			
+		TArray<TSharedPtr<FRigVMCodeOutput>> Outputs = FRigVMCodeGenerator(Language).GenerateAll(Converter);
+		for (const TSharedPtr<FRigVMCodeOutput>& Output : Outputs)
 		{
-			if (VM->ExternalPropertyPathDescriptions.IsValidIndex(RegisterOffsetIndex))
+			if (!Output.IsValid())
 			{
-				return VM->ExternalPropertyPathDescriptions[RegisterOffsetIndex];
+				continue;
 			}
-		}
-		else
-		{
-			if (const FRigVMMemoryStorageStruct* MemoryStorage = VM->GetDefaultMemoryByType(InOperand.GetMemoryType()))
+			
+			if (!Output->Content.IsSet())
 			{
-				if (MemoryStorage->IsValidPropertyPathDescriptionIndex((RegisterOffsetIndex)))
-				{
-					return *MemoryStorage->GetPropertyPathDescriptionByIndex(RegisterOffsetIndex);
-				}
+				UE_LOG(LogRigVMDeveloper, Error, TEXT("Template '%s' failed to generate (%s)."), *Output->Name, *Output->ErrorMessage)
+			}
+			else if (Output->bSaved.Get(false))
+			{
+				UE_LOG(LogRigVMDeveloper, Display, TEXT("Template '%s' -> '%s'"), *Output->Name, *Output->FilePath)
+			}
+			else
+			{
+				UE_LOG(LogRigVMDeveloper, Display, TEXT("Template '%s' -> '%s'\n\n%s"), *Output->Name, *Output->FilePath, *Output->Content.GetValue())
 			}
 		}
-	}
-
-	static const FRigVMPropertyPathDescription EmptyPropertyPath;
-	return EmptyPropertyPath;
-}
-
-int32 FRigVMCodeGenerator::GetPropertyIndex(const FRigVMPropertyDescription& InProperty) const
-{
-	if(const int32* Index = PropertyNameToIndex.Find(InProperty.Name))
-	{
-		return *Index;
-	}
-	return INDEX_NONE;
-}
+	})
+);
 
-int32 FRigVMCodeGenerator::GetPropertyIndex(const FRigVMExtendedExecuteContext& Context, const FRigVMOperand& InOperand) const
+FRigVMCodeGenerator::FRigVMCodeGenerator(ERigVMCodeLanguage::Type InLanguage)
+	: Language(InLanguage)
 {
-	return GetPropertyIndex(GetPropertyDescForOperand(Context, InOperand));
-}
-
-FRigVMCodeGenerator::ERigVMNativizedPropertyType FRigVMCodeGenerator::GetPropertyType(
-	const FRigVMPropertyDescription& InProperty) const
-{
-	const int32 PropertyIndex = GetPropertyIndex(InProperty);
-	if(Properties.IsValidIndex(PropertyIndex))
-	{
-		return Properties[PropertyIndex].PropertyType;
-	}
-	return ERigVMNativizedPropertyType::Invalid;
-}
-
-FRigVMCodeGenerator::ERigVMNativizedPropertyType FRigVMCodeGenerator::GetPropertyType(const FRigVMExtendedExecuteContext& Context, const FRigVMOperand& InOperand) const
-{
-	return GetPropertyType(GetPropertyDescForOperand(Context, InOperand));
-}
-
-void FRigVMCodeGenerator::CheckOperand(const FRigVMExtendedExecuteContext& Context, const FRigVMOperand& InOperand) const
-{
-	ensure(InOperand.IsValid());
-	ensure(InOperand.GetMemoryType() != ERigVMMemoryType::Invalid);
-	ensure(InOperand.GetMemoryType() != ERigVMMemoryType::Debug);
-
-	if (InOperand.GetMemoryType() == ERigVMMemoryType::External)
+	FString LanguageName;
+	switch (InLanguage)
 	{
-		ensure(VM->ExternalVariables.IsValidIndex(InOperand.GetRegisterIndex()));
-		if (InOperand.GetRegisterOffset() != INDEX_NONE)
+		case ERigVMCodeLanguage::CPlusPlus:
 		{
-			ensure(VM->ExternalPropertyPaths.IsValidIndex(InOperand.GetRegisterOffset()));
-			ensure(VM->ExternalPropertyPathDescriptions.IsValidIndex(InOperand.GetRegisterOffset()));
+			LanguageName = TEXT("CPlusPlus");
+			break;
 		}
-	}
-	else
-	{
-		const FRigVMMemoryStorageStruct* Memory = VM->GetDefaultMemoryByType(InOperand.GetMemoryType());
-		check(Memory);
-
-		ensure(Memory->GetProperties().IsValidIndex(InOperand.GetRegisterIndex()));
-		if (InOperand.GetRegisterOffset() != INDEX_NONE)
+		default:
 		{
-			ensure(Memory->GetPropertyPaths().IsValidIndex(InOperand.GetRegisterOffset()));
+			return;
 		}
 	}
-}
-
-FString FRigVMCodeGenerator::GetMappedType(const FString& InCPPType) const
-{
-	FString CPPType = InCPPType;
-	CPPType.RemoveSpacesInline();
-	
-	if(const FMappedType* MappedType = MappedCPPTypes.Find(CPPType))
-	{
-		return MappedType->Get<0>();
-	}
-
-	if(RigVMTypeUtils::IsArrayType(CPPType))
-	{
-		const FString BaseType = RigVMTypeUtils::BaseTypeFromArrayType(CPPType);
-		return RigVMTypeUtils::ArrayTypeFromBaseType(GetMappedType(BaseType));
-	}
 
-	return CPPType;
-}
+	const FString PluginContentDir = FPaths::EnginePluginsDir() / TEXT("Runtime/RigVM/Content");
+	const FString TemplatesDir = FPaths::Combine(PluginContentDir, TEXT("CodeGeneration"), LanguageName);
+	const FString TemplateSearchPath = FPaths::Combine(*TemplatesDir, TEXT("*.template"));
+	TArray<FString> TemplateFileNames;
+	FFileManagerGeneric::Get().FindFiles(TemplateFileNames, *TemplateSearchPath, true, false);
 
-const FString& FRigVMCodeGenerator::GetMappedTypeSuffix(const FString& InCPPType) const
-{
-	FString CPPType = InCPPType;
-	CPPType.RemoveSpacesInline();
-	
-	if(const FMappedType* MappedType = MappedCPPTypes.Find(CPPType))
+	for (const FString& TemplateFileName : TemplateFileNames)
 	{
-		return MappedType->Get<1>();
+		TSharedPtr<FRigVMCodeTemplate> Template = MakeShared<FRigVMCodeTemplate>();
+		Template->Name = TemplateFileName.LeftChop(9); // remove .template
+		Template->FilePath = FPaths::Combine(TemplatesDir, TemplateFileName);
+		CodeTemplates.Add(Template);
 	}
-
-	static const FString EmptyString;
-	return EmptyString;
 }
 
-FString FRigVMCodeGenerator::GetMappedArrayTypeName(const FString InBaseElementType) const
+int32 FRigVMCodeGenerator::NumTemplates() const
 {
-	return Format(RigVM_WrappedTypeNameFormat, *InBaseElementType, *ClassName);
-}
-
-const FRigVMCodeGenerator::FInstructionGroup& FRigVMCodeGenerator::GetGroup(int32 InGroupIndex) const
-{
-	if(InstructionGroups.IsValidIndex(InGroupIndex))
-	{
-		return InstructionGroups[InGroupIndex];
-	}
-	
-	static FInstructionGroup CompleteGroup;
-	if(CompleteGroup.First == INDEX_NONE)
-	{
-		CompleteGroup.First = 0;
-		CompleteGroup.Last = VM->GetByteCode().GetNumInstructions() - 1; 
-	}
-	return CompleteGroup;
+	return CodeTemplates.Num();
 }
 
-bool FRigVMCodeGenerator::IsInstructionPartOfGroup(int32 InInstructionIndex, int32 InGroupIndex, bool bIncludeChildGroups) const
+const TSharedPtr<FRigVMCodeTemplate>& FRigVMCodeGenerator::GetTemplate(const FString& InName) const
 {
-	if(!InstructionGroups.IsValidIndex(InGroupIndex) && InstructionGroups.IsEmpty())
-	{
-		return true;
-	}
-
-	const FInstructionGroup& Group = InstructionGroups[InGroupIndex];
-	for(const int32 ChildGroupIndex : Group.ChildGroups)
+	for (const TSharedPtr<FRigVMCodeTemplate>& CodeTemplate : CodeTemplates)
 	{
-		if(IsInstructionPartOfGroup(InInstructionIndex, ChildGroupIndex, true))
+		if (CodeTemplate->Name.Equals(InName, ESearchCase::IgnoreCase))
 		{
-			// if we found the operation in a child group
-			// we'll return true if we are also looking within the child groups,
-			// and false if we are not. this means that an operation that's part of a
-			// child group but not the main group will return false here
-			return bIncludeChildGroups;
+			return CodeTemplate;
 		}
 	}
 
-	return FMath::IsWithinInclusive(InInstructionIndex, Group.First, Group.Last);;
+	static TSharedPtr<FRigVMCodeTemplate> EmptyTemplate;
+	return EmptyTemplate;
 }
 
-bool FRigVMCodeGenerator::IsPropertyPartOfGroup(int32 InPropertyIndex, int32 InGroupIndex) const
+TSharedPtr<FRigVMCodeOutput> FRigVMCodeGenerator::Generate(const TSharedPtr<FRigVMCodeConverter>& InConverter, const FString& InTemplateName) const
 {
-	if(InGroupIndex == INDEX_NONE)
-	{
-		return true;
-	}
-
-	if(!InstructionGroups.IsValidIndex(InGroupIndex))
-	{
-		return false;
-	}
-
-	if(Properties.IsValidIndex(InPropertyIndex))
+	TSharedPtr<FRigVMCodeTemplate> CodeTemplate = GetTemplate(InTemplateName);
+	if (!CodeTemplate)
 	{
-		const FPropertyInfo& Property = Properties[InPropertyIndex];
-		const FInstructionGroup& Group = InstructionGroups[InGroupIndex];
-
-		// properties are only defines in the leaf groups
-		if(!Group.ChildGroups.IsEmpty())
-		{
-			return false;
-		}
-
-		if(Property.PropertyType == ERigVMNativizedPropertyType::Work)
-		{
-			if(Property.Groups.Num() > 1)
-			{
-				return false;
-			}
-		}
-
-		return Property.Groups.Contains(InGroupIndex);
+		UE_LOG(LogRigVMDeveloper, Error, TEXT("The Code Template '%s' cannot be found."), *InTemplateName);
+		return nullptr;
 	}
-	
-	return false;
+	return Generate(InConverter, CodeTemplate);
 }
 
-FString FRigVMCodeGenerator::GetEntryParameters() const
+TSharedPtr<FRigVMCodeOutput> FRigVMCodeGenerator::Generate(const TSharedPtr<FRigVMCodeConverter>& InConverter, const TSharedPtr<FRigVMCodeTemplate>& InCodeTemplate)
 {
-	const FRigVMByteCode& ByteCode = VM->GetByteCode();
-	const int32 EntryIndex = ByteCode.NumEntries() - 1;
-	const FRigVMByteCodeEntry& Entry = ByteCode.GetEntry(EntryIndex);
-	const FString EntryName = Entry.GetSanitizedName();
-	FString Parameters;
-
-	// find the entry's group and provide the needed arguments
-	for(const FInstructionGroup& EntryGroup : InstructionGroups)
-	{
-		if(EntryGroup.Depth <= 0 && EntryGroup.Entry == EntryName)
-		{
-			TArray<FString> ArgumentNames = {RigVM_ContextPublicFormat};
-			Parameters = FString::Join(ArgumentNames, RigVM_CommaSeparator);
-			break;
-		}
-	}
-
-	if(!Parameters.IsEmpty())
+	if (!InConverter.IsValid() || !InCodeTemplate.IsValid())
 	{
-		Parameters = RigVM_CommaSeparator + Parameters;
+		return nullptr;
 	}
-		
-	return Parameters;
+	return InConverter->Render(InCodeTemplate);
 }
 
-const TMap<FName, FRigVMCodeGenerator::TStructConstGenerator>& FRigVMCodeGenerator::GetStructConstGenerators()
+TArray<TSharedPtr<FRigVMCodeOutput>> FRigVMCodeGenerator::GenerateAll(const TSharedPtr<FRigVMCodeConverter>& InConverter) const
 {
-	static TMap<FName, FRigVMCodeGenerator::TStructConstGenerator> StructConstGenerators;
-	if(!StructConstGenerators.IsEmpty())
-	{
-		return StructConstGenerators;
-	}
-
-	struct DefaultValueHelpers
-	{
-		static FStringArray SplitIntoArray(const FString& InValue)
-		{
-			return RigVMStringUtils::SplitDefaultValue(InValue);
-		}
-		
-		static FStringMap SplitIntoMap(const FString& InValue)
-		{
-			FStringArray Pairs = SplitIntoArray(InValue);
-			FStringMap Map;
-			for(const FString& Pair : Pairs)
-			{
-				FString Key, Value;
-				if(Pair.Split(TEXT("="), &Key, &Value))
-				{
-					Map.Add(Key, Value);
-				}
-			}
-			return Map;
-		}
-
-		static FString RemoveQuotes(const FString& InValue)
-		{
-			FString Value = InValue;
-			if(Value.StartsWith(TEXT("\\\"")))
-			{
-				Value = Value.Mid(2, Value.Len()  - 4);
-			}
-			else
-			{
-				Value = Value.TrimQuotes();
-			}
-			return Value;
-		}
-	};
-
-	static constexpr TCHAR X[] = TEXT("X");
-	static constexpr TCHAR Y[] = TEXT("Y");
-	static constexpr TCHAR Z[] = TEXT("Z");
-	static constexpr TCHAR W[] = TEXT("W");
-	static constexpr TCHAR R[] = TEXT("R");
-	static constexpr TCHAR G[] = TEXT("G");
-	static constexpr TCHAR B[] = TEXT("B");
-	static constexpr TCHAR A[] = TEXT("A");
-	static constexpr TCHAR Translation[] = TEXT("Translation");
-	static constexpr TCHAR Rotation[] = TEXT("Rotation");
-	static constexpr TCHAR Scale[] = TEXT("Scale");
-	static constexpr TCHAR Scale3D[] = TEXT("Scale3D");
-	static constexpr TCHAR Pitch[] = TEXT("Pitch");
-	static constexpr TCHAR Yaw[] = TEXT("Yaw");
-	static constexpr TCHAR Roll[] = TEXT("Roll");
-	static constexpr TCHAR Type[] = TEXT("Type");
-	static constexpr TCHAR Name[] = TEXT("Name");
-	static constexpr TCHAR XPlane[] = TEXT("XPlane");
-	static constexpr TCHAR YPlane[] = TEXT("YPlane");
-	static constexpr TCHAR ZPlane[] = TEXT("ZPlane");
-	static constexpr TCHAR WPlane[] = TEXT("WPlane");
-
-	StructConstGenerators.Add(TEXT("FVector2D"), [](const FString& InDefault) -> FString
-	{
-		static constexpr TCHAR Constructor[] = TEXT("FVector2D({0}, {1})");
-		const FStringMap Defaults = DefaultValueHelpers::SplitIntoMap(InDefault);
-		if(Defaults.Contains(X) && Defaults.Contains(Y))
-		{
-			return Format(Constructor, Defaults.FindChecked(X), Defaults.FindChecked(Y));
-		}
-		return FString();
-	});
-
-	StructConstGenerators.Add(TEXT("FVector"), [](const FString& InDefault) -> FString
-	{
-		static constexpr TCHAR Constructor[] = TEXT("FVector({0}, {1}, {2})");
-		const FStringMap Defaults = DefaultValueHelpers::SplitIntoMap(InDefault);
-		if(Defaults.Contains(X) && Defaults.Contains(Y) && Defaults.Contains(Z))
-		{
-			return Format(Constructor, Defaults.FindChecked(X), Defaults.FindChecked(Y), Defaults.FindChecked(Z));
-		}
-		return FString();
-	});
-
-	StructConstGenerators.Add(TEXT("FVector4"), [](const FString& InDefault) -> FString
-	{
-		static constexpr TCHAR Constructor[] = TEXT("FVector4({0}, {1}, {2}, {3})");
-		const FStringMap Defaults = DefaultValueHelpers::SplitIntoMap(InDefault);
-		if(Defaults.Contains(X) && Defaults.Contains(Y) && Defaults.Contains(Z) && Defaults.Contains(W))
-		{
-			return Format(Constructor, Defaults.FindChecked(X), Defaults.FindChecked(Y), Defaults.FindChecked(Z), Defaults.FindChecked(W));
-		}
-		return FString();
-	});
-
-	StructConstGenerators.Add(TEXT("FLinearColor"), [](const FString& InDefault) -> FString
-	{
-		static constexpr TCHAR Constructor[] = TEXT("FLinearColor({0}, {1}, {2}, {3})");
-		const FStringMap Defaults = DefaultValueHelpers::SplitIntoMap(InDefault);
-		if(Defaults.Contains(R) && Defaults.Contains(G) && Defaults.Contains(B) && Defaults.Contains(A))
-		{
-			return Format(Constructor, Defaults.FindChecked(R), Defaults.FindChecked(G), Defaults.FindChecked(B), Defaults.FindChecked(A));
-		}
-		return FString();
-	});
-
-	StructConstGenerators.Add(TEXT("FQuat"), [](const FString& InDefault) -> FString
-	{
-		static constexpr TCHAR Constructor[] = TEXT("FQuat({0}, {1}, {2}, {3})");
-		const FStringMap Defaults = DefaultValueHelpers::SplitIntoMap(InDefault);
-		if(Defaults.Contains(X) && Defaults.Contains(Y) && Defaults.Contains(Z) && Defaults.Contains(W))
-		{
-			return Format(Constructor, Defaults.FindChecked(X), Defaults.FindChecked(Y), Defaults.FindChecked(Z), Defaults.FindChecked(W));
-		}
-		return FString();
-	});
-
-	StructConstGenerators.Add(TEXT("FRotator"), [](const FString& InDefault) -> FString
-	{
-		static constexpr TCHAR Constructor[] = TEXT("FRotator({0}, {1}, {2})");
-		const FStringMap Defaults = DefaultValueHelpers::SplitIntoMap(InDefault);
-		if(Defaults.Contains(Pitch) && Defaults.Contains(Yaw) && Defaults.Contains(Roll))
-		{
-			return Format(Constructor, Defaults.FindChecked(Pitch), Defaults.FindChecked(Yaw), Defaults.FindChecked(Roll));
-		}
-		return FString();
-	});
-
-	StructConstGenerators.Add(TEXT("FTransform"), [](const FString& InDefault) -> FString
-	{
-		static constexpr TCHAR Constructor[] = TEXT("FTransform({0}, {1}, {2})");
-		const FStringMap Defaults = DefaultValueHelpers::SplitIntoMap(InDefault);
-		if(Defaults.Contains(Rotation) && Defaults.Contains(Translation) && Defaults.Contains(Scale3D))
-		{
-			const FString RotationValue = StructConstGenerators.FindChecked(TEXT("FQuat"))(Defaults.FindChecked(Rotation));
-			const FString TranslationValue = StructConstGenerators.FindChecked(TEXT("FVector"))(Defaults.FindChecked(Translation));
-			const FString Scale3DValue = StructConstGenerators.FindChecked(TEXT("FVector"))(Defaults.FindChecked(Scale3D));
-			return Format(Constructor, RotationValue, TranslationValue, Scale3DValue);
-		}
-		return FString();
-	});
-
-	StructConstGenerators.Add(TEXT("FEulerTransform"), [](const FString& InDefault) -> FString
-	{
-		static constexpr TCHAR Constructor[] = TEXT("FEulerTransform({0}, {1}, {2})");
-		const FStringMap Defaults = DefaultValueHelpers::SplitIntoMap(InDefault);
-		if(Defaults.Contains(Rotation) && Defaults.Contains(Translation) && Defaults.Contains(Scale3D))
-		{
-			const FString RotationValue = StructConstGenerators.FindChecked(TEXT("FRotator"))(Defaults.FindChecked(Rotation));
-			const FString TranslationValue = StructConstGenerators.FindChecked(TEXT("FVector"))(Defaults.FindChecked(Translation));
-			const FString Scale3DValue = StructConstGenerators.FindChecked(TEXT("FVector"))(Defaults.FindChecked(Scale3D));
-			return Format(Constructor, RotationValue, TranslationValue, Scale3DValue);
-		}
-		return FString();
-	});
-
-	StructConstGenerators.Add(TEXT("FRigElementKey"), [](const FString& InDefault) -> FString
-	{
-		static constexpr TCHAR Constructor[] = TEXT("FRigElementKey(TEXT(\"{0}\"), ERigElementType::{1})");
-		const FStringMap Defaults = DefaultValueHelpers::SplitIntoMap(InDefault);
-		if(Defaults.Contains(Type) && Defaults.Contains(Name))
-		{
-			const FString NameValue = DefaultValueHelpers::RemoveQuotes(Defaults.FindChecked(Name));
-			return Format(Constructor, NameValue, Defaults.FindChecked(Type));
-		}
-		return FString();
-	});
-
-	StructConstGenerators.Add(TEXT("FPlane"), [](const FString& InDefault) -> FString
-	{
-		static constexpr TCHAR Constructor[] = TEXT("FPlane({0}, {1}, {2}, {3})");
-		const FStringMap Defaults = DefaultValueHelpers::SplitIntoMap(InDefault);
-		if(Defaults.Contains(X) && Defaults.Contains(Y) && Defaults.Contains(Z) && Defaults.Contains(W))
-		{
-			return Format(Constructor, Defaults.FindChecked(X), Defaults.FindChecked(Y), Defaults.FindChecked(Z), Defaults.FindChecked(W));
-		}
-		return FString();
-	});
+	FScopedSlowTask SlowTask(static_cast<float>(CodeTemplates.Num()), FText());
 
-	StructConstGenerators.Add(TEXT("FMatrix"), [](const FString& InDefault) -> FString
+	TArray<TSharedPtr<FRigVMCodeOutput>> Outputs;
+	for (const TSharedPtr<FRigVMCodeTemplate>& CodeTemplate : CodeTemplates)
 	{
-		static constexpr TCHAR Constructor[] = TEXT("FMatrix({0}, {1}, {2}, {3})");
-		const FStringMap Defaults = DefaultValueHelpers::SplitIntoMap(InDefault);
-		if(Defaults.Contains(XPlane) && Defaults.Contains(YPlane) && Defaults.Contains(ZPlane) && Defaults.Contains(WPlane))
-		{
-			const FString XPlaneValue = StructConstGenerators.FindChecked(TEXT("FPlane"))(Defaults.FindChecked(XPlane));
-			const FString YPlaneValue = StructConstGenerators.FindChecked(TEXT("FPlane"))(Defaults.FindChecked(YPlane));
-			const FString ZPlaneValue = StructConstGenerators.FindChecked(TEXT("FPlane"))(Defaults.FindChecked(ZPlane));
-			const FString WPlaneValue = StructConstGenerators.FindChecked(TEXT("FPlane"))(Defaults.FindChecked(WPlane));
-			return Format(Constructor, XPlaneValue, YPlaneValue, ZPlaneValue, WPlaneValue);
-		}
-		return FString();
-	});
-
-	return StructConstGenerators;
+		const FString Message = FString::Printf(TEXT("Generating code for template '%s'..."), *CodeTemplate->Name);
+		UE_LOG(LogRigVMDeveloper, Display, TEXT("%s"), *Message);
+		SlowTask.EnterProgressFrame(0.f, FText::FromString(Message));
+		Outputs.Add(Generate(InConverter, CodeTemplate));
+		SlowTask.EnterProgressFrame(1.f);
+	}
+	return Outputs;
 }
+
+#undef TOJSON
+#undef FROMJSON
