RigVM Nativization System

diff --git a/Engine/Plugins/Runtime/RigVM/Source/RigVMDeveloper/Private/RigVMCompiler/RigVMCompiler.cpp b/Engine/Plugins/Runtime/RigVM/Source/RigVMDeveloper/Private/RigVMCompiler/RigVMCompiler.cpp
--- a/Engine/Plugins/Runtime/RigVM/Source/RigVMDeveloper/Private/RigVMCompiler/RigVMCompiler.cpp
+++ b/Engine/Plugins/Runtime/RigVM/Source/RigVMDeveloper/Private/RigVMCompiler/RigVMCompiler.cpp
@@ -1891,7 +1891,7 @@ bool URigVMCompiler::Compile(const FRigVMCompileSettings& InSettings, TArray<URi
 
 			// convert this function to a callable
 			FRigVMCallableInfo Callable;
-			Callable.Hash = Pair.Value.CompilationData->Hash;
+			Callable.FunctionHash = Pair.Value.CompilationData->Hash;
 			Callable.Name = *Header.GetHash(); 
 			Callable.FirstInstruction = FirstInstructionIndex;
 			Callable.LastInstruction = Pair.Value.CompilationData->ByteCode.GetNumInstructions() - 1;
@@ -1920,7 +1920,14 @@ bool URigVMCompiler::Compile(const FRigVMCompileSettings& InSettings, TArray<URi
 
 				CallableArgument.InterfaceOperand = *InterfaceOperand;
 				CallableArgument.ForwardedOperand = FRigVMOperand();
-				CallableArgument.Direction = FunctionArgument.Direction;
+				if (FunctionArgument.bIsInputVariable)
+				{
+					CallableArgument.Direction = ERigVMPinDirection::IO;
+				}
+				else
+				{
+					CallableArgument.Direction = FunctionArgument.Direction;
+				}
 				Callable.Arguments.Add(CallableArgument);
 			}
 
@@ -2411,7 +2418,7 @@ bool URigVMCompiler::CompileCallable(FRigVMCompilerWorkData& WorkData, URigVM* O
 	
 	// if we already compiled this callable - carry on
 	FRigVMByteCode& ByteCode = OutVM->GetByteCode();
-	if (ByteCode.FindCallable(InCallable->Hash))
+	if (ByteCode.FindCallable(InCallable->FunctionHash))
 	{
 		return true;
 	}
@@ -2431,17 +2438,20 @@ bool URigVMCompiler::CompileCallable(FRigVMCompilerWorkData& WorkData, URigVM* O
 			OldForwardedHandleToInterfaceOperand.Add(Argument.ForwardedOperand, Argument.InterfaceOperand);
 		}
 	}
-	(void)WorkData.AddPropertiesForFunction(true, InCallable->Hash, InHeader, InFunctionCompilationData, InCallable->FirstInstruction, InCallable->LastInstruction, ArgumentOperands, OldForwardedHandleToInterfaceOperand);
+
+	const uint32 HashForProperties = GetTypeHash(*InCallable);
+	(void)WorkData.AddPropertiesForFunction(true, HashForProperties, InHeader, InFunctionCompilationData, InCallable->FirstInstruction, InCallable->LastInstruction, ArgumentOperands, OldForwardedHandleToInterfaceOperand);
 
 	FRigVMCallableInfo NewCallable;
 	NewCallable.Index = ByteCode.NumCallables();
-	NewCallable.Hash = InCallable->Hash;
+	NewCallable.FunctionHash = InCallable->FunctionHash;
 	NewCallable.Name = InCallable->Name;
 	NewCallable.FirstInstruction = ByteCode.GetNumInstructions();
 	NewCallable.LastInstruction = NewCallable.FirstInstruction + InCallable->LastInstruction - InCallable->FirstInstruction;
 
 	TMap<FRigVMOperand, FRigVMOperand> NewInterfaceOperandToForwardedHandle;
-	const TMap<FRigVMOperand, FRigVMOperand>& OperandMap = WorkData.OperandMapPerFunction.FindChecked(NewCallable.Hash);
+
+	const TMap<FRigVMOperand, FRigVMOperand>& OperandMap = WorkData.OperandMapPerFunction.FindChecked(HashForProperties);
 	auto RemapOperand = [&OperandMap, &NewInterfaceOperandToForwardedHandle, &OldForwardedHandleToInterfaceOperand, InHeader](const FRigVMOperand& InOldOperand, bool bRemapToForwardedHandle)
 	{
 		FRigVMOperand OperandNoOffset(InOldOperand.GetMemoryType(), InOldOperand.GetRegisterIndex(), INDEX_NONE);
@@ -2570,7 +2580,7 @@ bool URigVMCompiler::CompileCallable(FRigVMCompilerWorkData& WorkData, URigVM* O
 		{
 			const FRigVMCallableInfo* OldCallable = FunctionByteCode.GetCallable(InCallableIndex);
 			check(OldCallable);
-			const FRigVMCallableInfo* NewCallable = ByteCode.FindCallable(OldCallable->Hash);
+			const FRigVMCallableInfo* NewCallable = ByteCode.FindCallable(OldCallable->FunctionHash);
 			check(NewCallable);
 			return NewCallable->Index;
 		};
@@ -3667,7 +3677,7 @@ bool URigVMCompiler::TraverseInlineFunction(const FRigVMInlineFunctionExprAST* I
 	{
 		const FRigVMCallableInfo* OldCallable = FunctionByteCode.GetCallable(CallableIndex);
 		check(OldCallable);
-		const FRigVMCallableInfo* NewCallable = ByteCode.FindCallable(OldCallable->Hash);
+		const FRigVMCallableInfo* NewCallable = ByteCode.FindCallable(OldCallable->FunctionHash);
 		if (!ensure(NewCallable != nullptr))
 		{
 			UE_LOG(LogRigVMDeveloper, Error, TEXT("'%s': Unexpected missing callable '%s' when inlining function '%s'."),
@@ -3739,7 +3749,7 @@ bool URigVMCompiler::TraverseInlineFunction(const FRigVMInlineFunctionExprAST* I
 		{
 			const FRigVMCallableInfo* OldCallable = FunctionByteCode.GetCallable(InOldCallableIndex);
 			check(OldCallable);
-			const FRigVMCallableInfo* NewCallable = ByteCode.FindCallable(OldCallable->Hash);
+			const FRigVMCallableInfo* NewCallable = ByteCode.FindCallable(OldCallable->FunctionHash);
 			check(NewCallable);
 			return NewCallable->Index;
 		};
@@ -4062,7 +4072,7 @@ bool URigVMCompiler::TraverseCallableFunction(const FRigVMInlineFunctionExprAST*
 #endif
 		}
 		
-		const int32 InvokeCallableByteIndex = ByteCode.AddInvokeCallableOp(Callable->Hash, CallableOperands);
+		const int32 InvokeCallableByteIndex = ByteCode.AddInvokeCallableOp(Callable->FunctionHash, CallableOperands);
 		const int32 InvokeCallableInstructionIndex = ByteCode.GetNumInstructions() - 1;
 		const FRigVMInstruction InvokeCallableInstruction(ERigVMOpCode::InvokeCallable, InvokeCallableInstructionIndex, InvokeCallableByteIndex);
 
diff --git a/Engine/Plugins/Runtime/RigVM/Source/RigVMDeveloper/Private/RigVMEditorAsset.cpp b/Engine/Plugins/Runtime/RigVM/Source/RigVMDeveloper/Private/RigVMEditorAsset.cpp
--- a/Engine/Plugins/Runtime/RigVM/Source/RigVMDeveloper/Private/RigVMEditorAsset.cpp
+++ b/Engine/Plugins/Runtime/RigVM/Source/RigVMDeveloper/Private/RigVMEditorAsset.cpp
@@ -1011,57 +1011,15 @@ bool IRigVMEditorAssetInterface::SetEarlyExitInstruction(URigVMNode* InNodeToExi
 		return false;
 	}
 
-	GetDebugInfo().SetInstructionForEarlyExit(InstructionIndex);
+	GetDebugInfo().RunToInstruction(InstructionIndex);
 
 	if (URigVMHost* DebuggedHost = GetDebuggedRigVMHost())
 	{
-		DebuggedHost->GetDebugInfo().CopyFrom(GetDebugInfo());
+		DebuggedHost->GetDebugInfo().SetStepCondition(GetDebugInfo().GetStepCondition());
 		DebuggedHost->SetIsInDebugMode(true);
 	}
 
-	TArray<URigVMNode*> IncludedNodes;
-	TSet<URigVMNode*> VisitedNodes;
-
-	IncludedNodes.Add(InNodeToExitEarlyAfter);
-
-	for (int32 Index = 0; Index < IncludedNodes.Num(); Index++)
-	{
-		VisitedNodes.Add(IncludedNodes[Index]);
-		
-		const TArray<URigVMNode*> SourceNodes = IncludedNodes[Index]->GetLinkedSourceNodes();
-		for (URigVMNode* SourceNode : SourceNodes)
-		{
-			if (VisitedNodes.Contains(SourceNode))
-			{
-				continue;
-			}
-			IncludedNodes.AddUnique(SourceNode);
-		}
-	}
-
-	InNodeToExitEarlyAfter->SetHasEarlyExitMarker(true);
-
-	TArray<URigVMNode*> AllNodes = InNodeToExitEarlyAfter->GetGraph()->GetNodes();
-	for (URigVMNode* Node : AllNodes)
-	{
-		Node->SetHasEarlyExitMarker(Node == InNodeToExitEarlyAfter);
-		Node->SetIsExcludedByEarlyExit(!VisitedNodes.Contains(Node));
-	}
-
-	GetRigGraphDisplaySettings().bShowNodeRunCounts = true;
-
-	LastPreviewHereNodes.Emplace(InNodeToExitEarlyAfter, InstructionIndex);
-
-	if (bRequestHyperLink)
-	{
-		if (FRigVMClient* Client = GetRigVMClient())
-		{
-			if (URigVMController* Controller = Client->GetOrCreateController(InNodeToExitEarlyAfter->GetGraph()))
-			{
-				Controller->RequestJumpToHyperLink(InNodeToExitEarlyAfter);
-			}
-		}
-	}
+	MarkNodeForPreviewHere(InNodeToExitEarlyAfter, InstructionIndex, bRequestHyperLink);
 
 	return true;
 }
@@ -1072,7 +1030,7 @@ bool IRigVMEditorAssetInterface::ResetEarlyExitInstruction(bool bResetCallstack)
 
 	if (URigVMHost* DebuggedHost = GetDebuggedRigVMHost())
 	{
-		DebuggedHost->GetDebugInfo().CopyFrom(GetDebugInfo());
+		DebuggedHost->GetDebugInfo().SetStepCondition(GetDebugInfo().GetStepCondition());
 	}
 
 	TArray<URigVMGraph*> AllGraphs = GetAllModels();
@@ -1105,7 +1063,7 @@ void IRigVMEditorAssetInterface::TogglePreviewHere(const URigVMGraph* InGraph)
 		return;
 	}
 
-	const int32 CurrentEarlyExistInstruction = DebuggedHost->GetDebugInfo().GetInstructionForEarlyExit();
+	const int32 CurrentEarlyExistInstruction = DebuggedHost->GetDebugInfo().GetStepCondition().OriginInstruction;
 	if (CurrentEarlyExistInstruction != INDEX_NONE)
 	{
 		if (CurrentEarlyExistInstruction == GetPreviewNodeInstructionIndexFromSelection(InGraph))
@@ -1143,214 +1101,86 @@ void IRigVMEditorAssetInterface::PreviewHereStepForward()
 	{
 		return;
 	}
-
- 	int32 CurrentInstructionIndex = DebuggedHost->GetDebugInfo().GetInstructionForEarlyExit();
-	if (CurrentInstructionIndex == INDEX_NONE)
+	if (!DebuggedHost->GetRigVMExtendedExecuteContext().StepInto())
 	{
 		return;
 	}
 
-	const URigVM* VM = DebuggedHost->GetVM();
-	if (!VM)
+	const int32 InstructionIndex = DebuggedHost->GetDebugInfo().GetStepCondition().OriginInstruction;
+	if (InstructionIndex == INDEX_NONE)
 	{
 		return;
 	}
 
-	URigVMNode* Node = nullptr;
-	if (!LastPreviewHereNodes.IsEmpty())
-	{
-		const TPair<TWeakObjectPtr<URigVMNode>,int32> PreviousEntry = LastPreviewHereNodes.Pop();
-		Node = PreviousEntry.Key.Get();
-		if (Node)
-		{
-			CurrentInstructionIndex = PreviousEntry.Value;
-		}
-	}
-
-	if (!Node)
-	{
-		Node = Cast<URigVMNode>(VM->GetByteCode().GetSubjectForInstruction(CurrentInstructionIndex));
-	}
-	
+	URigVMNode* Node = Cast<URigVMNode>(DebuggedHost->GetVM()->GetByteCode().GetSubjectForInstruction(InstructionIndex));
 	if (!Node)
 	{
 		return;
 	}
+	MarkNodeForPreviewHere(Node, InstructionIndex, true);
+}
 
-	struct Local
+bool IRigVMEditorAssetInterface::CanPreviewHereStepForward() const
+{
+	URigVMHost* DebuggedHost = Cast<URigVMHost>(GetObjectBeingDebugged());
+	if (!DebuggedHost)
 	{
-		static URigVMNode* FindNextNodeOnCompletedPin(URigVMNode* Node, TSet<URigVMNode*>& VisitedNodes)
-		{
-			if (VisitedNodes.Contains(Node))
-			{
-				return nullptr;
-			}
-			VisitedNodes.Add(Node);
-			
-			for (const URigVMPin* Pin : Node->GetPins())
-			{
-				if (Pin->GetDirection() == ERigVMPinDirection::Input || Pin->GetDirection() == ERigVMPinDirection::IO)
-				{
-					if (!Pin->IsExecuteContext())
-					{
-						continue;
-					}
-
-					const TArray<URigVMPin*> SourcePins = Pin->GetLinkedSourcePins();
-					if (SourcePins.IsEmpty())
-					{
-						continue;
-					}
-
-					URigVMNode* SourceNode = SourcePins[0]->GetNode();
-					if (SourcePins[0]->GetFName() == FRigVMStruct::ControlFlowCompletedName)
-					{
-						return FindNextNodeOnCompletedPin(SourceNode,VisitedNodes);
-					}
+		return false;
+	}
 
-					TArray<const URigVMPin*> PinsToFollow;
+	return DebuggedHost->GetDebugInfo().GetStepCondition().OriginInstruction != INDEX_NONE;
 
-					if (const URigVMPin* CompletedPin = SourceNode->FindRootPinByName(FRigVMStruct::ControlFlowCompletedName))
-					{
-						PinsToFollow.Add(CompletedPin);
-					}
-					else
-					{
-						// since we don't have a completed pin, try to the find the next execute pin after this one
-						const int32 SourcePinIndex = SourcePins[0]->GetPinIndex();
-						for (int32 Index = SourcePinIndex + 1; Index < SourceNode->GetPins().Num(); Index++)
-						{
-							const URigVMPin* NextExecutePin = SourceNode->GetPins()[Index];
-							if (NextExecutePin->GetDirection() == ERigVMPinDirection::Output && NextExecutePin->IsExecuteContext())
-							{
-								PinsToFollow.Add(NextExecutePin);
-							}
-						}
-					}
-
-					for (const URigVMPin* NextExecutePin : PinsToFollow)
-					{
-						const TArray<URigVMPin*> TargetPins = NextExecutePin->GetLinkedTargetPins();
-						if (TargetPins.IsEmpty())
-						{
-							continue;
-						}
-						return TargetPins[0]->GetNode();
-					}
+}
 
-					return FindNextNodeOnCompletedPin(SourceNode,VisitedNodes);
-				}
-			}
+bool IRigVMEditorAssetInterface::MarkNodeForPreviewHere(URigVMNode* InNode, int32 InInstructionIndex, bool bRequestHyperLink)
+{
+	check(InNode);
+	
+	TArray<URigVMNode*> IncludedNodes;
+	TSet<URigVMNode*> VisitedNodes;
 
-			return nullptr;
-		}
+	IncludedNodes.Add(InNode);
 
-		static URigVMNode* FindNextNode(URigVMNode* Node, TSet<URigVMNode*>& VisitedNodes)
+	for (int32 Index = 0; Index < IncludedNodes.Num(); Index++)
+	{
+		VisitedNodes.Add(IncludedNodes[Index]);
+		
+		const TArray<URigVMNode*> SourceNodes = IncludedNodes[Index]->GetLinkedSourceNodes();
+		for (URigVMNode* SourceNode : SourceNodes)
 		{
-			if (!Node->IsMutable())
-			{
-				return nullptr;
-			}
-
-			for (const URigVMPin* Pin : Node->GetPins())
+			if (VisitedNodes.Contains(SourceNode))
 			{
-				if (Pin->GetDirection() == ERigVMPinDirection::Output || Pin->GetDirection() == ERigVMPinDirection::IO)
-				{
-					if (!Pin->IsExecuteContext())
-					{
-						continue;
-					}
-
-					const TArray<URigVMPin*> TargetPins = Pin->GetLinkedTargetPins();
-					if (TargetPins.IsEmpty())
-					{
-						continue;
-					}
-
-					if (!TargetPins[0]->IsExecuteContext())
-					{
-						continue;
-					}
-
-					return TargetPins[0]->GetNode();
-				}
+				continue;
 			}
-
-			return FindNextNodeOnCompletedPin(Node, VisitedNodes);
+			IncludedNodes.AddUnique(SourceNode);
 		}
-	};
+	}
 
-	TSet<URigVMNode*> VisitedNodes;
-	URigVMNode* NextNode = Local::FindNextNode(Node, VisitedNodes);
+	InNode->SetHasEarlyExitMarker(true);
 
-	while (NextNode && NextNode->IsA<URigVMRerouteNode>())
-	{
-		const TArray<URigVMNode*> TargetNodes = NextNode->GetLinkedTargetNodes();
-		if (TargetNodes.IsEmpty())
-		{
-			break;
-		}
-		NextNode = TargetNodes[0];
-	}
-	
-	if (!NextNode)
+	TArray<URigVMNode*> AllNodes = InNode->GetGraph()->GetNodes();
+	for (URigVMNode* Node : AllNodes)
 	{
-		LastPreviewHereNodes.Emplace(Node, CurrentInstructionIndex);
-		return;
+		Node->SetHasEarlyExitMarker(Node == InNode);
+		Node->SetIsExcludedByEarlyExit(!VisitedNodes.Contains(Node));
 	}
 
-	if (NextNode->IsA<URigVMFunctionInterfaceNode>())
-	{
-		int32 InstructionIndex = CurrentInstructionIndex;
-		NextNode = nullptr;
-		while (NextNode == nullptr && (++InstructionIndex) < VM->GetByteCode().GetNumInstructions())
-		{
-			NextNode = Cast<URigVMNode>(VM->GetByteCode().GetSubjectForInstruction(InstructionIndex));
-			if (Cast<URigVMFunctionInterfaceNode>(NextNode))
-			{
-				NextNode = nullptr;
-			}
-		}
+	GetRigGraphDisplaySettings().bShowNodeRunCounts = true;
 
-		if (NextNode == nullptr)
-		{
-			return;
-		}
-	}
+	LastPreviewHereNodes.Emplace(InNode, InInstructionIndex);
 
-	bool bFoundNextNode = false;
-	while (!bFoundNextNode && (++CurrentInstructionIndex) < VM->GetByteCode().GetNumInstructions())
+	if (bRequestHyperLink)
 	{
-		const TArray<TWeakObjectPtr<UObject>>* Callstack = VM->GetByteCode().GetCallstackForInstruction(CurrentInstructionIndex);
-		if (Callstack)
+		if (FRigVMClient* Client = GetRigVMClient())
 		{
-			for (const TWeakObjectPtr<UObject>& WeakObjectPtr : *Callstack)
+			if (URigVMController* Controller = Client->GetOrCreateController(InNode->GetGraph()))
 			{
-				if (UObject* WeakObject = WeakObjectPtr.Get())
-				{
-					if (WeakObject == NextNode)
-					{
-						bFoundNextNode = true;
-						break;
-					}
-				}
+				Controller->RequestJumpToHyperLink(InNode);
 			}
 		}
 	}
 
-	SetEarlyExitInstruction(NextNode, CurrentInstructionIndex,  true);
-}
-
-bool IRigVMEditorAssetInterface::CanPreviewHereStepForward() const
-{
-	URigVMHost* DebuggedHost = Cast<URigVMHost>(GetObjectBeingDebugged());
-	if (!DebuggedHost)
-	{
-		return false;
-	}
-
-	return DebuggedHost->GetDebugInfo().GetInstructionForEarlyExit() != INDEX_NONE;
-
+	return true;
 }
 
 int32 IRigVMEditorAssetInterface::GetPreviewNodeInstructionIndexFromSelection(const URigVMGraph* InGraph) const
@@ -2089,6 +1919,109 @@ void IRigVMEditorAssetInterface::DecrementVMRecompileBracket()
 	}
 }
 
+void IRigVMEditorAssetInterface::SetEditingLocked(bool InEditingLocked)
+{
+	if (!CanSetEditingLocked())
+	{
+		return;
+	}
+	
+	if (IsEditingLocked() == InEditingLocked)
+	{
+		return;
+	}
+
+	URigVMHost* Host = GetEditorHost();
+	if (!Host)
+	{
+		return;
+	}
+
+#if WITH_EDITOR
+	Host->EnableNativizedVM(InEditingLocked);
+#endif
+	Host->Initialize(true);
+}
+
+bool IRigVMEditorAssetInterface::IsEditingLocked() const
+{
+	if (!CanSetEditingLocked())
+	{
+		return false;
+	}
+	
+	URigVMHost* Host = const_cast<IRigVMEditorAssetInterface*>(this)->GetEditorHost();
+	if (!Host)
+	{
+		return false;
+	}
+
+	if (Host->IsPlayingRewindDebugTrace())
+	{
+		return true;
+	}
+
+	if (const URigVM* VM = Host->GetVM())
+	{
+		return VM->IsNativized();
+	}
+	
+	return false;
+}
+
+bool IRigVMEditorAssetInterface::CanSetEditingLocked() const
+{
+	URigVMHost* Host = const_cast<IRigVMEditorAssetInterface*>(this)->GetEditorHost();
+	if (!Host)
+	{
+		return false;
+	}
+
+	if (Host->IsPlayingRewindDebugTrace())
+	{
+		return false;
+	}
+
+	return Host->CanSwapVMToNativized();
+}
+
+bool IRigVMEditorAssetInterface::IsEditingLockedVisible() const
+{
+	return IsEditingLocked() || CanSetEditingLocked();
+}
+
+FText IRigVMEditorAssetInterface::GetEditingLockedLabel() const
+{
+	URigVMHost* Host = const_cast<IRigVMEditorAssetInterface*>(this)->GetEditorHost();
+	if (!Host)
+	{
+		return FText();
+	}
+
+	if (Host->IsPlayingRewindDebugTrace())
+	{
+		return LOCTEXT("EditingLockedRewindDebuggerActive", "Rewind Debugger");
+	}
+
+	return LOCTEXT("EditingLockedNativized", "Nativized");
+}
+
+FText IRigVMEditorAssetInterface::GetEditingLockedTooltip() const
+{
+	URigVMHost* Host = const_cast<IRigVMEditorAssetInterface*>(this)->GetEditorHost();
+	if (!Host)
+	{
+		return FText();
+	}
+
+	if (Host->IsPlayingRewindDebugTrace())
+	{
+		return LOCTEXT("EditingLockedRewindDebuggerActiveTooltip", "This asset is being driven by Rewind Debugger.");
+	}
+
+	return LOCTEXT("EditingLockedNativizedTooltip", "This asset is using a nativized VM if on.");
+}
+
 void IRigVMEditorAssetInterface::RefreshAllModels(ERigVMLoadType InLoadType)
 {
 	const bool bEnablePostLoadHashing = CVarRigVMEnablePostLoadHashing->GetBool();
diff --git a/Engine/Plugins/Runtime/RigVM/Source/RigVMDeveloper/Public/RigVMCompiler/RigVMCodeGenerator.h b/Engine/Plugins/Runtime/RigVM/Source/RigVMDeveloper/Public/RigVMCompiler/RigVMCodeGenerator.h
--- a/Engine/Plugins/Runtime/RigVM/Source/RigVMDeveloper/Public/RigVMCompiler/RigVMCodeGenerator.h
+++ b/Engine/Plugins/Runtime/RigVM/Source/RigVMDeveloper/Public/RigVMCompiler/RigVMCodeGenerator.h
@@ -2,275 +2,109 @@
 
 #pragma once
 
-#include "RigVMCore/RigVM.h"
-#include "RigVMModel/RigVMGraph.h"
-
-#include "UObject/StrongObjectPtr.h"
+#include "RigVMCodeConverter.h"
 
 #define UE_API RIGVMDEVELOPER_API
 
-class URigVMGraph;
+/**
+ * Supported target languages for code generation.
+ */
+namespace ERigVMCodeLanguage
+{
+	enum Type : uint8
+	{
+		/** Generate C++ source files (.h and .cpp). */
+		CPlusPlus,
 
+		/** Generate Verse source files (not yet implemented). */
+		Verse
+	};
+};
 
+/**
+ * Discovers and manages code generation templates for a target language.
+ *
+ * FRigVMCodeGenerator is responsible for:
+ * - Discovering template files from the Content/CodeGeneration/{Language}/ directory
+ * - Providing access to templates by name
+ * - Orchestrating code generation by combining a converter with templates
+ *
+ * Usage:
+ * @code
+ *   // Create generator for C++
+ *   FRigVMCodeGenerator Generator(ERigVMCodeLanguage::CPlusPlus);
+ *
+ *   // Create converter from a RigVM asset
+ *   auto Converter = MakeShared<FRigVMCodeConverter>(EditorAsset, Settings);
+ *
+ *   // Generate all output files
+ *   TArray<TSharedPtr<FRigVMCodeOutput>> Outputs = Generator.GenerateAll(Converter);
+ *
+ *   // Or generate a specific template
+ *   auto Output = Generator.Generate(Converter, TEXT("AssetNameVM.h.template"));
+ * @endcode
+ *
+ * @see FRigVMCodeConverter for bytecode parsing and JSON generation
+ * @see FRigVMCodeTemplate for template file representation
+ */
 struct FRigVMCodeGenerator
 {
 public:
 
-	FRigVMCodeGenerator(const FString& InClassName, const FString& InModuleName,
-		URigVMGraph* InModelToNativize, URigVM* InVMToNativize, FRigVMExtendedExecuteContext& InVMContext,
-		const UScriptStruct* PublicContextStruct, TMap<FString,FRigVMOperand> InPinToOperandMap,
-		int32 InMaxInstructionsPerFunction = 100)
-	{
-		ParseVM(InClassName, InModuleName, InModelToNativize, InVMToNativize, InVMContext, PublicContextStruct, InPinToOperandMap, InMaxInstructionsPerFunction);
-	}
+	/**
+	 * Constructs a generator and discovers all templates for the specified language.
+	 *
+	 * Templates are loaded from: Content/CodeGeneration/{LanguageName}/{FileName}.template
+	 * For C++, this is: Content/CodeGeneration/CPlusPlus/{FileName}.template
+	 *
+	 * @param InLanguage The target language for code generation.
+	 */
+	UE_API FRigVMCodeGenerator(ERigVMCodeLanguage::Type InLanguage);
+
+	/** Returns the number of discovered templates. */
+	UE_API int32 NumTemplates() const;
+
+	/**
+	 * Retrieves a template by name.
+	 *
+	 * @param InName Template name (e.g., "AssetNameVM.h.template").
+	 * @return The template, or an invalid shared pointer if not found.
+	 */
+	UE_API const TSharedPtr<FRigVMCodeTemplate>& GetTemplate(const FString& InName) const;
+
+	/**
+	 * Generates output from a single template.
+	 *
+	 * @param InConverter The converter containing parsed bytecode as JSON.
+	 * @param InTemplateName Name of the template to render.
+	 * @return The rendered output, or nullptr if the template was not found.
+	 */
+	UE_API TSharedPtr<FRigVMCodeOutput> Generate(const TSharedPtr<FRigVMCodeConverter>& InConverter, const FString& InTemplateName) const;
+
+	/**
+	 * Generates output from a template (static version).
+	 *
+	 * @param InConverter The converter containing parsed bytecode as JSON.
+	 * @param InCodeTemplate The template to render.
+	 * @return The rendered output.
+	 */
+	UE_API static TSharedPtr<FRigVMCodeOutput> Generate(const TSharedPtr<FRigVMCodeConverter>& InConverter, const TSharedPtr<FRigVMCodeTemplate>& InCodeTemplate);
+
+	/**
+	 * Generates output from all discovered templates.
+	 *
+	 * @param InConverter The converter containing parsed bytecode as JSON.
+	 * @return Array of rendered outputs, one per template.
+	 */
+	UE_API TArray<TSharedPtr<FRigVMCodeOutput>> GenerateAll(const TSharedPtr<FRigVMCodeConverter>& InConverter) const;
 
-	UE_API FString DumpIncludes(bool bLog = false);
-	UE_DEPRECATED(5.4, "Please, use DumpExternalVariables with ExtendedExecuteContext parameter")
-	FString DumpExternalVariables(bool bForHeader, bool bLog = false) { return FString(); }
-	UE_API FString DumpExternalVariables(const FRigVMExtendedExecuteContext& Context, bool bForHeader, bool bLog = false);
-	UE_API FString DumpEntries(bool bForHeader, bool bLog = false);
-	UE_API FString DumpBlockNames(bool bForHeader, bool bLog = false);
-	UE_DEPRECATED(5.4, "Please, use DumpProperties with ExtendedExecuteContext parameter")
-	FString DumpProperties(bool bForHeader, int32 InInstructionGroup, bool bLog = false) { return FString(); }
-	UE_API FString DumpProperties(const FRigVMExtendedExecuteContext& Context, bool bForHeader, int32 InInstructionGroup, bool bLog = false);
-	UE_API FString DumpDispatches(bool bLog = false);
-	UE_API FString DumpRequiredUProperties(bool bLog = false);
-	UE_API FString DumpInitialize(bool bLog = false);
-	UE_DEPRECATED(5.4, "Please, use DumpInstructions with ExtendedExecuteContext parameter")
-	FString DumpInstructions(int32 InInstructionGroup, bool bLog = false) { return FString(); }
-	UE_API FString DumpInstructions(const FRigVMExtendedExecuteContext& Context, int32 InInstructionGroup, bool bLog = false);
-	UE_DEPRECATED(5.4, "Please, use DumpHeader with ExtendedExecuteContext parameter")
-	FString DumpHeader(bool bLog = false) { return FString(); }
-	UE_API FString DumpHeader(const FRigVMExtendedExecuteContext& Context, bool bLog = false);
-	UE_DEPRECATED(5.4, "Please, use DumpSource with ExtendedExecuteContext parameter")
-	FString DumpSource(bool bLog = false) { return FString(); }
-	UE_API FString DumpSource(const FRigVMExtendedExecuteContext& Context, bool bLog = false);
-	UE_API FString DumpLines(const TArray<FString>& InLines, bool bLog = false);
-	
 private:
 
-	typedef TArray<FString> FStringArray;
-	typedef TMap<FString, FString> FStringMap;
-	typedef TArray<FRigVMPropertyDescription> FRigVMPropertyDescriptionArray;
-	typedef TTuple<FString, FString> FMappedType;
-	typedef TFunction<FString(const FString&)> TStructConstGenerator;
-
-	enum ERigVMNativizedPropertyType
-	{
-		Literal,
-		Work,
-		Sliced,
-		Invalid
-	};
-
-	struct FPropertyInfo
-	{
-		FRigVMPropertyDescription Description;
-		int32 MemoryPropertyIndex;
-		ERigVMNativizedPropertyType PropertyType;
-		TArray<int32> Groups;
-	};
-
-	struct FInstructionGroup
-	{
-		FString Entry;
-		int32 First;
-		int32 Last;
-		int32 Depth;
-		int32 ParentGroup;
-		TArray<int32> ChildGroups;
-		TArray<int32> RequiredLabels;
-
-		FInstructionGroup()
-			: Entry()
-			, First(INDEX_NONE)
-			, Last(INDEX_NONE)
-			, Depth(INDEX_NONE)
-			, ParentGroup(INDEX_NONE)
-		{}
-
-		FInstructionGroup(const FInstructionGroup& InOther)
-			: Entry(InOther.Entry)
-			, First(InOther.First)
-			, Last(InOther.Last)
-			, Depth(InOther.Depth)
-			, ParentGroup(InOther.ParentGroup)
-			, ChildGroups(InOther.ChildGroups)
-			, RequiredLabels(InOther.RequiredLabels)
-		{}
-	};
-
-	UE_API void Reset();
-	UE_API void ParseVM(const FString& InClassName, const FString& InModuleName,
-		URigVMGraph* InModelToNativize, URigVM* InVMToNativize, FRigVMExtendedExecuteContext& InVMContext,
-		const UScriptStruct* PublicContextStruct, TMap<FString,FRigVMOperand> InPinToOperandMap,
-		int32 InMaxInstructionsPerFunction);
-	UE_API void ParseInclude(UStruct* InDependency, const FName& InMethodName = NAME_None);
-	UE_API void ParseRequiredUProperties(const FRigVMExtendedExecuteContext& Context);
-	UE_API void ParseMemory(const FRigVMExtendedExecuteContext& Context, FRigVMMemoryStorageStruct* InMemory);
-	UE_API void ParseProperty(const FRigVMExtendedExecuteContext& Context, ERigVMMemoryType InMemoryType, const FProperty* InProperty, FRigVMMemoryStorageStruct* InMemory);
-	UE_API void ParseInstructionGroups(const FRigVMExtendedExecuteContext& Context);
-	UE_API FString DumpInstructions(const FRigVMExtendedExecuteContext& Context, const FString& InPrefix, int32 InFirstInstruction, int32 InLastInstruction, const FInstructionGroup& InGroup, bool bLog = false);
-	UE_API FString DumpInstructions(const FRigVMExtendedExecuteContext& Context, const FString& InPrefix, const TArray<int32> InInstructionIndices, const FInstructionGroup& InGroup, bool bLog = false);
-	UE_API FString GetOperandName(const FRigVMExtendedExecuteContext& Context, const FRigVMOperand& InOperand, bool bPerSlice = true, bool bAsInput = true) const;
-	UE_API FString GetOperandCPPType(const FRigVMExtendedExecuteContext& Context, const FRigVMOperand& InOperand) const;
-	UE_API FString GetOperandCPPBaseType(const FRigVMExtendedExecuteContext& Context, const FRigVMOperand& InOperand) const;
-	static UE_API FString SanitizeName(const FString& InName, const FString& CPPType = FString()); 
-	static UE_API FString SanitizeValue(const FString& InValue, const FString& CPPType, const UObject* CPPTypeObject); 
-	UE_API FRigVMPropertyDescription GetPropertyDescForOperand(const FRigVMExtendedExecuteContext& Context, const FRigVMOperand& InOperand) const;
-	UE_API FRigVMPropertyDescription GetPropertyForOperand(const FRigVMExtendedExecuteContext& Context, const FRigVMOperand& InOperand) const;
-	UE_API const FRigVMPropertyPathDescription& GetPropertyPathForOperand(const FRigVMExtendedExecuteContext& Context, const FRigVMOperand& InOperand) const;
-	UE_API int32 GetPropertyIndex(const FRigVMPropertyDescription& InProperty) const;
-	UE_API int32 GetPropertyIndex(const FRigVMExtendedExecuteContext& Context, const FRigVMOperand& InOperand) const;
-	UE_API ERigVMNativizedPropertyType GetPropertyType(const FRigVMPropertyDescription& InProperty) const;
-	UE_API ERigVMNativizedPropertyType GetPropertyType(const FRigVMExtendedExecuteContext& Context, const FRigVMOperand& InOperand) const;
-	UE_API void CheckOperand(const FRigVMExtendedExecuteContext& Context, const FRigVMOperand& InOperand) const;
-	UE_API FString GetMappedType(const FString& InCPPType) const;
-	UE_API const FString& GetMappedTypeSuffix(const FString& InCPPType) const;
-	UE_API FString GetMappedArrayTypeName(const FString InBaseElementType) const;
-	UE_API const FInstructionGroup& GetGroup(int32 InGroupIndex) const;
-	UE_API bool IsInstructionPartOfGroup(int32 InInstructionIndex, int32 InGroupIndex, bool bIncludeChildGroups) const;
-	UE_API bool IsPropertyPartOfGroup(int32 InPropertyIndex, int32 InGroupIndex) const;
-	UE_API FString GetEntryParameters() const;
-	static UE_API TArray<int32> GetInstructionIndicesFromRange(int32 First, int32 Last);
-
-	static FString FormatArgs(const TCHAR* InFormatString, const FStringFormatOrderedArguments& InArgs)
-	{
-		return FString::Format(InFormatString, InArgs);
-	}
-
-	template<typename TypeA>
-	static FString Format(const TCHAR* InFormatString, const TypeA& InArgA)
-	{
-		return FormatArgs(InFormatString, {
-			FStringFormatArg(InArgA)
-		});
-	}
-
-	template<typename TypeA, typename TypeB>
-	static FString Format(const TCHAR* InFormatString, const TypeA& InArgA, const TypeB& InArgB)
-	{
-		return FormatArgs(InFormatString, {
-			FStringFormatArg(InArgA),
-			FStringFormatArg(InArgB)
-		});
-	}
-
-	template<typename TypeA, typename TypeB, typename TypeC>
-	static FString Format(const TCHAR* InFormatString, const TypeA& InArgA, const TypeB& InArgB, const TypeC& InArgC)
-	{
-		return FormatArgs(InFormatString, {
-			FStringFormatArg(InArgA),
-			FStringFormatArg(InArgB),
-			FStringFormatArg(InArgC)
-		});
-	}
-
-	template<typename TypeA, typename TypeB, typename TypeC, typename TypeD>
-	static FString Format(const TCHAR* InFormatString,
-		const TypeA& InArgA,
-		const TypeB& InArgB,
-		const TypeC& InArgC,
-		const TypeD& InArgD
-	)
-	{
-		return FormatArgs(InFormatString, {
-			FStringFormatArg(InArgA),
-			FStringFormatArg(InArgB),
-			FStringFormatArg(InArgC),
-			FStringFormatArg(InArgD)
-		});
-	}
-
-	template<typename TypeA, typename TypeB, typename TypeC, typename TypeD, typename TypeE>
-	static FString Format(const TCHAR* InFormatString,
-		const TypeA& InArgA,
-		const TypeB& InArgB,
-		const TypeC& InArgC,
-		const TypeD& InArgD,
-		const TypeE& InArgE
-	)
-	{
-		return FormatArgs(InFormatString, {
-			FStringFormatArg(InArgA),
-			FStringFormatArg(InArgB),
-			FStringFormatArg(InArgC),
-			FStringFormatArg(InArgD),
-			FStringFormatArg(InArgE)
-		});
-	}
-
-	template<typename TypeA, typename TypeB, typename TypeC, typename TypeD, typename TypeE, typename TypeF>
-	static FString Format(const TCHAR* InFormatString,
-		const TypeA& InArgA,
-		const TypeB& InArgB,
-		const TypeC& InArgC,
-		const TypeD& InArgD,
-		const TypeE& InArgE,
-		const TypeF& InArgF
-	)
-	{
-		return FormatArgs(InFormatString, {
-			FStringFormatArg(InArgA),
-			FStringFormatArg(InArgB),
-			FStringFormatArg(InArgC),
-			FStringFormatArg(InArgD),
-			FStringFormatArg(InArgE),
-			FStringFormatArg(InArgF)
-		});
-	}
-
-	template<typename TypeA, typename TypeB, typename TypeC, typename TypeD, typename TypeE, typename TypeF, typename TypeG>
-	static FString Format(const TCHAR* InFormatString,
-		const TypeA& InArgA,
-		const TypeB& InArgB,
-		const TypeC& InArgC,
-		const TypeD& InArgD,
-		const TypeE& InArgE,
-		const TypeF& InArgF,
-		const TypeG& InArgG
-	)
-	{
-		return FormatArgs(InFormatString, {
-			FStringFormatArg(InArgA),
-			FStringFormatArg(InArgB),
-			FStringFormatArg(InArgC),
-			FStringFormatArg(InArgD),
-			FStringFormatArg(InArgE),
-			FStringFormatArg(InArgF),
-			FStringFormatArg(InArgG)
-		});
-	}
-
-	TStrongObjectPtr<URigVMGraph> Model;
-	TStrongObjectPtr<URigVM> VM;
-	TMap<FString,FRigVMOperand> PinToOperandMap;
-	TMap<FRigVMOperand, FString> OperandToPinMap;
-	int32 MaxInstructionsPerFunction;
-	FString ClassName;
-	FString ModuleName;
-	FString ExecuteContextType;
-	
-	FStringArray Libraries;
-	FStringArray Includes;
-
-	struct FRigVMDispatchInfo
-	{
-		FString Name;
-		const FRigVMFunction* Function;
-		FRigVMDispatchContext Context;
-	};
-	
-	TMap<FString, FRigVMDispatchInfo> Dispatches;
-	TMap<TRigVMTypeIndex,TTuple<FString,FString>> RequiredUProperties;
-	TArray<FInstructionGroup> InstructionGroups;
-	TMap<FString, FMappedType> MappedCPPTypes;
-	TArray<FPropertyInfo> Properties;
-	TMap<FName, int32> PropertyNameToIndex;
-	TMap<FName, FName> MappedPropertyNames;
-	TMap<FString, FString> OverriddenOperatorNames;
+	/** The target language for this generator. */
+	ERigVMCodeLanguage::Type Language;
 
-	static UE_API const TMap<FName, TStructConstGenerator>& GetStructConstGenerators();
+	/** All discovered templates for the target language. */
+	TArray<TSharedPtr<FRigVMCodeTemplate>> CodeTemplates;
 };
 
 #undef UE_API
diff --git a/Engine/Plugins/Runtime/RigVM/Source/RigVMDeveloper/Public/RigVMEditorAsset.h b/Engine/Plugins/Runtime/RigVM/Source/RigVMDeveloper/Public/RigVMEditorAsset.h
--- a/Engine/Plugins/Runtime/RigVM/Source/RigVMDeveloper/Public/RigVMEditorAsset.h
+++ b/Engine/Plugins/Runtime/RigVM/Source/RigVMDeveloper/Public/RigVMEditorAsset.h
@@ -253,6 +253,13 @@ public:
 	UE_API virtual void IncrementVMRecompileBracket();
 	UE_API virtual void DecrementVMRecompileBracket();
 
+	UE_API virtual void SetEditingLocked(bool InEditingLocked );
+	UE_API virtual bool IsEditingLocked() const;
+	UE_API virtual bool CanSetEditingLocked() const;
+	UE_API virtual bool IsEditingLockedVisible() const;
+	UE_API virtual FText GetEditingLockedLabel() const;
+	UE_API virtual FText GetEditingLockedTooltip() const;
+
 	UE_API virtual void SetObjectBeingDebugged(UObject* NewObject);
 	virtual FOnRigVMSetObjectBeingDebugged& OnSetObjectBeingDebugged()  { return SetObjectBeingDebuggedEvent; }
 
@@ -308,6 +315,11 @@ public:
 	 */
 	UE_API bool CanPreviewHereStepForward() const;
 
+	/**
+	 * Returns true if we can step forward in a preview here session
+	 */
+	UE_API bool MarkNodeForPreviewHere(URigVMNode* InNode, int32 InInstructionIndex, bool bRequestHyperLink);
+
 	/**
 	 * Returns the instruction index of the currently selected node
 	 */
diff --git a/Engine/Plugins/Runtime/RigVM/Source/RigVMDeveloper/RigVMDeveloper.Build.cs b/Engine/Plugins/Runtime/RigVM/Source/RigVMDeveloper/RigVMDeveloper.Build.cs
--- a/Engine/Plugins/Runtime/RigVM/Source/RigVMDeveloper/RigVMDeveloper.Build.cs
+++ b/Engine/Plugins/Runtime/RigVM/Source/RigVMDeveloper/RigVMDeveloper.Build.cs
@@ -1,6 +1,7 @@
 // Copyright Epic Games, Inc. All Rights Reserved.
 
 using UnrealBuildTool;
+using System.IO;
 
 public class RigVMDeveloper : ModuleRules
 {
@@ -19,7 +20,8 @@ public class RigVMDeveloper : ModuleRules
                 "Projects",
                 "Json",
                 "Scribble",
-                "AssetTools"
+                "AssetTools",
+                "AnimationCore"
             }
         );
 
@@ -49,5 +51,7 @@ public class RigVMDeveloper : ModuleRules
             PrivateIncludePathModuleNames.Add("RigVMEditor");
             DynamicallyLoadedModuleNames.Add("RigVMEditor");
         }
+        PrivateIncludePaths.Add(Path.Combine(EngineDirectory,"Source/ThirdParty/inja/inja-3.5.0/single_include"));
+        PrivateIncludePaths.Add(Path.Combine(EngineDirectory,"Source/ThirdParty/inja/inja-3.5.0/third_party/include"));
     }
 }
diff --git a/Engine/Plugins/Runtime/RigVM/Source/RigVMEditor/Private/Editor/RigVMEditor.cpp b/Engine/Plugins/Runtime/RigVM/Source/RigVMEditor/Private/Editor/RigVMEditor.cpp
--- a/Engine/Plugins/Runtime/RigVM/Source/RigVMEditor/Private/Editor/RigVMEditor.cpp
+++ b/Engine/Plugins/Runtime/RigVM/Source/RigVMEditor/Private/Editor/RigVMEditor.cpp
@@ -1290,15 +1290,10 @@ bool FRigVMEditorBase::IsEditable(UEdGraph* InGraph) const
 				return false;
 			}
 		}
-
-		URigVMHost* RigVMHost = GetRigVMHost();
-		if(RigVMHost && RigVMHost->GetVM())
-		{
-			const bool bIsReadOnly = RigVMHost->GetVM()->IsNativized();
-			const bool bIsEditable = !bIsReadOnly;
-			InGraph->bEditable = bIsEditable ? 1 : 0;
-			return bIsEditable;
-		}
+		const bool bIsReadOnly = RigVMBlueprint->IsEditingLocked();
+		const bool bIsEditable = !bIsReadOnly;
+		InGraph->bEditable = bIsEditable;
+		return bIsEditable;
 	}
 
 	return IsEditableImpl(InGraph);
@@ -1469,6 +1464,12 @@ void FRigVMEditorBase::BindCommands()
 		FIsActionChecked::CreateSP(ThisRef, &FRigVMEditorBase::CanAutoCompileGraph),
 		FIsActionChecked::CreateSP(ThisRef, &FRigVMEditorBase::IsAutoCompileGraphOn));
 
+	ToolkitCommands->MapAction(
+		EditorCommands.LockGraph,
+		FExecuteAction::CreateSP(ThisRef, &FRigVMEditorBase::ToggleLockGraph), 
+		FIsActionChecked::CreateSP(ThisRef, &FRigVMEditorBase::CanLockGraph),
+		FIsActionChecked::CreateSP(ThisRef, &FRigVMEditorBase::IsLockGraphOn));
+
 	ToolkitCommands->MapAction(
 		EditorCommands.ToggleEventQueue,
 		FExecuteAction::CreateSP(ThisRef, &FRigVMEditorBase::ToggleEventQueue),
@@ -1550,6 +1551,7 @@ void FRigVMEditorBase::UnbindCommands()
 	FRigVMEditorCommands& EditorCommands = FRigVMEditorCommands::Get();
 
 	ToolkitCommands->UnmapAction(EditorCommands.AutoCompileGraph);
+	ToolkitCommands->UnmapAction(EditorCommands.LockGraph);
 	ToolkitCommands->UnmapAction(EditorCommands.ToggleEventQueue);
 	ToolkitCommands->UnmapAction(EditorCommands.FrameSelection);
 	ToolkitCommands->UnmapAction(EditorCommands.SwapFunctionWithinAsset);
@@ -1581,6 +1583,33 @@ bool FRigVMEditorBase::IsAutoCompileGraphOn() const
 	return false;
 }
 
+void FRigVMEditorBase::ToggleLockGraph()
+{
+	if (FRigVMEditorAssetInterfacePtr AssetInterface = GetRigVMAssetInterface())
+	{
+		AssetInterface->SetEditingLocked(!AssetInterface->IsEditingLocked());
+		AssetInterface->RequestAutoVMRecompilation();
+	}
+}
+
+bool FRigVMEditorBase::IsLockGraphOn() const
+{
+	if (FRigVMEditorAssetInterfacePtr AssetInterface = GetRigVMAssetInterface())
+	{
+		return AssetInterface->IsEditingLocked();
+	}
+	return false;
+}
+
+bool FRigVMEditorBase::CanLockGraph() const
+{
+	if (FRigVMEditorAssetInterfacePtr AssetInterface = GetRigVMAssetInterface())
+	{
+		return AssetInterface->CanSetEditingLocked();
+	}
+	return false;
+}
+
 void FRigVMEditorBase::ToggleEventQueue()
 {
 	SetEventQueue(LastEventQueue);
@@ -1773,17 +1802,25 @@ FGraphAppearanceInfo FRigVMEditorBase::GetGraphAppearance(UEdGraph* InGraph) con
 	{
 		AppearanceInfo.CornerText = LOCTEXT("AppearanceCornerText_RigVMEditor", "RigVM");
 
+		if (RigVMBlueprint->IsEditingLocked())
+		{
+			AppearanceInfo.InstructionFade = 1;
+			AppearanceInfo.InstructionText = FText::FromString(TEXT("This graph is locked."));
+		}
+
 		if(URigVMHost* RigVMHost = GetRigVMHost())
 		{
-			if(RigVMHost->GetVM() && RigVMHost->GetVM()->IsNativized())
+			if (RigVMHost->IsPlayingRewindDebugTrace())
 			{
-				if (UClass* NativizedClass = RigVMHost->GetVM()->GetNativizedClass())
-				{
-					AppearanceInfo.InstructionFade = 1;
-					AppearanceInfo.InstructionText = FText::FromString(
-						FString::Printf(TEXT("This graph runs a nativized VM (U%s)."), *NativizedClass->GetName())
-					);
-				}
+				AppearanceInfo.InstructionFade = 1;
+				AppearanceInfo.InstructionText = FText::FromString(TEXT("Rewind Debug in Progress."));
+			}
+			else if(RigVMHost->GetVM() && RigVMHost->GetVM()->IsNativized())
+			{
+				AppearanceInfo.InstructionFade = 1;
+				AppearanceInfo.InstructionText = FText::FromString(
+					FString::Printf(TEXT("Running a nativized VM (U%s)."), *RigVMHost->GetVM()->GetClass()->GetName())
+				);
 			}
 
 			if(RigVMHost->VMRuntimeSettings.bEnableProfiling)
@@ -3694,6 +3731,8 @@ void FRigVMEditorBase::FillToolbar(FToolBarBuilder& ToolbarBuilder, bool bEndSec
 	{
 		AddCompileWidget(ToolbarBuilder);
 		
+		AddLockEditingWidget(ToolbarBuilder);
+		
 		ToolbarBuilder.AddToolBarButton(
 			FRigVMEditorCommands::Get().ToggleEventQueue,
 			NAME_None, 
@@ -3736,6 +3775,45 @@ void FRigVMEditorBase::FillToolbar(FToolBarBuilder& ToolbarBuilder, bool bEndSec
 	}
 }
 
+void FRigVMEditorBase::AddLockEditingWidget(FToolBarBuilder& ToolbarBuilder)
+{
+	ToolbarBuilder.AddToolBarButton(
+		FRigVMEditorCommands::Get().LockGraph,
+		NAME_None,
+		TAttribute<FText>::Create([this]() -> FText
+		{
+			if (FRigVMEditorAssetInterfacePtr EditorAssetInterface = GetRigVMAssetInterface())
+			{
+				return EditorAssetInterface->GetEditingLockedLabel();
+			}
+			return FText();
+		}),
+		TAttribute<FText>::Create([this]() -> FText
+		{
+			if (FRigVMEditorAssetInterfacePtr EditorAssetInterface = GetRigVMAssetInterface())
+			{
+				return EditorAssetInterface->GetEditingLockedTooltip();
+			}
+			return FText();
+		}),
+		FSlateIcon(FAppStyle::GetAppStyleSetName(),
+		"Icons.Lock"),
+		NAME_None,
+		FNewMenuDelegate(),
+		TAttribute<EVisibility>::CreateLambda([this]() -> EVisibility
+		{
+			if (FRigVMEditorAssetInterfacePtr EditorAssetInterface = GetRigVMAssetInterface())
+			{
+				if (EditorAssetInterface->IsEditingLockedVisible())
+				{
+					return EVisibility::Visible;
+				}
+			}
+			return EVisibility::Collapsed;
+		})
+	);
+}
+
 void FRigVMEditorBase::UpdateStaleWatchedPins()
 {
 	FRigVMEditorAssetInterfacePtr RigVMBlueprint = GetRigVMAssetInterface();
diff --git a/Engine/Plugins/Runtime/RigVM/Source/RigVMEditor/Private/Editor/RigVMEditorCommands.cpp b/Engine/Plugins/Runtime/RigVM/Source/RigVMEditor/Private/Editor/RigVMEditorCommands.cpp
--- a/Engine/Plugins/Runtime/RigVM/Source/RigVMEditor/Private/Editor/RigVMEditorCommands.cpp
+++ b/Engine/Plugins/Runtime/RigVM/Source/RigVMEditor/Private/Editor/RigVMEditorCommands.cpp
@@ -24,8 +24,9 @@ void FRigVMEditorCommands::RegisterCommands()
 	UI_COMMAND(EditClassDefaults, "Defaults", "Edit the initial values of your class.", EUserInterfaceActionType::ToggleButton, FInputChord());
 
 	UI_COMMAND(DeleteItem, "Delete", "Deletes the selected items and removes their nodes from the graph.", EUserInterfaceActionType::Button, FInputChord(EKeys::Delete));
-	UI_COMMAND(ExecuteGraph, "Execute", "Execute the rig graph if On.", EUserInterfaceActionType::ToggleButton, FInputChord());
-	UI_COMMAND(AutoCompileGraph, "Auto Compile", "Auto-compile the rig graph if On.", EUserInterfaceActionType::ToggleButton, FInputChord());
+	UI_COMMAND(ExecuteGraph, "Execute", "Execute the VM if On.", EUserInterfaceActionType::ToggleButton, FInputChord());
+	UI_COMMAND(AutoCompileGraph, "Auto Compile", "Auto-compile the VM if On.", EUserInterfaceActionType::ToggleButton, FInputChord());
+	UI_COMMAND(LockGraph, "Nativized", "Runs a nativized backend if On.", EUserInterfaceActionType::ToggleButton, FInputChord());
 	UI_COMMAND(ToggleEventQueue, "Toggle Event", "Toggle between the current and last running event", EUserInterfaceActionType::Button, FInputChord());
 	UI_COMMAND(RunGraphOnce, "Run Once", "Only evaluate a single frame if On.", EUserInterfaceActionType::ToggleButton, FInputChord());
 	UI_COMMAND(FrameSelection, "Frame Selection", "Frames the selected nodes in the Graph View.", EUserInterfaceActionType::Button, FInputChord(EKeys::F));
diff --git a/Engine/Plugins/Runtime/RigVM/Source/RigVMEditor/Private/Editor/RigVMEditorCommands.h b/Engine/Plugins/Runtime/RigVM/Source/RigVMEditor/Private/Editor/RigVMEditorCommands.h
--- a/Engine/Plugins/Runtime/RigVM/Source/RigVMEditor/Private/Editor/RigVMEditorCommands.h
+++ b/Engine/Plugins/Runtime/RigVM/Source/RigVMEditor/Private/Editor/RigVMEditorCommands.h
@@ -43,6 +43,9 @@ public:
 	/** Toggle Auto Compilation in the Graph */
 	TSharedPtr< FUICommandInfo > AutoCompileGraph;
 
+	/** Toggle Locking of the Graph */
+	TSharedPtr< FUICommandInfo > LockGraph;
+
 	/** Toggle between this and the last event queue */
 	TSharedPtr< FUICommandInfo > ToggleEventQueue;
 
diff --git a/Engine/Plugins/Runtime/RigVM/Source/RigVMEditor/Private/Editor/RigVMExecutionStackCommands.cpp b/Engine/Plugins/Runtime/RigVM/Source/RigVMEditor/Private/Editor/RigVMExecutionStackCommands.cpp
--- a/Engine/Plugins/Runtime/RigVM/Source/RigVMEditor/Private/Editor/RigVMExecutionStackCommands.cpp
+++ b/Engine/Plugins/Runtime/RigVM/Source/RigVMEditor/Private/Editor/RigVMExecutionStackCommands.cpp
@@ -9,6 +9,8 @@ void FRigVMExecutionStackCommands::RegisterCommands()
 	UI_COMMAND(FocusOnSelection, "Focus On Selection", "Finds the selected operator's node in the graph.", EUserInterfaceActionType::Button, FInputChord());
 	UI_COMMAND(GoToInstruction, "Go To Instruction", "Looks for a specific instruction by index and brings it into focus.", EUserInterfaceActionType::Button, FInputChord(EKeys::G, EModifierKey::Control));
 	UI_COMMAND(SelectTargetInstructions, "Select Target Instruction(s)", "Looks up the target instructions and selects them.", EUserInterfaceActionType::Button, FInputChord(EKeys::T, EModifierKey::Control));
+	UI_COMMAND(ToggleEarlyExitInstruction, "Toggle Preview Here Instruction", "Sets or unsets the preview here early exit instruction.", EUserInterfaceActionType::Button, FInputChord(EKeys::F9));
+	UI_COMMAND(StepEarlyExitInstruction, "Step Preview Here Instruction", "Steps to the next early exist instruction.", EUserInterfaceActionType::Button, FInputChord(EKeys::F10));
 }
 
 #undef LOCTEXT_NAMESPACE
diff --git a/Engine/Plugins/Runtime/RigVM/Source/RigVMEditor/Private/Widgets/SRigVMExecutionStackView.cpp b/Engine/Plugins/Runtime/RigVM/Source/RigVMEditor/Private/Widgets/SRigVMExecutionStackView.cpp
--- a/Engine/Plugins/Runtime/RigVM/Source/RigVMEditor/Private/Widgets/SRigVMExecutionStackView.cpp
+++ b/Engine/Plugins/Runtime/RigVM/Source/RigVMEditor/Private/Widgets/SRigVMExecutionStackView.cpp
@@ -432,6 +432,17 @@ void SRigVMExecutionStackView::OnSelectionChanged(TSharedPtr<FRigStackEntry> Sel
 	UpdateTargetItemHighlighting();
 }
 
+FReply SRigVMExecutionStackView::OnKeyDown(const FGeometry& MyGeometry, const FKeyEvent& InKeyEvent)
+{
+	const FInputChord KeyEventAsInputChord = FInputChord(InKeyEvent.GetKey(), EModifierKey::FromBools(InKeyEvent.IsControlDown(), InKeyEvent.IsAltDown(), InKeyEvent.IsShiftDown(), InKeyEvent.IsCommandDown()));
+	if (CommandList->ProcessCommandBindings(InKeyEvent))
+	{
+		return FReply::Handled();
+	}
+
+	return SCompoundWidget::OnKeyDown(MyGeometry, InKeyEvent);
+}
+
 void SRigVMExecutionStackView::BindCommands()
 {
 	// create new command
@@ -439,6 +450,8 @@ void SRigVMExecutionStackView::BindCommands()
 	CommandList->MapAction(Commands.FocusOnSelection, FExecuteAction::CreateSP(this, &SRigVMExecutionStackView::HandleFocusOnSelectedGraphNode));
 	CommandList->MapAction(Commands.GoToInstruction, FExecuteAction::CreateSP(this, &SRigVMExecutionStackView::HandleGoToInstruction));
 	CommandList->MapAction(Commands.SelectTargetInstructions, FExecuteAction::CreateSP(this, &SRigVMExecutionStackView::HandleSelectTargetInstructions));
+	CommandList->MapAction(Commands.ToggleEarlyExitInstruction, FExecuteAction::CreateSP(this, &SRigVMExecutionStackView::HandleToggleEarlyExitInstruction));
+	CommandList->MapAction(Commands.StepEarlyExitInstruction, FExecuteAction::CreateSP(this, &SRigVMExecutionStackView::HandleStepEarlyExitInstruction));
 }
 
 TSharedRef<ITableRow> SRigVMExecutionStackView::MakeTableRowWidget(TSharedPtr<FRigStackEntry> InItem, const TSharedRef<STableViewBase>& OwnerTable, TWeakInterfacePtr<IRigVMEditorAssetInterface> InBlueprint)
@@ -463,6 +476,11 @@ void SRigVMExecutionStackView::PopulateStackView(URigVM* InVM, FRigVMExtendedExe
 	if (InVM && InVMContext)
 	{
 		const FRigVMInstructionArray Instructions = InVM->GetInstructions();
+		if (Instructions.Num() == 0)
+		{
+			return;
+		}
+		
 		const FRigVMByteCode& ByteCode = InVM->GetByteCode();
 
 		TArray<int32> CallableIndexPerInstruction;
@@ -1277,6 +1295,75 @@ void SRigVMExecutionStackView::HandleSelectTargetInstructions()
 	}
 }
 
+void SRigVMExecutionStackView::HandleToggleEarlyExitInstruction()
+{
+	const TArray<TSharedPtr<FRigStackEntry>> TargetItems = TreeView->GetSelectedItems();
+	if(TargetItems.IsEmpty())
+	{
+		return;
+	}
+
+	int32 InstructionIndex = TargetItems[0]->SubjectIndex;
+
+	TSharedPtr<IRigVMEditor> Editor = WeakEditor.Pin();
+	if (!Editor)
+	{
+		return;
+	}
+	FRigVMEditorAssetInterfacePtr RigBlueprint = Editor->GetRigVMAssetInterface();
+	if (!RigBlueprint)
+	{
+		return;
+	}
+
+	URigVMHost* RigVMHost = PreviouslyDebuggedHost.Get();
+	if (!RigVMHost)
+	{
+		return;
+	}
+
+	if (RigVMHost->GetDebugInfo().GetStepCondition().OriginInstruction == INDEX_NONE)
+	{
+		URigVM* VM = RigVMHost->GetVM();
+		check(VM);
+
+		const FRigVMInstructionArray Instructions = VM->GetByteCode().GetInstructions();
+		while (Instructions.IsValidIndex(InstructionIndex))
+		{
+			if (URigVMNode* Node = Cast<URigVMNode>(VM->GetByteCode().GetSubjectForInstruction(InstructionIndex)))
+			{
+				RigBlueprint->SetEarlyExitInstruction(Node, InstructionIndex, false);
+				return;
+			}
+			InstructionIndex++;
+		}
+	}
+	
+	RigBlueprint->ResetEarlyExitInstruction(true);
+}
+
+void SRigVMExecutionStackView::HandleStepEarlyExitInstruction()
+{
+	TSharedPtr<IRigVMEditor> Editor = WeakEditor.Pin();
+	if (!Editor)
+	{
+		return;
+	}
+	FRigVMEditorAssetInterfacePtr RigBlueprint = Editor->GetRigVMAssetInterface();
+	if (!RigBlueprint)
+	{
+		return;
+	}
+
+	URigVMHost* RigVMHost = PreviouslyDebuggedHost.Get();
+	if (!RigVMHost)
+	{
+		return;
+	}
+
+	RigVMHost->GetRigVMExtendedExecuteContext().StepInto();
+}
+
 TArray<TSharedPtr<FRigStackEntry>> SRigVMExecutionStackView::GetTargetItems(const TArray<TSharedPtr<FRigStackEntry>>& InItems) const
 {
 	TArray<TSharedPtr<FRigStackEntry>> TargetItems;
diff --git a/Engine/Plugins/Runtime/RigVM/Source/RigVMEditor/Public/Editor/RigVMEditor.h b/Engine/Plugins/Runtime/RigVM/Source/RigVMEditor/Public/Editor/RigVMEditor.h
--- a/Engine/Plugins/Runtime/RigVM/Source/RigVMEditor/Public/Editor/RigVMEditor.h
+++ b/Engine/Plugins/Runtime/RigVM/Source/RigVMEditor/Public/Editor/RigVMEditor.h
@@ -465,6 +465,9 @@ protected:
 	UE_API void ToggleAutoCompileGraph();
 	UE_API bool IsAutoCompileGraphOn() const;
 	bool CanAutoCompileGraph() const { return true; }
+	UE_API void ToggleLockGraph();
+	UE_API bool IsLockGraphOn() const;
+	bool CanLockGraph() const;
 	UE_API void ToggleEventQueue();
 	UE_API void ToggleRunGraphOnce();
 	UE_API void SetRunGraphOnce(bool InRunGraphOnceEnabled);
@@ -522,6 +525,7 @@ protected:
 	
 	/** Fill the toolbar with content */
 	UE_API virtual void FillToolbar(FToolBarBuilder& ToolbarBuilder, bool bEndSection = true);
+	UE_API virtual void AddLockEditingWidget(FToolBarBuilder& ToolbarBuilder);
 	
 	UE_API virtual TArray<FName> GetDefaultEventQueue() const;
 	UE_API TArray<FName> GetEventQueue() const;
diff --git a/Engine/Plugins/Runtime/RigVM/Source/RigVMEditor/Public/Editor/RigVMExecutionStackCommands.h b/Engine/Plugins/Runtime/RigVM/Source/RigVMEditor/Public/Editor/RigVMExecutionStackCommands.h
--- a/Engine/Plugins/Runtime/RigVM/Source/RigVMEditor/Public/Editor/RigVMExecutionStackCommands.h
+++ b/Engine/Plugins/Runtime/RigVM/Source/RigVMEditor/Public/Editor/RigVMExecutionStackCommands.h
@@ -26,6 +26,12 @@ public:
 	/** Selects the target instruction(s) of an instruction */
 	TSharedPtr< FUICommandInfo > SelectTargetInstructions;
 
+	/** Sets or unsets the preview here early exit instruction */
+	TSharedPtr< FUICommandInfo > ToggleEarlyExitInstruction;
+
+	/** Steps to the next early exist instruction */
+	TSharedPtr< FUICommandInfo > StepEarlyExitInstruction;
+
 	/**
 	 * Initialize commands
 	 */
diff --git a/Engine/Plugins/Runtime/RigVM/Source/RigVMEditor/Public/Widgets/SRigVMExecutionStackView.h b/Engine/Plugins/Runtime/RigVM/Source/RigVMEditor/Public/Widgets/SRigVMExecutionStackView.h
--- a/Engine/Plugins/Runtime/RigVM/Source/RigVMEditor/Public/Widgets/SRigVMExecutionStackView.h
+++ b/Engine/Plugins/Runtime/RigVM/Source/RigVMEditor/Public/Widgets/SRigVMExecutionStackView.h
@@ -142,6 +142,8 @@ public:
 	/** Set Selection Changed */
 	UE_API void OnSelectionChanged(TSharedPtr<FRigStackEntry> Selection, ESelectInfo::Type SelectInfo);
 
+	virtual FReply OnKeyDown(const FGeometry& MyGeometry, const FKeyEvent& InKeyEvent) override;
+
 	UE_API TSharedPtr< SWidget > CreateContextMenu();
 
 protected:
@@ -169,6 +171,12 @@ private:
 	/** Selects the target instructions for the current selection */
 	UE_API void HandleSelectTargetInstructions();
 
+	/** Sets or unsets the preview here early exit instruction */
+	UE_API void HandleToggleEarlyExitInstruction();
+	
+	/** Steps to the next early exist instruction */
+	UE_API void HandleStepEarlyExitInstruction();
+
 	UE_API TArray<TSharedPtr<FRigStackEntry>> GetTargetItems(const TArray<TSharedPtr<FRigStackEntry>>& InItems) const;
 
 	UE_API void UpdateTargetItemHighlighting();
