RigVM: Guard uses of tracing related methods with #if

#rnx

#Synced-CL 51941404

diff --git a/Engine/Plugins/Runtime/RigVM/Source/RigVMDeveloper/Private/RigVMEditorAsset.cpp b/Engine/Plugins/Runtime/RigVM/Source/RigVMDeveloper/Private/RigVMEditorAsset.cpp
--- a/Engine/Plugins/Runtime/RigVM/Source/RigVMDeveloper/Private/RigVMEditorAsset.cpp
+++ b/Engine/Plugins/Runtime/RigVM/Source/RigVMDeveloper/Private/RigVMEditorAsset.cpp
@@ -1956,11 +1956,13 @@ bool IRigVMEditorAssetInterface::IsEditingLocked() const
 		return false;
 	}
 
+#if RIGVM_TRACE_ENABLED
 	if (Host->IsPlayingRewindDebugTrace())
 	{
 		return true;
 	}
-
+#endif
+	
 	if (const URigVM* VM = Host->GetVM())
 	{
 		return VM->IsNativized();
@@ -1977,11 +1979,13 @@ bool IRigVMEditorAssetInterface::CanSetEditingLocked() const
 		return false;
 	}
 
+#if RIGVM_TRACE_ENABLED
 	if (Host->IsPlayingRewindDebugTrace())
 	{
 		return false;
 	}
-
+#endif
+	
 	return Host->CanSwapVMToNativized();
 }
 
@@ -1998,11 +2002,13 @@ FText IRigVMEditorAssetInterface::GetEditingLockedLabel() const
 		return FText();
 	}
 
+#if RIGVM_TRACE_ENABLED
 	if (Host->IsPlayingRewindDebugTrace())
 	{
 		return LOCTEXT("EditingLockedRewindDebuggerActive", "Rewind Debugger");
 	}
-
+#endif
+	
 	return LOCTEXT("EditingLockedNativized", "Nativized");
 }
 
@@ -2014,11 +2020,13 @@ FText IRigVMEditorAssetInterface::GetEditingLockedTooltip() const
 		return FText();
 	}
 
+#if RIGVM_TRACE_ENABLED
 	if (Host->IsPlayingRewindDebugTrace())
 	{
 		return LOCTEXT("EditingLockedRewindDebuggerActiveTooltip", "This asset is being driven by Rewind Debugger.");
 	}
-
+#endif
+	
 	return LOCTEXT("EditingLockedNativizedTooltip", "This asset is using a nativized VM if on.");
 }
 
diff --git a/Engine/Plugins/Runtime/RigVM/Source/RigVMEditor/Private/Editor/RigVMEditor.cpp b/Engine/Plugins/Runtime/RigVM/Source/RigVMEditor/Private/Editor/RigVMEditor.cpp
--- a/Engine/Plugins/Runtime/RigVM/Source/RigVMEditor/Private/Editor/RigVMEditor.cpp
+++ b/Engine/Plugins/Runtime/RigVM/Source/RigVMEditor/Private/Editor/RigVMEditor.cpp
@@ -1810,12 +1810,14 @@ FGraphAppearanceInfo FRigVMEditorBase::GetGraphAppearance(UEdGraph* InGraph) con
 
 		if(URigVMHost* RigVMHost = GetRigVMHost())
 		{
+#if RIGVM_TRACE_ENABLED
 			if (RigVMHost->IsPlayingRewindDebugTrace())
 			{
 				AppearanceInfo.InstructionFade = 1;
 				AppearanceInfo.InstructionText = FText::FromString(TEXT("Rewind Debug in Progress."));
-			}
-			else if(RigVMHost->GetVM() && RigVMHost->GetVM()->IsNativized())
+			} else
+#endif
+			if(RigVMHost->GetVM() && RigVMHost->GetVM()->IsNativized())
 			{
 				AppearanceInfo.InstructionFade = 1;
 				AppearanceInfo.InstructionText = FText::FromString(
