Control Rig: Fixed enable compilation of all functions on every comile by mistake
#rnx

#Synced-CL 52002492

diff --git a/Engine/Plugins/Runtime/RigVM/Source/RigVMDeveloper/Private/RigVMCompiler/RigVMCompiler.cpp b/Engine/Plugins/Runtime/RigVM/Source/RigVMDeveloper/Private/RigVMCompiler/RigVMCompiler.cpp
--- a/Engine/Plugins/Runtime/RigVM/Source/RigVMDeveloper/Private/RigVMCompiler/RigVMCompiler.cpp
+++ b/Engine/Plugins/Runtime/RigVM/Source/RigVMDeveloper/Private/RigVMCompiler/RigVMCompiler.cpp
@@ -943,7 +943,7 @@ bool URigVMCompiler::Compile(const FRigVMCompileSettings& InSettings, TArray<URi
 
 #if WITH_EDITOR
 	// During cooking, compile all public functions
-	//if (IsRunningCookCommandlet())
+	if (IsRunningCookCommandlet())
 	{
 		if (!CurrentlyCompilingFunctionNode)
 		{
