[RigVM] Fix non-determinism in CachedVMHash across process executions
- FRigVMCallstack hash (RigVMASTProxy.h): hashed TWeakObjectPtr via GUObjectArray
index (session-dependent); now uses object path name string instead.
- InterfaceOperands iteration (RigVMCompiler.cpp): TMap<FName, FRigVMOperand>
iterated unsorted; now sorted by pin name before building the operand array.
- UniqueHashForReferenceNode (RigVMCompiler.cpp): derived from raw pointer cast;
now uses GetTypeHash(FunctionReferenceNode->GetPathName()).

#Synced-CL 51945425

diff --git a/Engine/Plugins/Runtime/RigVM/Source/RigVMDeveloper/Private/RigVMCompiler/RigVMCompiler.cpp b/Engine/Plugins/Runtime/RigVM/Source/RigVMDeveloper/Private/RigVMCompiler/RigVMCompiler.cpp
--- a/Engine/Plugins/Runtime/RigVM/Source/RigVMDeveloper/Private/RigVMCompiler/RigVMCompiler.cpp
+++ b/Engine/Plugins/Runtime/RigVM/Source/RigVMDeveloper/Private/RigVMCompiler/RigVMCompiler.cpp
@@ -3664,7 +3664,7 @@ bool URigVMCompiler::TraverseInlineFunction(const FRigVMInlineFunctionExprAST* I
 
 	const FRigVMFunctionCompilationData* FunctionCompilationData = CompiledFunctions.FindChecked(FunctionHash).CompilationData;
 	const FRigVMByteCode& FunctionByteCode = FunctionCompilationData->ByteCode;
-	const uint32 UniqueHashForReferenceNode = (uint32)((PTRINT)FunctionReferenceNode); 
+	const uint32 UniqueHashForReferenceNode = GetTypeHash(FunctionReferenceNode->GetPathName());
 
 	// Bytecode to be inlined should never be aligned
 	checkf(!FunctionByteCode.bByteCodeIsAligned, TEXT("Trying to inline aligned function bytecode %s in package %s"), *FunctionReferenceNode->GetFunctionIdentifier().GetLibraryNodePath(), *GetPackage()->GetPathName());
@@ -3694,8 +3694,15 @@ bool URigVMCompiler::TraverseInlineFunction(const FRigVMInlineFunctionExprAST* I
 			return false;
 		}
 
+		// Sort by pin name for deterministic operand ordering across sessions
 		TArray<FRigVMOperand> InterfaceOperands;
-		for (const TPair<FName,FRigVMOperand>& Pair : FunctionCompilationData->InterfaceOperands)
+		TArray<TPair<FName, FRigVMOperand>> SortedPairs(FunctionCompilationData->InterfaceOperands.Array());
+		SortedPairs.Sort([](const TPair<FName, FRigVMOperand>& A, const TPair<FName, FRigVMOperand>& B)
+		{
+			return A.Key.ToString() < B.Key.ToString();
+		});
+		
+		for (const TPair<FName, FRigVMOperand>& Pair : SortedPairs)
 		{
 			InterfaceOperands.Add(Pair.Value);
 		}
diff --git a/Engine/Plugins/Runtime/RigVM/Source/RigVMDeveloper/Public/RigVMCompiler/RigVMASTProxy.h b/Engine/Plugins/Runtime/RigVM/Source/RigVMDeveloper/Public/RigVMCompiler/RigVMASTProxy.h
--- a/Engine/Plugins/Runtime/RigVM/Source/RigVMDeveloper/Public/RigVMCompiler/RigVMASTProxy.h
+++ b/Engine/Plugins/Runtime/RigVM/Source/RigVMDeveloper/Public/RigVMCompiler/RigVMASTProxy.h
@@ -109,11 +109,14 @@ private:
 		{
 			return 0;
 		}
+
+		const uint32 StackHash = Stack[Index].IsValid() ? GetTypeHash(Stack[Index]->GetPathName()) : 0;
 		if (Index == 0)
 		{
-			return GetTypeHash(Stack[Index]);
+			return StackHash;
 		}
-		return HashCombine(GetEntryTypeHash(Index - 1), GetTypeHash(Stack[Index]));
+
+		return HashCombine(GetEntryTypeHash(Index - 1), StackHash);
 	}
 
 	TArray<TWeakObjectPtr<UObject>> Stack;
