RigVM: Fix swap function references undo not properly handling orphan pins
#rnx

#Synced-CL 52002492

diff --git a/Engine/Plugins/Runtime/RigVM/Source/RigVMDeveloper/Private/RigVMModel/RigVMController.cpp b/Engine/Plugins/Runtime/RigVM/Source/RigVMDeveloper/Private/RigVMModel/RigVMController.cpp
--- a/Engine/Plugins/Runtime/RigVM/Source/RigVMDeveloper/Private/RigVMModel/RigVMController.cpp
+++ b/Engine/Plugins/Runtime/RigVM/Source/RigVMDeveloper/Private/RigVMModel/RigVMController.cpp
@@ -14790,7 +14790,6 @@ bool URigVMController::SwapFunctionReference(URigVMFunctionReferenceNode* InFunc
 		Action.SetTitle(TEXT("Swap Function Reference"));
 	}
 
-	InFunctionReferenceNode->Modify();
 	URigVMBuildData::Get()->UnregisterFunctionReference(InFunctionReferenceNode->ReferencedFunctionHeader.LibraryPointer, InFunctionReferenceNode);
 	InFunctionReferenceNode->ReferencedFunctionHeader.LibraryPointer = InNewFunctionIdentifier;
 	InFunctionReferenceNode->UpdateFunctionHeaderFromHost();
@@ -22733,6 +22732,15 @@ void URigVMController::DestroyObject(UObject* InObjectToDestroy) const
 			DestroyObject(PinsToDestroy[i]);
 		}
 	}
+	
+	if (URigVMPin* Pin = Cast<URigVMPin>(InObjectToDestroy))
+	{
+		if (URigVMNode* Node = Pin->GetNode())
+		{
+			Node->Pins.Remove(Pin);
+			Node->OrphanedPins.Remove(Pin);
+		}
+	}
 
 	// Invalidate the future loading of the destroyed object
 	FLinkerLoad::InvalidateExport(InObjectToDestroy);
