RigVM: Fix crash when calling PatchFunctionsOnLoad
#rnx

#Synced-CL 52002492

diff --git a/Engine/Plugins/Runtime/RigVM/Source/RigVMDeveloper/Private/RigVMEditorAsset.cpp b/Engine/Plugins/Runtime/RigVM/Source/RigVMDeveloper/Private/RigVMEditorAsset.cpp
--- a/Engine/Plugins/Runtime/RigVM/Source/RigVMDeveloper/Private/RigVMEditorAsset.cpp
+++ b/Engine/Plugins/Runtime/RigVM/Source/RigVMDeveloper/Private/RigVMEditorAsset.cpp
@@ -1432,6 +1432,14 @@ void IRigVMEditorAssetInterface::PostLoad()
 	}
 	
 	FRigVMRegistry::Get().RefreshEngineTypesIfRequired();
+	
+	FRigVMClient* RigVMClient = GetRigVMClient();
+	TScriptInterface<IRigVMClientHost> ClientHost = GetRigVMClientHost();
+	TScriptInterface<IRigVMGraphFunctionHost> FunctionHost;
+	if (ClientHost)
+	{
+		FunctionHost = ClientHost->GetRigVMGraphFunctionHost();
+	}
 
 	bVMRecompilationRequired = true;
 	{
@@ -1441,9 +1449,9 @@ void IRigVMEditorAssetInterface::PostLoad()
 
 		// PostLoad all referenced function hosts so that their function data are fully loaded 
 		// and ready to be inlined into this BP during compilation
-		for (TScriptInterface<IRigVMGraphFunctionHost> FunctionHost : ReferencedFunctionHosts)
+		for (TScriptInterface<IRigVMGraphFunctionHost> ReferencedFunctionHost : ReferencedFunctionHosts)
 		{
-			if (TScriptInterface<IRigVMRuntimeAssetInterface> RuntimeAsset = FunctionHost.GetObject())
+			if (TScriptInterface<IRigVMRuntimeAssetInterface> RuntimeAsset = ReferencedFunctionHost.GetObject())
 			{
 				if (RuntimeAsset.GetObject()->HasAllFlags(RF_NeedPostLoad))
 				{
@@ -1460,7 +1468,7 @@ void IRigVMEditorAssetInterface::PostLoad()
 		for (UEdGraph* Graph : GetUberGraphs())
 		{
 			URigVMEdGraph* RigGraph = Cast<URigVMEdGraph>(Graph);
-			if (RigGraph && RigGraph->GetClass() == GetRigVMClientHost()->GetRigVMEdGraphClass())
+			if (RigGraph && ClientHost && RigGraph->GetClass() == ClientHost->GetRigVMEdGraphClass())
 			{
 				NewUberGraphPages.Add(RigGraph);
 			}
@@ -1484,17 +1492,18 @@ void IRigVMEditorAssetInterface::PostLoad()
 		}
 		
 		InitializeModelIfRequired(false /* recompile vm */);
+		if (RigVMClient)
 		{
 			TGuardValue<bool> GuardNotifications(bSuspendModelNotificationsForSelf, true);
 			
-			const FRigVMClientPatchResult PatchResult = GetRigVMClient()->PatchModelsOnLoad();
+			const FRigVMClientPatchResult PatchResult = RigVMClient->PatchModelsOnLoad();
 			if(PatchResult.RequiresToMarkPackageDirty())
 			{
 				(void)GetObject()->MarkPackageDirty();
 				bDirtyDuringLoad = true;
 			}
 			
-			GetRigVMClient()->PatchFunctionReferencesOnLoad();
+			RigVMClient->PatchFunctionReferencesOnLoad();
 			GetFunctionReferenceNodeData() = GetReferenceNodeData();
 
 			PatchVariableNodesOnLoad();
@@ -1509,9 +1518,12 @@ void IRigVMEditorAssetInterface::PostLoad()
 			TArray<FName> BackwardsCompatiblePublicFunctions;
 			GetBackwardsCompatibilityPublicFunctions(BackwardsCompatiblePublicFunctions, OldHeaders);
 
-			GetRigVMClient()->PatchFunctionsOnLoad(GetRigVMClientHost()->GetRigVMGraphFunctionHost().GetInterface(), BackwardsCompatiblePublicFunctions, OldHeaders);
+			if (FunctionHost)
+			{
+				RigVMClient->PatchFunctionsOnLoad(FunctionHost.GetInterface(), BackwardsCompatiblePublicFunctions, OldHeaders);
+			}
 
-			const FRigVMClientPatchResult PinDefaultValuePatchResult = GetRigVMClient()->PatchPinDefaultValues();
+			const FRigVMClientPatchResult PinDefaultValuePatchResult = RigVMClient->PatchPinDefaultValues();
 			if(PinDefaultValuePatchResult.RequiresToMarkPackageDirty())
 			{
 				(void)GetObject()->MarkPackageDirty();
@@ -1534,9 +1546,18 @@ void IRigVMEditorAssetInterface::PostLoad()
 		
 		// at this point we may still have links which are detached. we may or may not be able to 
 		// reattach them.
-		GetRigVMClient()->ProcessDetachedLinks();
+		if (RigVMClient)
+		{
+			RigVMClient->ProcessDetachedLinks();
+		}
 
-		GetRigVMClientHost()->GetRigVMGraphFunctionHost()->GetRigVMGraphFunctionStore()->RemoveAllCompilationData();
+		if (FunctionHost)
+		{
+			if (FRigVMGraphFunctionStore* FunctionStore = FunctionHost->GetRigVMGraphFunctionStore())
+			{
+				FunctionStore->RemoveAllCompilationData();
+			}
+		}
 
 		// perform backwards compat value upgrades
 		TArray<URigVMGraph*> GraphsToValidate = GetAllModels();
@@ -1591,7 +1612,10 @@ void IRigVMEditorAssetInterface::PostLoad()
 #else
 	RecompileVMIfRequired();
 #endif
-	GetRigVMClientHost()->RequestRigVMInit();
+	if (ClientHost)
+	{
+		ClientHost->RequestRigVMInit();
+	}
 
 	
 	if (!GetAssetVariant().Guid.IsValid())
