Control Rig: Fix direct manipulation not working because of disable literal folding not being applied correctly
#rnx

#Synced-CL 52096374

diff --git a/Engine/Plugins/Runtime/RigVM/Source/RigVMDeveloper/Private/RigVMCompiler/RigVMCompiler.cpp b/Engine/Plugins/Runtime/RigVM/Source/RigVMDeveloper/Private/RigVMCompiler/RigVMCompiler.cpp
--- a/Engine/Plugins/Runtime/RigVM/Source/RigVMDeveloper/Private/RigVMCompiler/RigVMCompiler.cpp
+++ b/Engine/Plugins/Runtime/RigVM/Source/RigVMDeveloper/Private/RigVMCompiler/RigVMCompiler.cpp
@@ -183,8 +183,8 @@ FRigVMOperand FRigVMCompilerWorkData::FindOrAddProperty(const ERigVMMemoryType I
 	FString NewName = InNewName;
 	int32 Suffix = 0;
 
-	// always share literals
-	if(InMemoryType == ERigVMMemoryType::Literal && !bMakeUniqueProperty && PropertyDescriptions.Contains(ERigVMMemoryType::Literal))
+	// share literals only when literal folding is enabled (disabled during direct manipulation to keep per-node registers separate)
+	if(Settings.ASTSettings.bFoldLiterals && InMemoryType == ERigVMMemoryType::Literal && !bMakeUniqueProperty && PropertyDescriptions.Contains(ERigVMMemoryType::Literal))
 	{
 		const TArray<FRigVMPropertyDescription>& LiteralDescriptions = PropertyDescriptions[ERigVMMemoryType::Literal]; 
 		for(int32 Index = 0; Index < LiteralDescriptions.Num(); Index++)
