RigVM: Introduce FRigVMExecuteCallFrame to decouple per-execution state from FRigVMExtendedExecuteContext
The core change extracts all per-execution-frame state ? Slices, SliceOffsetsPerInstruction, SliceOffsetsPerCallable, and LazyBranchExecuteState ?
out of FRigVMExtendedExecuteContext and into a new FRigVMExecuteCallFrame struct. The context now holds a FrameStorage instance and a Frame pointer   that normally points to it, but can be redirected to point at a different frame ? enabling nativized callables to substitute their own frame
during execution without disturbing the outer VM state.

All existing access sites (RigVM.cpp, RigVMDebugInfo.cpp, RigVMMemoryStorage.cpp, and inline methods in RigVMExecuteContext.h) are updated to go     through Context.Frame-> instead of accessing the arrays directly on the context.

Additionally:
- FRigVMExecuteContext::InstructionIndex is moved from protected to public to allow nativized code to read/write it directly.
- FRigVMExecuteContext gains a FindDispatchFunction() accessor that delegates to the extended context, and FRigVMExtendedExecuteContext gains a      cached DispatchFunctionMap (TMap<FName, const FRigVMFunction*>) for efficient repeated dispatch function lookups.
- GetExtendedExecuteContext() accessors (const and non-const) are exposed on FRigVMExecuteContext.                                                   - TRigVMNativizedLazyValue gets a convenience constructor overload that drops the unused FRigVMNativizedContext parameter from the branch-index
variant, and adds a no-branch-index shorthand.                                                                                                       - Copy() and GetOperandSlice() on URigVMNativized are made static and GetOperandSlice gains an overload accepting FRigVMExtendedExecuteContext       directly.                                                                                                                                            - The code-gen Instruction.childtemplate is updated so that Copy no longer passes Context as its first argument, matching the now-static signature.  - All EngineTest nativized rig .cpp files and the RigVM compiler/test are updated to compile against the revised API.

#rnx

#Synced-CL 52245406

diff --git a/Engine/Plugins/Runtime/RigVM/Source/RigVM/Private/RigVMCore/RigVM.cpp b/Engine/Plugins/Runtime/RigVM/Source/RigVM/Private/RigVMCore/RigVM.cpp
--- a/Engine/Plugins/Runtime/RigVM/Source/RigVM/Private/RigVMCore/RigVM.cpp
+++ b/Engine/Plugins/Runtime/RigVM/Source/RigVM/Private/RigVMCore/RigVM.cpp
@@ -1686,8 +1686,8 @@ bool URigVM::InitializeInstance(FRigVMExtendedExecuteContext& Context, bool bCop
 	TGuardValue<int32> GuardThreadId(Context.ExecutingThreadId, FPlatformTLS::GetCurrentThreadId());
 
 	const int32 LazyBranchSize = GetByteCode().BranchInfos.Num();
-	Context.LazyBranchExecuteState.Reset(LazyBranchSize);
-	Context.LazyBranchExecuteState.SetNumZeroed(LazyBranchSize);
+	Context.Frame->LazyBranchExecuteState.Reset(LazyBranchSize);
+	Context.Frame->LazyBranchExecuteState.SetNumZeroed(LazyBranchSize);
 
 	if (bCopyMemory)
 	{
@@ -1757,8 +1757,8 @@ ERigVMExecuteResult URigVM::ExecuteVM(FRigVMExtendedExecuteContext& Context, con
 	{
 		Context.ResetExecutionState();
 		Context.SetVM(this);
-		Context.SliceOffsetsPerInstruction.AddZeroed(Instructions.Num());
-		Context.SliceOffsetsPerCallable.AddZeroed(ByteCode.NumCallables());
+		Context.Frame->SliceOffsetsPerInstruction.AddZeroed(Instructions.Num());
+		Context.Frame->SliceOffsetsPerCallable.AddZeroed(ByteCode.NumCallables());
 	}
 
 	ensure(Context.GetVM() == this);
diff --git a/Engine/Plugins/Runtime/RigVM/Source/RigVM/Private/RigVMCore/RigVMDebugInfo.cpp b/Engine/Plugins/Runtime/RigVM/Source/RigVM/Private/RigVMCore/RigVMDebugInfo.cpp
--- a/Engine/Plugins/Runtime/RigVM/Source/RigVM/Private/RigVMCore/RigVMDebugInfo.cpp
+++ b/Engine/Plugins/Runtime/RigVM/Source/RigVM/Private/RigVMCore/RigVMDebugInfo.cpp
@@ -196,7 +196,7 @@ int32 FRigVMDebugInfo::CompareCallstackPath(const FRigVMExtendedExecuteContext&
 void FRigVMDebugInfo::CaptureIterationPath(const FRigVMExtendedExecuteContext& Context, TArray<int32, TInlineAllocator<16>>& OutPath)
 {
 	OutPath.Reset();
-	for (const FRigVMSlice& Slice : Context.Slices)
+	for (const FRigVMSlice& Slice : Context.Frame->Slices)
 	{
 		if (Slice.IsLoop())
 		{
@@ -208,7 +208,7 @@ void FRigVMDebugInfo::CaptureIterationPath(const FRigVMExtendedExecuteContext& C
 int32 FRigVMDebugInfo::CompareIterationPath(const FRigVMExtendedExecuteContext& Context, const TArray<int32, TInlineAllocator<16>>& InPath)
 {
 	int32 PathIndex = 0;
-	for (const FRigVMSlice& Slice : Context.Slices)
+	for (const FRigVMSlice& Slice : Context.Frame->Slices)
 	{
 		if (Slice.IsLoop())
 		{
diff --git a/Engine/Plugins/Runtime/RigVM/Source/RigVM/Private/RigVMCore/RigVMExecuteContext.cpp b/Engine/Plugins/Runtime/RigVM/Source/RigVM/Private/RigVMCore/RigVMExecuteContext.cpp
--- a/Engine/Plugins/Runtime/RigVM/Source/RigVM/Private/RigVMCore/RigVMExecuteContext.cpp
+++ b/Engine/Plugins/Runtime/RigVM/Source/RigVM/Private/RigVMCore/RigVMExecuteContext.cpp
@@ -77,11 +77,19 @@ bool FRigVMExecuteContext::SerializeFromMismatchedTag(const FPropertyTag& Tag, F
 	return false;
 }
 
+const FRigVMFunction* FRigVMExecuteContext::FindDispatchFunction(const FName& InDispatchFunctionName) const
+{
+	check(ExtendedExecuteContext);
+	return ExtendedExecuteContext->FindDispatchFunction(InDispatchFunctionName);
+}
+
 PRAGMA_DISABLE_DEPRECATION_WARNINGS
+
 FRigVMExtendedExecuteContext::~FRigVMExtendedExecuteContext()
 {
 	Reset();
 }
+
 PRAGMA_ENABLE_DEPRECATION_WARNINGS
 
 // --- FRigVMExtendedExecuteContext ---
@@ -99,7 +107,7 @@ void FRigVMExtendedExecuteContext::Reset()
 	
 	OnExecutionReachedExitCallback.Reset();
 
-	LazyBranchExecuteState.Reset();
+	Frame->LazyBranchExecuteState.Reset();
 	ExternalVariableRuntimeData.Reset();
 
 	if(PublicDataScope.IsValid())
@@ -137,11 +145,10 @@ void FRigVMExtendedExecuteContext::ResetExecutionState()
 	}
 
 	VM = nullptr;
-	Slices.Reset();
-	Slices.Add(FRigVMSlice::MakeForCallable());
-	SliceOffsetsPerInstruction.Reset();
-	SliceOffsetsPerCallable.Reset();
+	Frame = &FrameStorage;
+	Frame->Reset();
 	Factory = nullptr;
+	DispatchFunctionMap.Reset();
 }
 
 void FRigVMExtendedExecuteContext::CopyMemoryStorage(const FRigVMExtendedExecuteContext& Other)
@@ -172,11 +179,12 @@ FRigVMExtendedExecuteContext& FRigVMExtendedExecuteContext::operator =(const FRi
 		SetDefaultNameCache();
 	}
 
-	Slices = Other.Slices;
-	SliceOffsetsPerInstruction = Other.SliceOffsetsPerInstruction;
-	SliceOffsetsPerCallable = Other.SliceOffsetsPerCallable;
-
-	LazyBranchExecuteState = Other.LazyBranchExecuteState;
+	FrameStorage.Slices = Other.FrameStorage.Slices;
+	FrameStorage.SliceOffsetsPerInstruction = Other.FrameStorage.SliceOffsetsPerInstruction;
+	FrameStorage.SliceOffsetsPerCallable = Other.FrameStorage.SliceOffsetsPerCallable;
+	FrameStorage.LazyBranchExecuteState = Other.FrameStorage.LazyBranchExecuteState;
+	Frame = &FrameStorage;
+	
 	ExternalVariableRuntimeData = Other.ExternalVariableRuntimeData;
 
 	return *this;
@@ -260,6 +268,30 @@ void FRigVMExtendedExecuteContext::SetVM(URigVM* InVM)
 	VM = InVM;
 }
 
+const FRigVMFunction* FRigVMExtendedExecuteContext::FindDispatchFunction(const FName& InDispatchFunctionName) const
+{
+	if (const FRigVMFunction* const* FunctionPtr = DispatchFunctionMap.Find(InDispatchFunctionName))
+	{
+		return *FunctionPtr;
+	}
+
+	if (const URigVM* CurrentVM = GetVM())
+	{
+		const FRigVMFunction* Function = nullptr;
+		if (FRigVMRegistry_NoLock* Registry = CurrentVM->GetLocalizedRegistry())
+		{
+			Function = Registry->FindFunction_NoLock(*InDispatchFunctionName.ToString());
+		}
+		else
+		{
+			Function = FRigVMRegistry::Get().FindFunction(*InDispatchFunctionName.ToString());
+		}
+		DispatchFunctionMap.Add(InDispatchFunctionName, Function);
+		return Function;
+	}
+	return nullptr;
+}
+
 const FRigVMOperand* FRigVMExtendedExecuteContext::GetDebugOperandForOperand(const FRigVMOperand& InOperand) const
 {
 	const FRigVMOperand KeyOperand(InOperand.GetMemoryType(), InOperand.GetRegisterIndex()); // no register offset
diff --git a/Engine/Plugins/Runtime/RigVM/Source/RigVM/Private/RigVMCore/RigVMMemoryStorage.cpp b/Engine/Plugins/Runtime/RigVM/Source/RigVM/Private/RigVMCore/RigVMMemoryStorage.cpp
--- a/Engine/Plugins/Runtime/RigVM/Source/RigVM/Private/RigVMCore/RigVMMemoryStorage.cpp
+++ b/Engine/Plugins/Runtime/RigVM/Source/RigVM/Private/RigVMCore/RigVMMemoryStorage.cpp
@@ -101,7 +101,7 @@ ERigVMExecuteResult FRigVMLazyBranch::Execute(FRigVMExtendedExecuteContext& Cont
 
 ERigVMExecuteResult FRigVMLazyBranch::ExecuteIfRequired(FRigVMExtendedExecuteContext& Context, uint32 InSliceHash)
 {
-	FRigVMInstructionSetExecuteState& ExecuteState = Context.LazyBranchExecuteState[BranchInfo.Index];
+	FRigVMInstructionSetExecuteState& ExecuteState = Context.Frame->LazyBranchExecuteState[BranchInfo.Index];
 	const uint32 HashForLazyBranch = Context.GetHashForLazyBranch();
 	if(ExecuteState.RequiresExecute(InSliceHash, HashForLazyBranch))
 	{
diff --git a/Engine/Plugins/Runtime/RigVM/Source/RigVM/Public/RigVMCore/RigVMExecuteContext.h b/Engine/Plugins/Runtime/RigVM/Source/RigVM/Public/RigVMCore/RigVMExecuteContext.h
--- a/Engine/Plugins/Runtime/RigVM/Source/RigVM/Public/RigVMCore/RigVMExecuteContext.h
+++ b/Engine/Plugins/Runtime/RigVM/Source/RigVM/Public/RigVMCore/RigVMExecuteContext.h
@@ -43,6 +43,7 @@ struct FRigVMExecuteContext;
 struct FRigVMExtendedExecuteContext;
 struct FRigVMLogSettings;
 struct FRigVMCallableInfo;
+struct FRigVMFunction;
 
 extern UE_API TAutoConsoleVariable<bool> CVarRigVMReportAllMessages;
 
@@ -309,9 +310,9 @@ struct FRigVMExecuteContext : public FRigVMExecutePin
 	GENERATED_BODY()
 
 	FRigVMExecuteContext()
-		: EventName(NAME_None)
+		: InstructionIndex(0)
+		, EventName(NAME_None)
 		, FunctionName(NAME_None)
-		, InstructionIndex(0)
 		, NumExecutions(0)
 		, DeltaTime(0.0)
 		, AbsoluteTime(0.0)
@@ -554,6 +555,8 @@ struct FRigVMExecuteContext : public FRigVMExecutePin
 		NumExecutions = 0;
 	}
 
+	int32 InstructionIndex;
+	
 protected:
 
 	virtual void Reset()
@@ -564,8 +567,6 @@ protected:
 	FName EventName;
 	
 	FName FunctionName;
-	
-	int32 InstructionIndex;
 
 	uint32 NumExecutions;
 
@@ -621,6 +622,11 @@ protected:
 	UEnum* InstanceOpCodeEnum;
 #endif
 
+public:
+	const FRigVMExtendedExecuteContext* GetExtendedExecuteContext() const { return ExtendedExecuteContext; }
+	FRigVMExtendedExecuteContext* GetExtendedExecuteContext() { return ExtendedExecuteContext; }
+	UE_API const FRigVMFunction* FindDispatchFunction(const FName& InDispatchFunctionName) const;
+
 private:
 	// Pointer back to owner execute context. Required for lazy branches Execute
 	FRigVMExtendedExecuteContext* ExtendedExecuteContext = nullptr; 
@@ -717,6 +723,38 @@ struct FRigVMExecuteCallstack
 	TArray<FRigVMExecuteCallstackEntry> Entries;
 };
 
+/**
+  * Slice addressing and lazy branch execution state for one callable invocation.
+  * Loops push/pop slices here; SliceOffsetsPerInstruction and SliceOffsetsPerCallable
+  * map each node and callsite to the correct offset into Work memory arrays.
+  * Entering a nativized RigVMStruct bends FRigVMExtendedExecuteContext::Frame to point
+  * at the Work-owned FRigVMExecuteCallFrame for the duration of StaticExecute.
+  */
+struct FRigVMExecuteCallFrame
+{
+	void Reset()
+	{
+		Slices.Reset();
+		Slices.Add(FRigVMSlice::MakeForCallable());
+		SliceOffsetsPerInstruction.Reset();
+		SliceOffsetsPerCallable.Reset();
+	}
+	
+	// The list of active slices. Slices can introduced by loops or callables
+	TArray<FRigVMSlice> Slices;
+
+	// The slice offset per instruction. Typically only loop / control flow nodes will
+	// introduce a new loop slice - thus their offset will increase.
+	TArray<uint16> SliceOffsetsPerInstruction;
+
+	// The slice offset per callable. When multiple call sites invoke the same callable,
+	// the offset for that callable index will be increased to make sure each call site
+	// gets its own memory.
+	TArray<uint16> SliceOffsetsPerCallable;
+
+	TArray<FRigVMInstructionSetExecuteState> LazyBranchExecuteState;
+};
+
 /**
  * The execute context is used for mutable nodes to
  * indicate execution order.
@@ -811,39 +849,39 @@ public:
 
 	const FRigVMSlice& GetSlice() const
 	{
-		return Slices.Last();
+		return Frame->Slices.Last();
 	}
 
 	FRigVMSlice& GetSlice()
 	{
-		return Slices.Last();
+		return Frame->Slices.Last();
 	}
 
 	void BeginLoopSlice(int32 InCount, int32 InRelativeIndex = 0)
 	{
 		ensure(!IsSliceComplete());
-		Slices.Add(FRigVMSlice::MakeForLoop(InCount));
-		Slices.Last().SetRelativeIndex(InRelativeIndex);
-		Slices.Last().SetInstructionOrCallableIndex(GetPublicData<>().InstructionIndex);
-		Slices.Last().SetOffset(SliceOffsetsPerInstruction[GetPublicData<>().InstructionIndex]);
+		Frame->Slices.Add(FRigVMSlice::MakeForLoop(InCount));
+		Frame->Slices.Last().SetRelativeIndex(InRelativeIndex);
+		Frame->Slices.Last().SetInstructionOrCallableIndex(GetPublicData<>().InstructionIndex);
+		Frame->Slices.Last().SetOffset(Frame->SliceOffsetsPerInstruction[GetPublicData<>().InstructionIndex]);
 	}
 
 	void EndLoopSlice()
 	{
-		ensure(Slices.Num() > 1);
-		ensure(Slices.Last().IsLoop());
+		ensure(Frame->Slices.Num() > 1);
+		ensure(Frame->Slices.Last().IsLoop());
 
 		// if this slice has reached its upperbound
 		// we want to make sure to increment the offset
 		// for the instruction the slice originated from,
 		// so that next time around the slice indices are not
 		// reused.
-		const FRigVMSlice PoppedSlice = Slices.Pop();
+		const FRigVMSlice PoppedSlice = Frame->Slices.Pop();
 		if(PoppedSlice.GetRelativeIndex() == PoppedSlice.Num() - 1)
 		{
-			if(SliceOffsetsPerInstruction.IsValidIndex(PoppedSlice.GetInstructionOrCallableIndex()))
+			if(Frame->SliceOffsetsPerInstruction.IsValidIndex(PoppedSlice.GetInstructionOrCallableIndex()))
 			{
-				SliceOffsetsPerInstruction[PoppedSlice.GetInstructionOrCallableIndex()] += static_cast<uint16>(PoppedSlice.Num());
+				Frame->SliceOffsetsPerInstruction[PoppedSlice.GetInstructionOrCallableIndex()] += static_cast<uint16>(PoppedSlice.Num());
 			}
 		}
 	}
@@ -851,20 +889,20 @@ public:
 	void BeginCallableSlice(int32 InCallableIndex)
 	{
 		ensure(!IsSliceComplete());
-		Slices.Add(FRigVMSlice::MakeForCallable());
-		Slices.Last().SetInstructionOrCallableIndex(InCallableIndex);
-		Slices.Last().SetOffset(SliceOffsetsPerCallable[InCallableIndex]);
+		Frame->Slices.Add(FRigVMSlice::MakeForCallable());
+		Frame->Slices.Last().SetInstructionOrCallableIndex(InCallableIndex);
+		Frame->Slices.Last().SetOffset(Frame->SliceOffsetsPerCallable[InCallableIndex]);
 	}
 
 	void EndCallableSlice()
 	{
-		ensure(Slices.Num() > 0);
-		ensure(Slices.Last().IsCallable());
+		ensure(Frame->Slices.Num() > 0);
+		ensure(Frame->Slices.Last().IsCallable());
 
-		const FRigVMSlice PoppedSlice = Slices.Pop();
-		if(SliceOffsetsPerCallable.IsValidIndex(PoppedSlice.GetInstructionOrCallableIndex()))
+		const FRigVMSlice PoppedSlice = Frame->Slices.Pop();
+		if(Frame->SliceOffsetsPerCallable.IsValidIndex(PoppedSlice.GetInstructionOrCallableIndex()))
 		{
-			SliceOffsetsPerCallable[PoppedSlice.GetInstructionOrCallableIndex()]++;
+			Frame->SliceOffsetsPerCallable[PoppedSlice.GetInstructionOrCallableIndex()]++;
 		}
 	}
 
@@ -877,7 +915,7 @@ public:
 	uint32 GetSliceHash() const
 	{
 		uint32 Hash = 0;
-		for(const FRigVMSlice& Slice : Slices)
+		for(const FRigVMSlice& Slice : Frame->Slices)
 		{
 			Hash = HashCombine(Hash, Slice.GetHash());
 		}
@@ -963,7 +1001,7 @@ public:
 
 	void InvalidateCachedMemory()
 	{
-		LazyBranchExecuteState.Reset();
+		Frame->LazyBranchExecuteState.Reset();
 	}
 	
 	void UpdateInstanceMemory(FRigVMMemoryStorageStruct* InLiteralMemory)
@@ -1000,6 +1038,8 @@ public:
 	UE_API FString GetVMPathName() const;
 	UE_API void SetVM(URigVM* InVM);
 
+	UE_API const FRigVMFunction* FindDispatchFunction(const FName& InDispatchFunctionName) const;
+
 	uint32 VMHash = 0;
 
 	FRigVMMemoryStorageStruct WorkMemoryStorage;
@@ -1024,22 +1064,15 @@ private:
 	URigVM* VM = nullptr;
 
 public:
-	// The list of active slices. Slices can introduced by loops or callables
-	TArray<FRigVMSlice> Slices;
-
-	// The slice offset per instruction. Typically only loop / control flow nodes will
-	// introduce a new loop slice - thus their offset will increase.
-	TArray<uint16> SliceOffsetsPerInstruction;
 
-	// The slice offset per callable. When multiple call sites invoke the same callable,
-	// the offset for that callable index will be increased to make sure each call site
-	// gets its own memory.
-	TArray<uint16> SliceOffsetsPerCallable;
+	FRigVMExecuteCallFrame FrameStorage;
+	FRigVMExecuteCallFrame* Frame = nullptr;
 
 	FRigVMExecuteCallstack Callstack;
 	uint32 CallstackHash = 0;
 
 	const FRigVMDispatchFactory* Factory = nullptr;
+	mutable TMap<FName, const FRigVMFunction*> DispatchFunctionMap;
 	FRigVMNameCache NameCache;
 	
 	int32 ExecutingThreadId = INDEX_NONE;
@@ -1081,7 +1114,6 @@ public:
 	/** Callback for external objects to be notified when the VM reaches an Exit Operation */
 	TFunction<void(const FName& EventName)> OnExecutionReachedExitCallback;
 
-	TArray<FRigVMInstructionSetExecuteState> LazyBranchExecuteState;
 	TArray<FRigVMExternalVariableRuntimeData> ExternalVariableRuntimeData;
 	TSharedPtr<FRigVMNativizedVMVariables> NativizedVMVariables;
 
diff --git a/Engine/Plugins/Runtime/RigVM/Source/RigVM/Public/RigVMCore/RigVMNativized.h b/Engine/Plugins/Runtime/RigVM/Source/RigVM/Public/RigVMCore/RigVMNativized.h
--- a/Engine/Plugins/Runtime/RigVM/Source/RigVM/Public/RigVMCore/RigVMNativized.h
+++ b/Engine/Plugins/Runtime/RigVM/Source/RigVM/Public/RigVMCore/RigVMNativized.h
@@ -74,7 +74,7 @@ public:
 		TRigVMLazyValueBase::MemoryHandle = &NativizedMemoryHandle;
 	}
 
-	TRigVMNativizedLazyValue(const ComputedType& Value, const FProperty* Property, int32 InBranchIndex, FRigVMNativizedContext& Context, const TFunction<ERigVMExecuteResult()>& ComputeCallback)
+	TRigVMNativizedLazyValue(const ComputedType& Value, const FProperty* Property, int32 InBranchIndex, const TFunction<ERigVMExecuteResult()>& ComputeCallback)
 	: TRigVMLazyValue<ComputedType>(Value)
 	, LazyBranch(InBranchIndex, ComputeCallback)
 	, NativizedMemoryHandle((uint8*)&Value, Property, nullptr, &LazyBranch)
@@ -82,6 +82,11 @@ public:
 		TRigVMLazyValueBase::MemoryHandle = &NativizedMemoryHandle;
 	}
 
+	TRigVMNativizedLazyValue(const ComputedType& Value, const FProperty* Property, const TFunction<ERigVMExecuteResult()>& ComputeCallback)
+	: TRigVMNativizedLazyValue(Value, Property, INDEX_NONE, ComputeCallback)
+	{
+	}
+	
 	virtual ~TRigVMNativizedLazyValue() override
 	{
 	}
@@ -214,8 +219,8 @@ protected:
 	{
 		Context.ResetExecutionState();
 		Context.SetVM(this);
-		Context.SliceOffsetsPerInstruction.SetNumZeroed(InNumberInstructions);
-		Context.SliceOffsetsPerCallable.SetNumZeroed(InNumberCallables);
+		Context.Frame->SliceOffsetsPerInstruction.SetNumZeroed(InNumberInstructions);
+		Context.Frame->SliceOffsetsPerCallable.SetNumZeroed(InNumberCallables);
 		Context.GetPublicData<ExecuteContextType>().EventName = InEntryName;
 		SetupInstructionTracking(Context, InNumberInstructions, InNumberCallables, InEntryName);
 		return Context.GetPublicData<ExecuteContextType>();
@@ -348,16 +353,24 @@ protected:
 		return EmptyValue;
 	}
 
+public:
+	
 	template<typename TargetT, typename SourceT>
-	void Copy(FRigVMNativizedContext& Context, const SourceT& InSource, TargetT& InTarget)
+	static void Copy(const SourceT& InSource, TargetT& InTarget)
 	{
 		InTarget = InSource;
 	}
 
 	template<typename ArrayType>
-	typename ArrayType::ElementType& GetOperandSlice(FRigVMNativizedContext& Context, ArrayType& InOutArray, const typename ArrayType::ElementType* InDefaultValue = nullptr)
+	static typename ArrayType::ElementType& GetOperandSlice(FRigVMNativizedContext& Context, ArrayType& InOutArray, const typename ArrayType::ElementType* InDefaultValue = nullptr)
 	{
-		const int32 SliceIndex = Context.ExecuteContext.GetSlice().GetIndex();
+		return GetOperandSlice(Context.ExecuteContext, InOutArray, InDefaultValue);
+	}
+
+	template<typename ArrayType>
+	static typename ArrayType::ElementType& GetOperandSlice(FRigVMExtendedExecuteContext& Context, ArrayType& InOutArray, const typename ArrayType::ElementType* InDefaultValue = nullptr)
+	{
+		const int32 SliceIndex = Context.GetSlice().GetIndex();
 		if(InOutArray.Num() <= SliceIndex)
 		{
 			const int32 FirstInsertedItem = InOutArray.AddDefaulted(1 + SliceIndex - InOutArray.Num());
@@ -371,9 +384,7 @@ protected:
 		}
 		return InOutArray[SliceIndex];
 	}
-
-public:
-
+	
 	template <
 		typename T,
 		typename TEnableIf<TRigVMIsBaseStructure<T>::Value, T>::Type* = nullptr
@@ -490,7 +501,7 @@ protected:
 			: Context(InContext.ExecuteContext)
 			, PublicData(InContext.PublicContextBase)
 			, PreviousInstructionIndex(INDEX_NONE)
-			, PreviousNumSlices(InContext.ExecuteContext.Slices.Num())
+			, PreviousNumSlices(InContext.ExecuteContext.Frame->Slices.Num())
 			, PreviousCallstackHash(Context.CallstackHash)
 			, bIsCallable(false)
 		{
@@ -523,13 +534,13 @@ protected:
 			}
 			else
 			{
-				while (Context.Slices.Num() > PreviousNumSlices)
+				while (Context.Frame->Slices.Num() > PreviousNumSlices)
 				{
-					if (Context.Slices.Last().IsLoop())
+					if (Context.Frame->Slices.Last().IsLoop())
 					{
 						Context.EndLoopSlice();
 					}
-					else if (Context.Slices.Last().IsCallable())
+					else if (Context.Frame->Slices.Last().IsCallable())
 					{
 						Context.EndCallableSlice();
 					}
@@ -562,19 +573,19 @@ protected:
 };
 
 template<>
-inline void URigVMNativized::Copy(FRigVMNativizedContext& Context, const double& InSource, float& InTarget)
+inline void URigVMNativized::Copy(const double& InSource, float& InTarget)
 {
 	InTarget = static_cast<float>(InSource);
 }
 
 template<>
-inline void URigVMNativized::Copy(FRigVMNativizedContext& Context, const float& InSource, double& InTarget)
+inline void URigVMNativized::Copy(const float& InSource, double& InTarget)
 {
 	InTarget = static_cast<double>(InSource);
 }
 
 template<>
-inline void URigVMNativized::Copy(FRigVMNativizedContext& Context, const TArray<double>& InSource, TArray<float>& InTarget)
+inline void URigVMNativized::Copy(const TArray<double>& InSource, TArray<float>& InTarget)
 {
 	InTarget.Reset();
 	InTarget.Reserve(InSource.Num());
@@ -585,7 +596,7 @@ inline void URigVMNativized::Copy(FRigVMNativizedContext& Context, const TArray<
 }
 
 template<>
-inline void URigVMNativized::Copy(FRigVMNativizedContext& Context, const TArray<float>& InSource, TArray<double>& InTarget)
+inline void URigVMNativized::Copy(const TArray<float>& InSource, TArray<double>& InTarget)
 {
 	InTarget.Reset();
 	InTarget.Reserve(InSource.Num());
