Restore and put RigVM property path warnings behind UE_BUILD_DEBUG

diff --git a/Engine/Plugins/Runtime/RigVM/Source/RigVM/Private/RigVMCore/RigVM.cpp b/Engine/Plugins/Runtime/RigVM/Source/RigVM/Private/RigVMCore/RigVM.cpp
--- a/Engine/Plugins/Runtime/RigVM/Source/RigVM/Private/RigVMCore/RigVM.cpp
+++ b/Engine/Plugins/Runtime/RigVM/Source/RigVM/Private/RigVMCore/RigVM.cpp
@@ -3711,6 +3711,14 @@ void URigVM::RefreshExternalPropertyPaths()
 			ExternalPropertyPaths[PropertyPathIndex] = FRigVMPropertyPath(
 				ExternalVariables[PropertyIndex].GetProperty(),
 				ExternalPropertyPathDescriptions[PropertyPathIndex].SegmentPath);
+#if WITH_EDITOR
+			if (!ExternalPropertyPaths[PropertyPathIndex].IsValid())
+			{
+#if UE_BUILD_DEBUG
+				UE_LOG(LogRigVM, Warning, TEXT("Unable to resolve property path '%s' for property '%s'."), *ExternalPropertyPathDescriptions[PropertyPathIndex].SegmentPath, *ExternalVariables[PropertyIndex].GetProperty()->GetPathName());
+#endif
+			}
+#endif
 		}
 	}
 }
