Allow structs within RigVM to implement ImportText for cache management purposes

#Synced-CL 52204442

diff --git a/Engine/Plugins/Runtime/RigVM/Source/RigVMDeveloper/Private/RigVMModel/RigVMPin.cpp b/Engine/Plugins/Runtime/RigVM/Source/RigVMDeveloper/Private/RigVMModel/RigVMPin.cpp
--- a/Engine/Plugins/Runtime/RigVM/Source/RigVMDeveloper/Private/RigVMModel/RigVMPin.cpp
+++ b/Engine/Plugins/Runtime/RigVM/Source/RigVMDeveloper/Private/RigVMModel/RigVMPin.cpp
@@ -1063,7 +1063,7 @@ FString URigVMPin::GetDefaultValue(const URigVMPin::FPinOverride& InOverride, bo
 	{
 		static const FString EmptyStructDefaultValue = TEXT("()");
 		
-		if((GetScriptStruct()->StructFlags & (STRUCT_ImportTextItemNative | STRUCT_ExportTextItemNative)) != 0)
+		if((GetScriptStruct()->StructFlags & (STRUCT_ImportTextItemNative | STRUCT_ExportTextItemNative)) == (STRUCT_ImportTextItemNative | STRUCT_ExportTextItemNative))
 		{
 			// If the struct has a native import/export, then its default value must be used verbatim (subpins are not displayed)
 			return DefaultValue;
