RigVM: Fix external variable type mismatch when inlining functions

Fix two related bugs in FRigVMCompilerWorkData::AddPropertiesForFunction that
caused crashes when inlining RigVM functions that reference external variables
with property paths (e.g. FTransform.Translation).

Bug 1 - Wrong external variable type resolution:
When processing a function's external variable operands during inlining,
the compiler looked up type information using VM->GetExternalVariableDefs()
indexed by the function's internal register index. However, the parent VM's
external variable array uses entirely different indices than the function's
compiled bytecode. This caused type confusion ? for example, a function's
external register 9 (FVector) could resolve to the parent's external
register 9 (double), producing invalid property path mappings and crashes.

The fix resolves external variable types by looking up the operand in the
function's InterfaceOperands map to find the argument name, then matching
it against InHeader.Arguments to get the correct CPPType and type object.
Falls back to the original index-based lookup when no match is found.

Bug 2 - Stale MemoryType for cached operand mappings:
When ProcessOperand encountered a base operand that was already mapped in a
previous call (e.g. the same external register seen a second time with a
property path offset), MemoryType was not updated from the cached mapping.
It remained set to External even though the operand had been remapped to
Work memory. This caused the property path verification to take the wrong
branch, comparing against the parent VM's external variable defs instead
of the work memory property descriptions.

The fix updates MemoryType from the cached operand's actual memory type.

#jra UE-366520
#rnx
#Synced-CL 52293635

diff --git a/Engine/Plugins/Runtime/RigVM/Source/RigVMDeveloper/Private/RigVMCompiler/RigVMCompiler.cpp b/Engine/Plugins/Runtime/RigVM/Source/RigVMDeveloper/Private/RigVMCompiler/RigVMCompiler.cpp
--- a/Engine/Plugins/Runtime/RigVM/Source/RigVMDeveloper/Private/RigVMCompiler/RigVMCompiler.cpp
+++ b/Engine/Plugins/Runtime/RigVM/Source/RigVMDeveloper/Private/RigVMCompiler/RigVMCompiler.cpp
@@ -489,6 +489,7 @@ bool FRigVMCompilerWorkData::AddPropertiesForFunction(
 			if (const FRigVMOperand* ExistingNewOperandNoOffset = OperandMap.Find(OldOperandNoOffset))
 			{
 				NewOperandNoOffset = *ExistingNewOperandNoOffset;
+				MemoryType = NewOperandNoOffset.GetMemoryType();
 			}
 			else
 			{
@@ -502,13 +503,53 @@ bool FRigVMCompilerWorkData::AddPropertiesForFunction(
 				else
 				{
 					check(MemoryType == ERigVMMemoryType::External);
-					check(VM->GetExternalVariableDefs().IsValidIndex(PropertyIndex));
-					
-					const FRigVMExternalVariableDef& ExternalVariableDef = VM->GetExternalVariableDefs()[PropertyIndex];
-					CustomDescription.Name = ExternalVariableDef.GetName();
-					CustomDescription.CPPType = ExternalVariableDef.GetExtendedCPPType().ToString();
-					CustomDescription.CPPTypeObject = const_cast<UObject*>(ExternalVariableDef.GetCPPTypeObject());
-					CustomDescription.DefaultValue = FString();
+
+					// Find the argument name for this external operand by looking it up in the
+					// function's InterfaceOperands map (operand -> name), then get the type from
+					// the function header's argument list. We cannot use
+					// VM->GetExternalVariableDefs()[PropertyIndex] because the parent VM's
+					// external variable indices differ from the function's register indices.
+					FName ArgumentName = NAME_None;
+					for (const TPair<FName, FRigVMOperand>& InterfacePair : InCompilationData->InterfaceOperands)
+					{
+						if (InterfacePair.Value.GetMemoryType() == OldOperandNoOffset.GetMemoryType() &&
+							InterfacePair.Value.GetRegisterIndex() == OldOperandNoOffset.GetRegisterIndex())
+						{
+							ArgumentName = InterfacePair.Key;
+							break;
+						}
+					}
+
+					const FRigVMGraphFunctionArgument* MatchingArgument = nullptr;
+					if (!ArgumentName.IsNone())
+					{
+						for (const FRigVMGraphFunctionArgument& Argument : InHeader.Arguments)
+						{
+							if (Argument.Name == ArgumentName)
+							{
+								MatchingArgument = &Argument;
+								break;
+							}
+						}
+					}
+
+					if (MatchingArgument)
+					{
+						CustomDescription.Name = MatchingArgument->Name;
+						CustomDescription.CPPType = MatchingArgument->CPPType.ToString();
+						CustomDescription.CPPTypeObject = MatchingArgument->CPPTypeObject.Get();
+						CustomDescription.DefaultValue = FString();
+					}
+					else
+					{
+						// Fallback: use the parent VM's external variable def by index (original behavior)
+						check(VM->GetExternalVariableDefs().IsValidIndex(PropertyIndex));
+						const FRigVMExternalVariableDef& ExternalVariableDef = VM->GetExternalVariableDefs()[PropertyIndex];
+						CustomDescription.Name = ExternalVariableDef.GetName();
+						CustomDescription.CPPType = ExternalVariableDef.GetExtendedCPPType().ToString();
+						CustomDescription.CPPTypeObject = const_cast<UObject*>(ExternalVariableDef.GetCPPTypeObject());
+						CustomDescription.DefaultValue = FString();
+					}
 					Description = &CustomDescription;
 
 					// if this is an input argument - make sure to turn it into work state.
