RigVM: Add BreakFromLoop and ContinueLoop nodes with compiler support

Adds two new RigVM execution nodes for early loop control flow:
- BreakFromLoop: terminates the enclosing loop after the current iteration
by writing ControlFlowCompletedName into the loop's BlockToRun register
via Singleton operand aliasing.
- ContinueLoop: skips remaining body nodes in the current iteration via
compiler-emitted JumpForward, then proceeds to the next iteration.

Compiler changes (RigVMCompiler.cpp):
- FindEnclosingSlicedLoopExpr: walks AST parent chain to locate the
nearest enclosing sliced loop for Break/Continue validation.
- Break: aliases BlockToRun operand to enclosing loop's register during
SetupMemory; force-traverses the loop's BlockToRun var expr when it
has not yet been registered (body block may be visited before vars).
- Break/Continue: emits JumpForward placeholder during BuildInstructions,
patched to the loop's EndBlock instruction on emission.

Bug fix (RigVMNode.cpp):
- IsLoopNode() now accepts both ExecuteContextName and ExecutePinName
as the body block name, fixing detection for FRigVMStructMutable-based
loops like ForLoopCount.

Node design:
- Both nodes inherit from FRigVMStruct (not FRigVMStructMutable) with
input-only ExecuteContext pin ??? no output execution pin since downstream
nodes would be dead code.
- Styled with NodeColor black and loop icon to match ForLoopCount.

Engine tests (EngineTestRigVMCompiler.cpp):
- BreakFromLoop: ForLoopCount, ForEach, outside-loop negative, no-output-link
- ContinueLoop: ForLoopCount, ForEach, outside-loop negative, no-output-link

#rnx

#Synced-CL 52296664

diff --git a/Engine/Plugins/Runtime/RigVM/Source/RigVM/Private/RigVMFunctions/Execution/RigVMFunction_BreakLoop.cpp b/Engine/Plugins/Runtime/RigVM/Source/RigVM/Private/RigVMFunctions/Execution/RigVMFunction_BreakLoop.cpp
new file mode 100644
--- /dev/null
+++ b/Engine/Plugins/Runtime/RigVM/Source/RigVM/Private/RigVMFunctions/Execution/RigVMFunction_BreakLoop.cpp
@@ -0,0 +1,18 @@
+// Copyright Epic Games, Inc. All Rights Reserved.
+
+#include "RigVMFunctions/Execution/RigVMFunction_BreakLoop.h"
+
+#include UE_INLINE_GENERATED_CPP_BY_NAME(RigVMFunction_BreakLoop)
+
+/**
+ * Sets BlockToRun to ControlFlowCompletedName, signaling the enclosing loop
+ * to exit after the current iteration finishes. Because BlockToRun is aliased
+ * (via Singleton) to the loop's own BlockToRun register, the loop's next
+ * JumpToBranch will route to the Completed branch instead of re-entering the body.
+ */
+FRigVMFunction_BreakFromLoop_Execute()
+{
+	// Route the shared BlockToRun register to the Completed branch,
+	// causing the enclosing loop to terminate on its next iteration check.
+	BlockToRun = ControlFlowCompletedName;
+}
diff --git a/Engine/Plugins/Runtime/RigVM/Source/RigVM/Private/RigVMFunctions/Execution/RigVMFunction_ContinueLoop.cpp b/Engine/Plugins/Runtime/RigVM/Source/RigVM/Private/RigVMFunctions/Execution/RigVMFunction_ContinueLoop.cpp
new file mode 100644
--- /dev/null
+++ b/Engine/Plugins/Runtime/RigVM/Source/RigVM/Private/RigVMFunctions/Execution/RigVMFunction_ContinueLoop.cpp
@@ -0,0 +1,14 @@
+// Copyright Epic Games, Inc. All Rights Reserved.
+
+#include "RigVMFunctions/Execution/RigVMFunction_ContinueLoop.h"
+
+#include UE_INLINE_GENERATED_CPP_BY_NAME(RigVMFunction_ContinueLoop)
+
+/**
+ * Intentionally empty. ContinueLoop's skip behaviour is a compiler artifact:
+ * a JumpForward instruction emitted after the CallExtern bypasses the remaining
+ * loop body nodes. No runtime state modification is needed.
+ */
+FRigVMFunction_ContinueLoop_Execute()
+{
+}
diff --git a/Engine/Plugins/Runtime/RigVM/Source/RigVM/Public/RigVMFunctions/Execution/RigVMFunction_BreakLoop.h b/Engine/Plugins/Runtime/RigVM/Source/RigVM/Public/RigVMFunctions/Execution/RigVMFunction_BreakLoop.h
new file mode 100644
--- /dev/null
+++ b/Engine/Plugins/Runtime/RigVM/Source/RigVM/Public/RigVMFunctions/Execution/RigVMFunction_BreakLoop.h
@@ -0,0 +1,73 @@
+// Copyright Epic Games, Inc. All Rights Reserved.
+
+/*
+ * =============================================================================
+ * DESIGN INTENT
+ * =============================================================================
+ * GOAL:
+ *   Defines FRigVMFunction_BreakFromLoop, a RigVM node that allows early exit
+ *   from a For Loop or For Each loop body during execution.
+ *
+ * APPROACH:
+ *   The node carries a singleton BlockToRun (FName) pin that the compiler
+ *   aliases — at compile time — to the enclosing loop's BlockToRun memory slot.
+ *   At runtime, Execute() writes ControlFlowCompletedName into that shared slot,
+ *   causing the loop's JumpToBranch to route to the Completed branch on the next
+ *   iteration check. The current iteration always completes in full.
+ *   This node has no control flow blocks of its own (it does not branch); it is
+ *   a pure "write-and-let-the-loop-read" passthrough.
+ *
+ * REVIEWER NOTES:
+ *   The Singleton meta on BlockToRun is critical — it tells the compiler to alias
+ *   this pin's memory to the enclosing loop's BlockToRun register rather than
+ *   allocating new memory. Without Singleton, the write would go to an isolated
+ *   slot the loop never reads, making Break a no-op at runtime.
+ *   Do NOT add Hidden meta to BlockToRun — Hidden suppresses pin creation
+ *   entirely, which would prevent the compiler from performing the aliasing.
+ * =============================================================================
+ */
+
+#pragma once
+
+#include "RigVMCore/RigVMStruct.h"
+#include "RigVMFunction_BreakLoop.generated.h"
+
+/**
+ * Breaks out of the enclosing For Loop or For Each loop at the end of
+ * the current iteration. When executed, it sets the enclosing loop's
+ * BlockToRun to "Completed", causing the loop to skip remaining
+ * iterations and proceed to its Completed branch.
+ *
+ * Must be placed inside a loop body; the compiler will emit an error
+ * if this node is used outside a loop scope.
+ */
+USTRUCT(meta=(DisplayName="Break From Loop", Category="Execution", NodeColor = "0, 0, 0, 1", Keywords="Break,Stop,Exit,Loop", Icon="EditorStyle|GraphEditor.Macro.Loop_16x"))
+struct FRigVMFunction_BreakFromLoop : public FRigVMStruct
+{
+	GENERATED_BODY()
+
+	FRigVMFunction_BreakFromLoop()
+	{
+		BlockToRun = NAME_None;
+	}
+
+	/** Writes ControlFlowCompletedName into BlockToRun to signal the enclosing loop to exit. */
+	RIGVM_METHOD()
+	RIGVM_API virtual void Execute() override;
+
+	// Input-only execution pin — no output execution since the compiler emits
+	// a JumpForward after this node, making any downstream nodes dead code.
+	UPROPERTY(meta = (Input, DisplayName = "Execute"))
+	FRigVMExecuteContext ExecuteContext;
+
+	/**
+	 * The name of the block the enclosing loop should run next.
+	 * Singleton meta causes the compiler to alias this property to the enclosing
+	 * loop's BlockToRun register — no separate memory is allocated for this node.
+	 * At runtime, Execute() writes ControlFlowCompletedName here, and the loop's
+	 * JumpToBranch reads it to decide whether to re-enter the body or jump to Completed.
+	 * Do NOT add Hidden meta — it suppresses pin creation and prevents aliasing.
+	 */
+	UPROPERTY(meta = (Singleton))
+	FName BlockToRun;
+};
diff --git a/Engine/Plugins/Runtime/RigVM/Source/RigVM/Public/RigVMFunctions/Execution/RigVMFunction_ContinueLoop.h b/Engine/Plugins/Runtime/RigVM/Source/RigVM/Public/RigVMFunctions/Execution/RigVMFunction_ContinueLoop.h
new file mode 100644
--- /dev/null
+++ b/Engine/Plugins/Runtime/RigVM/Source/RigVM/Public/RigVMFunctions/Execution/RigVMFunction_ContinueLoop.h
@@ -0,0 +1,60 @@
+// Copyright Epic Games, Inc. All Rights Reserved.
+
+/*
+ * =============================================================================
+ * DESIGN INTENT
+ * =============================================================================
+ * GOAL:
+ *   Defines FRigVMFunction_ContinueLoop, a RigVM node that skips the remaining
+ *   body nodes in the current loop iteration and immediately re-enters the
+ *   enclosing For Loop or For Each loop at the next iteration.
+ *
+ * APPROACH:
+ *   Unlike BreakFromLoop, ContinueLoop has NO BlockToRun pin. It does not need
+ *   to write any value to the loop's control-flow register. Instead, the skip
+ *   behaviour is entirely a compiler artifact: the RigVM compiler emits a
+ *   JumpForward instruction immediately after the CallExtern for this node,
+ *   which jumps past the remaining body instructions to the loop's increment/
+ *   re-entry point. The loop's own BlockToRun remains set to its body branch,
+ *   so iteration continues normally — only the remainder of the *current*
+ *   iteration is skipped.
+ *
+ * REVIEWER NOTES:
+ *   Because the semantics are fully handled by the compiler's JumpForward
+ *   emission, Execute() is intentionally empty — no runtime side effects are
+ *   needed. The struct exists solely to provide a typed CallExtern target that
+ *   the compiler recognizes as a continue directive.
+ *   Do NOT add a BlockToRun property here. Doing so would cause the compiler
+ *   to treat this node like BreakFromLoop and attempt Singleton aliasing,
+ *   which is incorrect for continue semantics.
+ * =============================================================================
+ */
+
+#pragma once
+
+#include "RigVMCore/RigVMStruct.h"
+#include "RigVMFunction_ContinueLoop.generated.h"
+
+/**
+ * Skips the remaining body nodes of the current loop iteration and
+ * re-enters the enclosing For Loop or For Each loop at the next iteration.
+ * The loop itself is not terminated — only the current iteration's remaining
+ * nodes are bypassed.
+ *
+ * Must be placed inside a loop body; the compiler will emit an error
+ * if this node is used outside a loop scope.
+ */
+USTRUCT(meta=(DisplayName="Continue Loop", Category="Execution", NodeColor = "0, 0, 0, 1", Keywords="Continue,Skip,Next,Loop", Icon="EditorStyle|GraphEditor.Macro.Loop_16x"))
+struct FRigVMFunction_ContinueLoop : public FRigVMStruct
+{
+	GENERATED_BODY()
+
+	/** Intentionally empty — continue semantics are handled by the compiler's JumpForward. */
+	RIGVM_METHOD()
+	RIGVM_API virtual void Execute() override;
+
+	// Input-only execution pin — no output execution since the compiler emits
+	// a JumpForward after this node, making any downstream nodes dead code.
+	UPROPERTY(meta = (Input, DisplayName = "Execute"))
+	FRigVMExecuteContext ExecuteContext;
+};
diff --git a/Engine/Plugins/Runtime/RigVM/Source/RigVMDeveloper/Private/RigVMCompiler/RigVMCompiler.cpp b/Engine/Plugins/Runtime/RigVM/Source/RigVMDeveloper/Private/RigVMCompiler/RigVMCompiler.cpp
--- a/Engine/Plugins/Runtime/RigVM/Source/RigVMDeveloper/Private/RigVMCompiler/RigVMCompiler.cpp
+++ b/Engine/Plugins/Runtime/RigVM/Source/RigVMDeveloper/Private/RigVMCompiler/RigVMCompiler.cpp
@@ -16,6 +16,8 @@
 #include "ProfilingDebugging/ScopedTimers.h"
 #include "RigVMModel/RigVMClient.h"
 #include "RigVMFunctions/RigVMDispatch_Array.h"
+#include "RigVMFunctions/Execution/RigVMFunction_BreakLoop.h" // FRigVMFunction_BreakFromLoop — used for Break node detection in TraverseCallExtern
+#include "RigVMFunctions/Execution/RigVMFunction_ContinueLoop.h" // FRigVMFunction_ContinueLoop — used for Continue node detection in TraverseCallExtern
 #include "Algo/Sort.h"
 #include "String/Join.h"
 #include "RigVMDeveloperTypeUtils.h"
@@ -3070,6 +3072,55 @@ bool URigVMCompiler::TraverseEntry(const FRigVMEntryExprAST* InExpr, FRigVMCompi
 	return TraverseChildren(InExpr, WorkData);
 }
 
+/**
+ * Walks the AST parent chain from InExpr to find the nearest enclosing loop
+ * that has at least one sliced control-flow block. Stops (returns nullptr) if
+ * an InlineFunction boundary is encountered before a matching loop is found.
+ *
+ * This is shared by Break (SetupMemory + BuildInstructions phases) and Continue
+ * (SetupMemory + BuildInstructions phases) to avoid duplicating the walk logic.
+ *
+ * @param InExpr  The expression whose parents are walked (typically a Break/Continue CallExternExprAST).
+ * @return        The nearest enclosing sliced-loop CallExternExprAST, or nullptr if not found.
+ */
+static const FRigVMCallExternExprAST* FindEnclosingSlicedLoopExpr(const FRigVMExprAST* InExpr)
+{
+	for (const FRigVMExprAST* ParentExpr = InExpr->GetParent();
+		 ParentExpr != nullptr;
+		 ParentExpr = ParentExpr->GetParent())
+	{
+		// Stop at inline-function boundaries — Break/Continue must not jump across function scope.
+		// In the current AST topology this is unreachable (inlined function bodies are pre-compiled
+		// separately), but guards against future AST changes.
+		if (ParentExpr->IsA(FRigVMExprAST::EType::InlineFunction))
+		{
+			break;
+		}
+
+		if (!ParentExpr->IsA(FRigVMExprAST::EType::CallExtern))
+		{
+			continue;
+		}
+
+		const FRigVMCallExternExprAST* ParentCallExtern = ParentExpr->To<FRigVMCallExternExprAST>();
+		URigVMNode* ParentNode = ParentCallExtern->GetNode();
+
+		if (ParentNode->IsLoopNode())
+		{
+			// IsLoopNode() does not imply slicing — verify that at least one block is sliced.
+			// A non-sliced loop (if one existed) would have no per-iteration BlockToRun memory.
+			for (const FName& BlockName : ParentNode->GetControlFlowBlocks())
+			{
+				if (ParentNode->IsControlFlowBlockSliced(BlockName))
+				{
+					return ParentCallExtern;
+				}
+			}
+		}
+	}
+	return nullptr;
+}
+
 bool URigVMCompiler::TraverseCallExtern(const FRigVMCallExternExprAST* InExpr, FRigVMCompilerWorkData& WorkData)
 {
 	URigVMNode* Node = InExpr->GetNode();
@@ -3141,6 +3192,131 @@ bool URigVMCompiler::TraverseCallExtern(const FRigVMCallExternExprAST* InExpr, F
 	int32 CallExternInstructionIndex = INDEX_NONE;
 	const FRigVMCallstack Callstack = InExpr->GetProxy().GetCallstack();
 
+	// When a Break or Continue node is detected during the BuildInstructions phase,
+	// this is set to the enclosing loop's expression so a JumpForward can be emitted
+	// after the CallExtern op for the Break/Continue node.
+	const FRigVMCallExternExprAST* EnclosingLoopExprForJump = nullptr;
+
+	// ── Break-from-loop special case ──────────────────────────────────
+	// When the current node is FRigVMFunction_BreakFromLoop, its BlockToRun
+	// pin must alias the enclosing loop's BlockToRun memory slot instead of
+	// allocating a new register. We walk the AST parent chain to locate the
+	// nearest enclosing sliced loop, retrieve its already-registered
+	// BlockToRun operand, and pre-insert it into ExprToOperand for the
+	// Break's own BlockToRun var expr. FindOrAddRegister will then find the
+	// entry and reuse the operand — no new memory is allocated.
+	if (UnitNode && UnitNode->GetScriptStruct() == FRigVMFunction_BreakFromLoop::StaticStruct())
+	{
+		// Only alias during the SetupMemory phase — that is when
+		// FindOrAddRegister is called and would otherwise allocate.
+		if (WorkData.CompilerPhase == ERigVMCompilerPhase_SetupMemory)
+		{
+			// Walk the AST parent chain to find the nearest enclosing sliced loop.
+			// See FindEnclosingSlicedLoopExpr (above TraverseCallExtern) for the walk
+			// logic and cross-scope safety rationale.
+			const FRigVMCallExternExprAST* EnclosingLoopExpr = FindEnclosingSlicedLoopExpr(InExpr);
+
+			// Error: Break used outside any loop scope.
+			if (EnclosingLoopExpr == nullptr)
+			{
+				static const FString BreakOutsideLoopMessage = TEXT("Break From Loop node @@ is not inside a loop.");
+				WorkData.Settings.Report(EMessageSeverity::Error, UnitNode, BreakOutsideLoopMessage);
+				return false;
+			}
+
+			// Retrieve the loop's BlockToRun var expr.
+			const FRigVMVarExprAST* LoopBlockToRunVarExpr =
+				EnclosingLoopExpr->FindVarWithPinName(FRigVMStruct::ControlFlowBlockToRunName);
+			if (!LoopBlockToRunVarExpr)
+			{
+				static const FString MissingBlockToRunMessage = TEXT("Break From Loop node @@ could not find BlockToRun on enclosing loop.");
+				WorkData.Settings.Report(EMessageSeverity::Error, UnitNode, MissingBlockToRunMessage);
+				return false;
+			}
+
+			// The Break node lives inside the loop's body block, which may be
+			// traversed before the loop's BlockToRun var expr has been registered
+			// in ExprToOperand (child ordering is not guaranteed). Force-traverse
+			// the loop's BlockToRun now so FindOrAddRegister allocates its operand.
+			const FRigVMVarExprAST* LoopBlockToRunSourceExpr = GetSourceVarExpr(LoopBlockToRunVarExpr);
+			if (LoopBlockToRunSourceExpr == nullptr)
+			{
+				LoopBlockToRunSourceExpr = LoopBlockToRunVarExpr;
+			}
+
+			if (!WorkData.ExprToOperand.Contains(LoopBlockToRunSourceExpr))
+			{
+				if (!TraverseExpression(const_cast<FRigVMExprAST*>(static_cast<const FRigVMExprAST*>(LoopBlockToRunVarExpr)), WorkData))
+				{
+					return false;
+				}
+			}
+
+			const FRigVMOperand* LoopOperandPtr = WorkData.ExprToOperand.Find(LoopBlockToRunSourceExpr);
+			if (!LoopOperandPtr)
+			{
+				static const FString UnregisteredOperandMessage = TEXT("Break From Loop node @@ found enclosing loop but its BlockToRun operand is not yet registered.");
+				WorkData.Settings.Report(EMessageSeverity::Error, UnitNode, UnregisteredOperandMessage);
+				return false;
+			}
+
+			// Pre-insert the Break's BlockToRun var expr with the loop's operand so
+			// FindOrAddRegister reuses it during TraverseChildren below.
+			const FRigVMVarExprAST* BreakBlockToRunVarExpr =
+				InExpr->FindVarWithPinName(FRigVMStruct::ControlFlowBlockToRunName);
+			if (!BreakBlockToRunVarExpr)
+			{
+				static const FString MissingBreakBlockToRunMessage = TEXT("Break From Loop node @@ has no BlockToRun pin expression.");
+				WorkData.Settings.Report(EMessageSeverity::Error, UnitNode, MissingBreakBlockToRunMessage);
+				return false;
+			}
+
+			const FRigVMVarExprAST* BreakBlockToRunSourceExpr = GetSourceVarExpr(BreakBlockToRunVarExpr);
+			if (!WorkData.ExprToOperand.Contains(BreakBlockToRunSourceExpr))
+			{
+				WorkData.ExprToOperand.Add(BreakBlockToRunSourceExpr, *LoopOperandPtr);
+			}
+		}
+		// ── Break: BuildInstructions phase — locate enclosing loop for JumpForward ──
+		// The CallExtern op for the Break node is emitted by the standard path below.
+		// We find the enclosing loop here and store it in EnclosingLoopExprForJump so
+		// that a JumpForward placeholder can be emitted immediately after the CallExtern.
+		// If nullptr, SetupMemory already reported an error and returned false, so we
+		// will not reach this point for an out-of-loop Break.
+		else if (WorkData.CompilerPhase == ERigVMCompilerPhase_BuildInstructions)
+		{
+			EnclosingLoopExprForJump = FindEnclosingSlicedLoopExpr(InExpr);
+		}
+		// In other phases, fall through to the normal CallExtern path —
+		// the aliased operand is already in ExprToOperand from the SetupMemory pass.
+	}
+
+	// ── Continue-loop special case ──────────────────────────────────
+	// When the current node is FRigVMFunction_ContinueLoop, the compiler emits a
+	// JumpForward after its CallExtern op to skip the remaining loop body nodes.
+	// Unlike Break, Continue has no BlockToRun pin and needs no operand aliasing.
+	if (UnitNode && UnitNode->GetScriptStruct() == FRigVMFunction_ContinueLoop::StaticStruct())
+	{
+		// ── Continue: SetupMemory phase — validate that we're inside a loop ──
+		if (WorkData.CompilerPhase == ERigVMCompilerPhase_SetupMemory)
+		{
+			// Error: Continue used outside any loop scope.
+			if (FindEnclosingSlicedLoopExpr(InExpr) == nullptr)
+			{
+				static const FString ContinueOutsideLoopMessage = TEXT("Continue Loop node @@ is not inside a loop.");
+				WorkData.Settings.Report(EMessageSeverity::Error, UnitNode, ContinueOutsideLoopMessage);
+				return false;
+			}
+			// No operand aliasing needed — Continue has no BlockToRun pin.
+		}
+		// ── Continue: BuildInstructions phase — locate enclosing loop for JumpForward ──
+		// If nullptr, SetupMemory already reported an error for an out-of-loop Continue.
+		else if (WorkData.CompilerPhase == ERigVMCompilerPhase_BuildInstructions)
+		{
+			EnclosingLoopExprForJump = FindEnclosingSlicedLoopExpr(InExpr);
+		}
+	}
+
 	if (WorkData.CompilerPhase == ERigVMCompilerPhase_SetupMemory)
 	{
 		if (!TraverseChildren(InExpr, WorkData))
@@ -3587,6 +3763,23 @@ bool URigVMCompiler::TraverseCallExtern(const FRigVMCallExternExprAST* InExpr, F
 			WorkData.VM->GetByteCode().SetSubject(CallExternInstructionIndex, Callstack.GetCallPath(), Callstack.GetStack());
 		}
 
+		// ── Emit JumpForward for Break / Continue nodes ──────────────────
+		// If EnclosingLoopExprForJump was set during the Break or Continue
+		// detection above, the CallExtern for the Break/Continue has now been
+		// emitted. Emit a JumpForward placeholder (with INDEX_NONE) that will
+		// be patched to jump to the loop's EndBlock instruction once it is emitted.
+		if (EnclosingLoopExprForJump != nullptr)
+		{
+			const int32 JumpByteOffset = WorkData.VM->GetByteCode().AddJumpOp(ERigVMOpCode::JumpForward, INDEX_NONE);
+			const int32 JumpInstructionIndex = WorkData.VM->GetByteCode().GetNumInstructions() - 1;
+
+			// Register the pending jump so it can be patched when EndBlock is emitted for the loop.
+			FRigVMCompilerWorkData::FPendingLoopJump PendingJump;
+			PendingJump.ByteOffset = JumpByteOffset;
+			PendingJump.InstructionIndex = JumpInstructionIndex;
+			WorkData.PendingLoopJumps.FindOrAdd(EnclosingLoopExprForJump).Add(PendingJump);
+		}
+
 		if(Node->IsControlFlowNode())
 		{
 			// add an operator to jump to the right branch
@@ -3646,6 +3839,25 @@ bool URigVMCompiler::TraverseCallExtern(const FRigVMCallExternExprAST* InExpr, F
 				if(Node->IsControlFlowBlockSliced(BlockName))
 				{
 					WorkData.VM->GetByteCode().AddEndBlockOp();
+
+					// ── Patch pending Break/Continue JumpForward ops for this loop ──
+					// Any Break or Continue nodes inside this loop body will have emitted
+					// JumpForward placeholders that need their relative offsets set to land
+					// on the EndBlock instruction (so EndBlock still executes and unwinds
+					// BeginLoopSlice/EndLoopSlice state). Now that EndBlock has been emitted,
+					// we know its instruction index and can patch all pending jumps.
+					if (TArray<FRigVMCompilerWorkData::FPendingLoopJump>* Jumps = WorkData.PendingLoopJumps.Find(InExpr))
+					{
+						// EndBlock was just emitted — its instruction index is the last one.
+						const int32 EndBlockInstructionIndex = WorkData.VM->GetByteCode().GetNumInstructions() - 1;
+						for (const FRigVMCompilerWorkData::FPendingLoopJump& PendingJump : *Jumps)
+						{
+							// Relative forward distance from the JumpForward to the EndBlock.
+							WorkData.VM->GetByteCode().GetOpAt<FRigVMJumpOp>(PendingJump.ByteOffset).InstructionIndex =
+								EndBlockInstructionIndex - PendingJump.InstructionIndex;
+						}
+						WorkData.PendingLoopJumps.Remove(InExpr);
+					}
 				}
 
 				// if this is not the completed block - we need to jump back to the control flow instruction
diff --git a/Engine/Plugins/Runtime/RigVM/Source/RigVMDeveloper/Private/RigVMModel/RigVMNode.cpp b/Engine/Plugins/Runtime/RigVM/Source/RigVMDeveloper/Private/RigVMModel/RigVMNode.cpp
--- a/Engine/Plugins/Runtime/RigVM/Source/RigVMDeveloper/Private/RigVMModel/RigVMNode.cpp
+++ b/Engine/Plugins/Runtime/RigVM/Source/RigVMDeveloper/Private/RigVMModel/RigVMNode.cpp
@@ -1371,11 +1371,12 @@ bool URigVMNode::IsLoopNode() const
 {
 	if(IsControlFlowNode())
 	{
-		static const TArray<FName> ExpectedLoopBlocks = {FRigVMStruct::ExecuteContextName, FRigVMStruct::ForLoopCompletedPinName};
 		const TArray<FName>& Blocks = GetControlFlowBlocks();
-		if(Blocks.Num() == ExpectedLoopBlocks.Num())
+		if(Blocks.Num() == 2 && Blocks[1] == FRigVMStruct::ForLoopCompletedPinName)
 		{
-			return Blocks[0] == ExpectedLoopBlocks[0] && Blocks[1] == ExpectedLoopBlocks[1];
+			// The body block is named "ExecuteContext" for FRigVMStruct-based loops
+			// and "ExecutePin" for FRigVMStructMutable-based loops (e.g. ForLoopCount).
+			return Blocks[0] == FRigVMStruct::ExecuteContextName || Blocks[0] == FRigVMStruct::ExecutePinName;
 		}
 	}
 
diff --git a/Engine/Plugins/Runtime/RigVM/Source/RigVMDeveloper/Public/RigVMCompiler/RigVMCompiler.h b/Engine/Plugins/Runtime/RigVM/Source/RigVMDeveloper/Public/RigVMCompiler/RigVMCompiler.h
--- a/Engine/Plugins/Runtime/RigVM/Source/RigVMDeveloper/Public/RigVMCompiler/RigVMCompiler.h
+++ b/Engine/Plugins/Runtime/RigVM/Source/RigVMDeveloper/Public/RigVMCompiler/RigVMCompiler.h
@@ -265,6 +265,23 @@ public:
 	TArray<uint32> LazyBlocksToProcess;
 	TOptional<uint32> CurrentBlockHash;
 
+	/**
+	 * Tracks a JumpForward op emitted for a Break or Continue node that needs
+	 * its relative offset patched once the enclosing loop's EndBlock is emitted.
+	 */
+	struct FPendingLoopJump
+	{
+		int32 ByteOffset;       // Byte offset of the JumpForward op in the bytecode stream (returned by AddJumpOp)
+		int32 InstructionIndex; // Instruction index of the JumpForward op (for relative offset calculation)
+	};
+
+	/**
+	 * Maps each enclosing loop's CallExtern expression to the list of Break/Continue
+	 * JumpForward ops that must be patched to jump to that loop's EndBlock.
+	 * Entries are consumed (removed) when the loop's EndBlock is emitted.
+	 */
+	TMap<const FRigVMCallExternExprAST*, TArray<FPendingLoopJump>> PendingLoopJumps;
+
 	UE_API void ReportInfo(const FString& InMessage) const;
 	UE_API void ReportWarning(const FString& InMessage) const;
 	UE_API void ReportError(const FString& InMessage) const;
@@ -322,6 +339,7 @@ public:
 		PropertyDescriptions.Reset();
 		PropertyPathDescriptions.Reset();
 		DeferredCopyOps.Reset();
+		PendingLoopJumps.Reset();
 	}
 };
 
