Fix ControlRig ExecutionStackView: O(N^2) perf regression in RestoreSparseItemInfos

#rnx

#Synced-CL 52295263

diff --git a/Engine/Plugins/Runtime/RigVM/Source/RigVMEditor/Public/Widgets/SRigVMExecutionStackView.h b/Engine/Plugins/Runtime/RigVM/Source/RigVMEditor/Public/Widgets/SRigVMExecutionStackView.h
--- a/Engine/Plugins/Runtime/RigVM/Source/RigVMEditor/Public/Widgets/SRigVMExecutionStackView.h
+++ b/Engine/Plugins/Runtime/RigVM/Source/RigVMEditor/Public/Widgets/SRigVMExecutionStackView.h
@@ -103,6 +103,14 @@ public:
 		if (!BaseView::SparseItemInfos.IsEmpty())
 		{
 			OldSparseItemInfos = BaseView::SparseItemInfos;
+			// Pre-build a hash-keyed map so RestoreSparseItemInfos can do O(1) lookup
+			// instead of an O(N) linear scan per entry (avoid O(N^2) on large rigs).
+			OldSparseItemInfosByHash.Reset();
+			OldSparseItemInfosByHash.Reserve(OldSparseItemInfos.Num());
+			for (const auto& Pair : OldSparseItemInfos)
+			{
+				OldSparseItemInfosByHash.Add(GetTypeHash(*Pair.Key.Get()), Pair.Value);
+			}
 		}
 		ClearExpandedItems();
 	}
@@ -111,22 +119,21 @@ public:
 	void RestoreSparseItemInfos(TSharedPtr<FRigStackEntry> EntryPtr)
 	{
 		const uint32 NewHash = GetTypeHash(*EntryPtr.Get());
-		for (const auto& Pair : OldSparseItemInfos)
+		if (const FSparseItemInfo* Found = OldSparseItemInfosByHash.Find(NewHash))
 		{
-			const uint32 OldHash = GetTypeHash(*Pair.Key.Get());
-			if (OldHash == NewHash)
-			{
-				// the SparseItemInfos now reference the new element, but keep the same expansion state
-				BaseView::SparseItemInfos.Add(EntryPtr, Pair.Value);
-				return;
-			}
+			// the SparseItemInfos now reference the new element, but keep the same expansion state
+			BaseView::SparseItemInfos.Add(EntryPtr, *Found);
+		}
+		else
+		{
+			// set default state as expanded if not found
+			BaseView::SparseItemInfos.Add(EntryPtr, FSparseItemInfo(true, false));
 		}
-		// set default state as expanded if not found
-		BaseView::SparseItemInfos.Add(EntryPtr, FSparseItemInfo(true, false));
 	}
 
 private:
 	TSparseItemMap OldSparseItemInfos;
+	TMap<uint32, FSparseItemInfo> OldSparseItemInfosByHash;
 };
 
 class SRigVMExecutionStackView : public SCompoundWidget
