[RigVM] Fixed a threadsafety issue with IRigVMEditorAssetInterface::RefreshAllNodes. It is not safe to run it on non-game thread, which can happen if this asset is a BP asset and also a dependency of another BP that is getting async loaded, such as an AnimBP.

#Synced-CL 52307748

diff --git a/Engine/Plugins/Runtime/RigVM/Source/RigVMDeveloper/Private/RigVMEditorAsset.cpp b/Engine/Plugins/Runtime/RigVM/Source/RigVMDeveloper/Private/RigVMEditorAsset.cpp
--- a/Engine/Plugins/Runtime/RigVM/Source/RigVMDeveloper/Private/RigVMEditorAsset.cpp
+++ b/Engine/Plugins/Runtime/RigVM/Source/RigVMDeveloper/Private/RigVMEditorAsset.cpp
@@ -4136,6 +4136,14 @@ void IRigVMEditorAssetInterface::RefreshAllNodes()
 		return;
 	}
 	
+	// The only time we need to refresh EdGraph is when this asset is being opened in an editor,
+	// which can only happen on Game thread. We need a check here because this function can be called from
+	// AsyncLoadingThread while this asset itself is going through PostLoad() on GameThread, 
+	// if this asset is a BP asset and a dependency of another BP that is getting async loaded, such as an AnimBP, 
+	if (!IsInGameThread())
+	{
+		return;
+	}
 	
 	// Avoid refreshing EdGraph if PostLoad() hasn't been called, since reconstruct node later
 	// can access model data that hasn't been fully loaded. And it is ok to skip here because
