RigVM: Force to register types when loading a localized registry

#rnx

#Synced-CL 52302161

diff --git a/Engine/Plugins/Runtime/RigVM/Source/RigVM/Private/RigVMCore/RigVMRegistry.cpp b/Engine/Plugins/Runtime/RigVM/Source/RigVM/Private/RigVMCore/RigVMRegistry.cpp
--- a/Engine/Plugins/Runtime/RigVM/Source/RigVM/Private/RigVMCore/RigVMRegistry.cpp
+++ b/Engine/Plugins/Runtime/RigVM/Source/RigVM/Private/RigVMCore/RigVMRegistry.cpp
@@ -583,7 +583,7 @@ void FRigVMRegistry_NoLock::Serialize_NoLock(FArchive& Ar)
 			{
 				FTypeInfo TypeInfo;
 				TypeInfo.Serialize_NoLock(Ar);
-				(void)FindOrAddType_NoLock(TypeInfo.Type);
+				(void)FindOrAddType_NoLock(TypeInfo.Type, true /* force loading of the type */);
 			}
 		}
 
