[ControlRig] - Added OneMinus node for floats and doubles.

#Synced-CL 52369864

diff --git a/Engine/Plugins/Runtime/RigVM/Source/RigVM/Private/RigVMFunctions/Math/RigVMFunction_MathDouble.cpp b/Engine/Plugins/Runtime/RigVM/Source/RigVM/Private/RigVMFunctions/Math/RigVMFunction_MathDouble.cpp
--- a/Engine/Plugins/Runtime/RigVM/Source/RigVM/Private/RigVMFunctions/Math/RigVMFunction_MathDouble.cpp
+++ b/Engine/Plugins/Runtime/RigVM/Source/RigVM/Private/RigVMFunctions/Math/RigVMFunction_MathDouble.cpp
@@ -108,6 +108,11 @@ FRigVMFunction_MathDoubleAbs_Execute()
 	Result = FMath::Abs(Value);
 }
 
+FRigVMFunction_MathDoubleOneMinus_Execute()
+{
+	Result = 1.0 - Value;
+}
+
 FRigVMFunction_MathDoubleFloor_Execute()
 {
 	Result = FMath::FloorToDouble(Value);
diff --git a/Engine/Plugins/Runtime/RigVM/Source/RigVM/Private/RigVMFunctions/Math/RigVMFunction_MathFloat.cpp b/Engine/Plugins/Runtime/RigVM/Source/RigVM/Private/RigVMFunctions/Math/RigVMFunction_MathFloat.cpp
--- a/Engine/Plugins/Runtime/RigVM/Source/RigVM/Private/RigVMFunctions/Math/RigVMFunction_MathFloat.cpp
+++ b/Engine/Plugins/Runtime/RigVM/Source/RigVM/Private/RigVMFunctions/Math/RigVMFunction_MathFloat.cpp
@@ -108,6 +108,11 @@ FRigVMFunction_MathFloatAbs_Execute()
 	Result = FMath::Abs(Value);
 }
 
+FRigVMFunction_MathFloatOneMinus_Execute()
+{
+	Result = 1.0f - Value;
+}
+
 FRigVMFunction_MathFloatFloor_Execute()
 {
 	Result = FMath::FloorToFloat(Value);
diff --git a/Engine/Plugins/Runtime/RigVM/Source/RigVM/Public/RigVMFunctions/Math/RigVMFunction_MathDouble.h b/Engine/Plugins/Runtime/RigVM/Source/RigVM/Public/RigVMFunctions/Math/RigVMFunction_MathDouble.h
--- a/Engine/Plugins/Runtime/RigVM/Source/RigVM/Public/RigVMFunctions/Math/RigVMFunction_MathDouble.h
+++ b/Engine/Plugins/Runtime/RigVM/Source/RigVM/Public/RigVMFunctions/Math/RigVMFunction_MathDouble.h
@@ -352,6 +352,19 @@ struct FRigVMFunction_MathDoubleAbs : public FRigVMFunction_MathDoubleUnaryOp
 	RIGVM_API virtual void Execute() override;
 };
 
+/**
+ * Returns one minus the input value (1 - x).
+ * Commonly used for inverting normalized values such as blend weights.
+ */
+USTRUCT(meta=(DisplayName="One Minus", TemplateName="OneMinus", Keywords="Invert,Inverse,Subtract,One,Minus,Negate,Flip"))
+struct FRigVMFunction_MathDoubleOneMinus : public FRigVMFunction_MathDoubleUnaryOp
+{
+	GENERATED_BODY()
+
+	RIGVM_METHOD()
+	RIGVM_API virtual void Execute() override;
+};
+
 /**
  * Returns the closest lower full number (integer) of the value
  */
diff --git a/Engine/Plugins/Runtime/RigVM/Source/RigVM/Public/RigVMFunctions/Math/RigVMFunction_MathFloat.h b/Engine/Plugins/Runtime/RigVM/Source/RigVM/Public/RigVMFunctions/Math/RigVMFunction_MathFloat.h
--- a/Engine/Plugins/Runtime/RigVM/Source/RigVM/Public/RigVMFunctions/Math/RigVMFunction_MathFloat.h
+++ b/Engine/Plugins/Runtime/RigVM/Source/RigVM/Public/RigVMFunctions/Math/RigVMFunction_MathFloat.h
@@ -354,6 +354,19 @@ struct FRigVMFunction_MathFloatAbs : public FRigVMFunction_MathFloatUnaryOp
 	UE_API virtual void Execute() override;
 };
 
+/**
+ * Returns one minus the input value (1 - x).
+ * Commonly used for inverting normalized values such as blend weights.
+ */
+USTRUCT(meta=(DisplayName="One Minus", TemplateName="OneMinus", Keywords="Invert,Inverse,Subtract,One,Minus,Negate,Flip"))
+struct FRigVMFunction_MathFloatOneMinus : public FRigVMFunction_MathFloatUnaryOp
+{
+	GENERATED_BODY()
+
+	RIGVM_METHOD()
+	UE_API virtual void Execute() override;
+};
+
 /**
  * Returns the closest lower full number (integer) of the value
  */
