Add debug host object to logging functions within RigVM local registry serialization

#Synced-CL 52402383

diff --git a/Engine/Plugins/Runtime/RigVM/Source/RigVM/Private/RigVMCore/RigVM.cpp b/Engine/Plugins/Runtime/RigVM/Source/RigVM/Private/RigVMCore/RigVM.cpp
--- a/Engine/Plugins/Runtime/RigVM/Source/RigVM/Private/RigVMCore/RigVM.cpp
+++ b/Engine/Plugins/Runtime/RigVM/Source/RigVM/Private/RigVMCore/RigVM.cpp
@@ -214,7 +214,7 @@ void URigVM::Save(FArchive& Ar)
 		int64 ArchiveOffsetForSerializedRegistry = 0;
 		Ar << ArchiveOffsetForSerializedRegistry;
 
-		LocalizedRegistry->Serialize_NoLock(Ar);
+		LocalizedRegistry->Serialize_NoLock(Ar, this);
 
 		// store the archive positions after the registry so we
 		// can skip it during load
@@ -387,7 +387,7 @@ void URigVM::Load(FArchive& Ar)
 			else
 			{
 				LocalizedRegistry = FRigVMRegistry_NoLock::CreateLocalizedRegistry();
-				LocalizedRegistry->Serialize_NoLock(Ar);
+				LocalizedRegistry->Serialize_NoLock(Ar, this);
 			}
 		}
 	}
@@ -3141,7 +3141,7 @@ FString URigVM::GetLocalizedRegistryAsString() const
 	FRigVMObjectArchive Archive;
 	{
 		FRigVMObjectArchiveWriter Writer(Archive, nullptr);
-		LocalizedRegistry->Serialize_NoLock(Writer);
+		LocalizedRegistry->Serialize_NoLock(Writer, this);
 	}
 	Archive.Compress();
 	return Archive.ToString();
@@ -3170,7 +3170,7 @@ void URigVM::SetLocalizedRegistryFromString(const FString& InData)
 	LocalizedRegistry = FRigVMRegistry_NoLock::CreateLocalizedRegistry();
 	{
 		FRigVMObjectArchiveReader Reader(Archive, nullptr);
-		LocalizedRegistry->Serialize_NoLock(Reader);
+		LocalizedRegistry->Serialize_NoLock(Reader, this);
 	}
 }
 
diff --git a/Engine/Plugins/Runtime/RigVM/Source/RigVM/Private/RigVMCore/RigVMFunction.cpp b/Engine/Plugins/Runtime/RigVM/Source/RigVM/Private/RigVMCore/RigVMFunction.cpp
--- a/Engine/Plugins/Runtime/RigVM/Source/RigVM/Private/RigVMCore/RigVMFunction.cpp
+++ b/Engine/Plugins/Runtime/RigVM/Source/RigVM/Private/RigVMCore/RigVMFunction.cpp
@@ -291,7 +291,7 @@ const FName& FRigVMFunction::GetArgumentNameForOperandIndex(int32 InOperandIndex
 	return ArgumentNames[InOperandIndex];
 }
 
-void FRigVMFunction::Serialize_NoLock(FArchive& Ar, FRigVMRegistryHandle& InRegistry)
+void FRigVMFunction::Serialize_NoLock(FArchive& Ar, FRigVMRegistryHandle& InRegistry, const UObject* Owner)
 {
 	if(Ar.IsSaving())
 	{
@@ -338,13 +338,13 @@ void FRigVMFunction::Serialize_NoLock(FArchive& Ar, FRigVMRegistryHandle& InRegi
 #if WITH_EDITOR
 				checkf(GlobalFunction->Struct == Struct, TEXT("Struct mismatch when looking for function '%s'. GlobalFunction->Struct:%s Struct:%s"), *Name, *GetPathNameSafe(GlobalFunction->Struct), *GetPathNameSafe(Struct));
 
-				auto CompareArguments = [GlobalFunction](const FRigVMFunctionArgument& SerializedArg, const FRigVMFunctionArgument& CurrentArg)
+				auto CompareArguments = [GlobalFunction, Owner](const FRigVMFunctionArgument& SerializedArg, const FRigVMFunctionArgument& CurrentArg)
 				{
 					if (SerializedArg == CurrentArg)
 					{
 						return;
 					}
-					UE_LOG(LogRigVM, Error, TEXT("Localized Registry Function '%s' mismatch for argument '%s'. Please outdate function instead of changing them in the code and roll back potential changes."), *GlobalFunction->Struct->GetStructCPPName(), CurrentArg.Name);
+					UE_LOG(LogRigVM, Error, TEXT("Localized Registry Function '%s' mismatch for argument '%s'. Please outdate function instead of changing them in the code and roll back potential changes. Owner: [%s]"), *GlobalFunction->Struct->GetStructCPPName(), CurrentArg.Name, Owner ? *Owner->GetPathName() : TEXT("null"));
 				};
 
 				// do a rough pass to see of all arguments match.
@@ -375,7 +375,7 @@ void FRigVMFunction::Serialize_NoLock(FArchive& Ar, FRigVMRegistryHandle& InRegi
 						const int32* GlobalArgIndex = GlobalArgumentToIndex.Find(ArgName);
 						if (!GlobalArgIndex)
 						{
-							UE_LOG(LogRigVM, Error, TEXT("Localized Registry Function '%s' is missing formerly existing argument '%s'. Please outdate function instead of changing them in the code and roll back potential changes."), *GlobalFunction->Struct->GetStructCPPName(), *ArgName.ToString());
+							UE_LOG(LogRigVM, Error, TEXT("Localized Registry Function '%s' is missing formerly existing argument '%s'. Please outdate function instead of changing them in the code and roll back potential changes. Owner: [%s]"), *GlobalFunction->Struct->GetStructCPPName(), *ArgName.ToString(), Owner ? *Owner->GetPathName() : TEXT("null"));
 							continue;
 						}
 						CompareArguments(Arguments[SerializedArgIndex], GlobalFunction->Arguments[*GlobalArgIndex]);
diff --git a/Engine/Plugins/Runtime/RigVM/Source/RigVM/Private/RigVMCore/RigVMRegistry.cpp b/Engine/Plugins/Runtime/RigVM/Source/RigVM/Private/RigVMCore/RigVMRegistry.cpp
--- a/Engine/Plugins/Runtime/RigVM/Source/RigVM/Private/RigVMCore/RigVMRegistry.cpp
+++ b/Engine/Plugins/Runtime/RigVM/Source/RigVM/Private/RigVMCore/RigVMRegistry.cpp
@@ -564,7 +564,7 @@ void FRigVMRegistry_NoLock::InitializeBaseTypes_NoLock()
 	}
 }
 
-void FRigVMRegistry_NoLock::Serialize_NoLock(FArchive& Ar)
+void FRigVMRegistry_NoLock::Serialize_NoLock(FArchive& Ar, const UObject* Owner)
 {
 	check(!bIsGlobalRegistry);
 
@@ -657,7 +657,7 @@ void FRigVMRegistry_NoLock::Serialize_NoLock(FArchive& Ar)
 			for(int32 FunctionIndex = 0; FunctionIndex < NumFunctions; FunctionIndex++)
 			{
 				FRigVMFunction Function;
-				Function.Serialize_NoLock(Ar, ThisHandle);
+				Function.Serialize_NoLock(Ar, ThisHandle, Owner);
 
 				if(Function.Struct)
 				{
@@ -730,7 +730,7 @@ void FRigVMRegistry_NoLock::Serialize_NoLock(FArchive& Ar)
 			Ar << NumFunctions;
 			for(FRigVMFunction& Function : Functions)
 			{
-				Function.Serialize_NoLock(Ar, ThisHandle);
+				Function.Serialize_NoLock(Ar, ThisHandle, Owner);
 			}
 		}
 	}
diff --git a/Engine/Plugins/Runtime/RigVM/Source/RigVM/Public/RigVMCore/RigVMFunction.h b/Engine/Plugins/Runtime/RigVM/Source/RigVM/Public/RigVMCore/RigVMFunction.h
--- a/Engine/Plugins/Runtime/RigVM/Source/RigVM/Public/RigVMCore/RigVMFunction.h
+++ b/Engine/Plugins/Runtime/RigVM/Source/RigVM/Public/RigVMCore/RigVMFunction.h
@@ -203,7 +203,7 @@ struct FRigVMFunction
 	RIGVM_API bool SupportsExecuteContextStruct(const UScriptStruct* InExecuteContextStruct) const;
 	RIGVM_API bool SupportsExecuteContextStruct_NoLock(const UScriptStruct* InExecuteContextStruct, const FRigVMRegistryHandle& InRegistry) const;
 	RIGVM_API const FName& GetArgumentNameForOperandIndex(int32 InOperandIndex, int32 InTotalOperands, FRigVMRegistryHandle& InRegistry) const;
-	RIGVM_API void Serialize_NoLock(FArchive& Ar, FRigVMRegistryHandle& InRegistry);
+	RIGVM_API void Serialize_NoLock(FArchive& Ar, FRigVMRegistryHandle& InRegistry, const UObject* Owner);
 
 #if WITH_EDITOR
 	RIGVM_API bool IsFixedSizeArray(const FName& InArgumentName) const;
diff --git a/Engine/Plugins/Runtime/RigVM/Source/RigVM/Public/RigVMCore/RigVMRegistry.h b/Engine/Plugins/Runtime/RigVM/Source/RigVM/Public/RigVMCore/RigVMRegistry.h
--- a/Engine/Plugins/Runtime/RigVM/Source/RigVM/Public/RigVMCore/RigVMRegistry.h
+++ b/Engine/Plugins/Runtime/RigVM/Source/RigVM/Public/RigVMCore/RigVMRegistry.h
@@ -162,7 +162,7 @@ public:
 	UE_API virtual void Reset_NoLock();
 
 	// Serializes the registry to an archive. This throws for the global registry.
-	UE_API void Serialize_NoLock(FArchive& Ar);
+	UE_API void Serialize_NoLock(FArchive& Ar, const UObject* Owner);
 
 	// Adds a type if it doesn't exist yet and returns its index.
 	// This function is thread-safe
