Add missing owner object pointer to EngineTest to match new signature
Make owner object pointer null by default

#Synced-CL 52402383

diff --git a/Engine/Plugins/Runtime/RigVM/Source/RigVM/Public/RigVMCore/RigVMFunction.h b/Engine/Plugins/Runtime/RigVM/Source/RigVM/Public/RigVMCore/RigVMFunction.h
--- a/Engine/Plugins/Runtime/RigVM/Source/RigVM/Public/RigVMCore/RigVMFunction.h
+++ b/Engine/Plugins/Runtime/RigVM/Source/RigVM/Public/RigVMCore/RigVMFunction.h
@@ -203,7 +203,7 @@ struct FRigVMFunction
 	RIGVM_API bool SupportsExecuteContextStruct(const UScriptStruct* InExecuteContextStruct) const;
 	RIGVM_API bool SupportsExecuteContextStruct_NoLock(const UScriptStruct* InExecuteContextStruct, const FRigVMRegistryHandle& InRegistry) const;
 	RIGVM_API const FName& GetArgumentNameForOperandIndex(int32 InOperandIndex, int32 InTotalOperands, FRigVMRegistryHandle& InRegistry) const;
-	RIGVM_API void Serialize_NoLock(FArchive& Ar, FRigVMRegistryHandle& InRegistry, const UObject* Owner);
+	RIGVM_API void Serialize_NoLock(FArchive& Ar, FRigVMRegistryHandle& InRegistry, const UObject* Owner = nullptr);
 
 #if WITH_EDITOR
 	RIGVM_API bool IsFixedSizeArray(const FName& InArgumentName) const;
diff --git a/Engine/Plugins/Runtime/RigVM/Source/RigVM/Public/RigVMCore/RigVMRegistry.h b/Engine/Plugins/Runtime/RigVM/Source/RigVM/Public/RigVMCore/RigVMRegistry.h
--- a/Engine/Plugins/Runtime/RigVM/Source/RigVM/Public/RigVMCore/RigVMRegistry.h
+++ b/Engine/Plugins/Runtime/RigVM/Source/RigVM/Public/RigVMCore/RigVMRegistry.h
@@ -162,7 +162,7 @@ public:
 	UE_API virtual void Reset_NoLock();
 
 	// Serializes the registry to an archive. This throws for the global registry.
-	UE_API void Serialize_NoLock(FArchive& Ar, const UObject* Owner);
+	UE_API void Serialize_NoLock(FArchive& Ar, const UObject* Owner = nullptr);
 
 	// Adds a type if it doesn't exist yet and returns its index.
 	// This function is thread-safe
