foldLiterals fix for any node instance being directly manipulated

#Synced-CL 52518730

diff --git a/Engine/Plugins/Animation/ControlRig/Source/ControlRigEditor/Private/ControlRigEditorModule.cpp b/Engine/Plugins/Animation/ControlRig/Source/ControlRigEditor/Private/ControlRigEditorModule.cpp
--- a/Engine/Plugins/Animation/ControlRig/Source/ControlRigEditor/Private/ControlRigEditorModule.cpp
+++ b/Engine/Plugins/Animation/ControlRig/Source/ControlRigEditor/Private/ControlRigEditorModule.cpp
@@ -1012,10 +1012,13 @@ void FControlRigEditorModule::GetDirectManipulationMenuActions(IRigVMClientHost*
 							FSlateIcon(),
 							FUIAction(FExecuteAction::CreateLambda([ControlRigBlueprint, UnitNode, Target]() {
 
-								// disable literal folding for the moment
-								if(ControlRigBlueprint->VMCompileSettings.ASTSettings.bFoldLiterals)
+								// Mark this node as interactive so the compiler gives it its own literal
+								// registers rather than sharing them with other nodes that have the same
+								// default values. Only recompile if the node is not already in the set.
+								auto& ASTSettings = ControlRigBlueprint->VMCompileSettings.ASTSettings;
+								if(!ASTSettings.InteractiveNodes.Contains(FObjectKey(UnitNode)))
 								{
-									ControlRigBlueprint->VMCompileSettings.ASTSettings.bFoldLiterals = false;
+									ASTSettings.InteractiveNodes.Add(FObjectKey(UnitNode));
 									ControlRigBlueprint->RecompileVM();
 								}
 
diff --git a/Engine/Plugins/Animation/ControlRig/Source/ControlRigEditor/Private/EditMode/ControlRigEditMode.cpp b/Engine/Plugins/Animation/ControlRig/Source/ControlRigEditor/Private/EditMode/ControlRigEditMode.cpp
--- a/Engine/Plugins/Animation/ControlRig/Source/ControlRigEditor/Private/EditMode/ControlRigEditMode.cpp
+++ b/Engine/Plugins/Animation/ControlRig/Source/ControlRigEditor/Private/EditMode/ControlRigEditMode.cpp
@@ -758,6 +758,31 @@ void FControlRigEditMode::Tick(FEditorViewportClient* ViewportClient, float Delt
 		OnControlRigShapeActorsRecreatedDelegate.Broadcast();
 	}
 
+	// Ensure transient control shape actors are selected while DM is active.
+	// RecreateControlShapeActors() may be called from paths outside this Tick
+	// (e.g. OnPostConstruction_AnyThread), which resets RecreateControlShapesRequired
+	// before the selection loop above can run, leaving the shape actor unselected.
+	// ComputePivotFromShapeActors gates entirely on IsSelected(), so an unselected
+	// shape actor causes the gizmo to fall back to the world origin.
+	for (TWeakObjectPtr<UControlRig>& RuntimeRigPtr : RuntimeControlRigs)
+	{
+		if (UControlRig* ControlRig = RuntimeRigPtr.Get())
+		{
+			TArray<FRigElementKey> SelectedRigElements = GetSelectedRigElements(ControlRig);
+			for (const FRigElementKey& SelectedKey : SelectedRigElements)
+			{
+				if (SelectedKey.Type == ERigElementType::Control)
+				{
+					AControlRigShapeActor* ShapeActor = GetControlShapeFromControlName(ControlRig, SelectedKey.Name);
+					if (ShapeActor && !ShapeActor->IsSelected())
+					{
+						ShapeActor->SetSelected(true);
+					}
+				}
+			}
+		}
+	}
+
 	{
 		// We need to tick here since changing a bone for example
 		// might have changed the transform of the Control
diff --git a/Engine/Plugins/Animation/ControlRig/Source/ControlRigEditor/Private/Editor/ControlRigEditor.cpp b/Engine/Plugins/Animation/ControlRig/Source/ControlRigEditor/Private/Editor/ControlRigEditor.cpp
--- a/Engine/Plugins/Animation/ControlRig/Source/ControlRigEditor/Private/Editor/ControlRigEditor.cpp
+++ b/Engine/Plugins/Animation/ControlRig/Source/ControlRigEditor/Private/Editor/ControlRigEditor.cpp
@@ -4168,6 +4168,18 @@ bool FControlRigBaseEditor::HandleRequestDirectManipulation(ERigControlType InCo
 
 			if(FControlRigAssetInterfacePtr Blueprint = GetControlRigAssetInterface())
 			{
+				// Mark this node as interactive so the compiler gives it its own literal
+				// registers rather than sharing them with other nodes that have the same
+				// default values. Only recompile if the node is not already in the set.
+				if(DirectManipulationSubject.IsValid())
+				{
+					auto& ASTSettings = Blueprint->GetRigVMAssetInterface()->GetVMCompileSettings().ASTSettings;
+					if(!ASTSettings.InteractiveNodes.Contains(FObjectKey(DirectManipulationSubject.Get())))
+					{
+						ASTSettings.InteractiveNodes.Add(FObjectKey(DirectManipulationSubject.Get()));
+						Blueprint->GetRigVMAssetInterface()->RecompileVM();
+					}
+				}
 				Blueprint->AddTransientControl(DirectManipulationSubject.Get(), Target);
 			}
 			return true;
@@ -4243,11 +4255,16 @@ void FControlRigBaseEditor::OnDirectManipulationChanged(TSharedPtr<FString> NewV
 		return;
 	}
 
-	// disable literal folding for the moment
-	if(ControlRigBlueprint->GetRigVMAssetInterface()->GetVMCompileSettings().ASTSettings.bFoldLiterals)
+	// Mark this node as interactive so the compiler gives it its own literal
+	// registers rather than sharing them with other nodes that have the same
+	// default values. Only recompile if the node is not already in the set.
 	{
-		ControlRigBlueprint->GetRigVMAssetInterface()->GetVMCompileSettings().ASTSettings.bFoldLiterals = false;
-		ControlRigBlueprint->GetRigVMAssetInterface()->RecompileVM();
+		auto& ASTSettings = ControlRigBlueprint->GetRigVMAssetInterface()->GetVMCompileSettings().ASTSettings;
+		if(!ASTSettings.InteractiveNodes.Contains(FObjectKey(UnitNode)))
+		{
+			ASTSettings.InteractiveNodes.Add(FObjectKey(UnitNode));
+			ControlRigBlueprint->GetRigVMAssetInterface()->RecompileVM();
+		}
 	}
 
 	const FString& DesiredTarget = *NewValue.Get();
diff --git a/Engine/Plugins/Runtime/RigVM/Source/RigVMDeveloper/Private/RigVMCompiler/RigVMAST.cpp b/Engine/Plugins/Runtime/RigVM/Source/RigVMDeveloper/Private/RigVMCompiler/RigVMAST.cpp
--- a/Engine/Plugins/Runtime/RigVM/Source/RigVMDeveloper/Private/RigVMCompiler/RigVMAST.cpp
+++ b/Engine/Plugins/Runtime/RigVM/Source/RigVMDeveloper/Private/RigVMCompiler/RigVMAST.cpp
@@ -2550,6 +2550,24 @@ void FRigVMParserAST::FoldLiterals()
 	TMap<FString, FRigVMLiteralExprAST*> ValueToLiteral;
 	TArray<FRigVMExprAST*> ExpressionsToRemove;
 
+	// Returns true if the literal's owning node is marked as interactive (e.g. has DM active).
+	// Interactive nodes must not share literal registers with other nodes.
+	auto IsInteractive = [this](const FRigVMLiteralExprAST* InLiteral) -> bool
+	{
+		if (Settings.InteractiveNodes.IsEmpty())
+		{
+			return false;
+		}
+		if (const URigVMPin* Pin = InLiteral->GetPin())
+		{
+			if (const URigVMNode* Node = Pin->GetNode())
+			{
+				return Settings.InteractiveNodes.Contains(FObjectKey(Node));
+			}
+		}
+		return false;
+	};
+
 	for (int32 ExpressionIndex = 0; ExpressionIndex < Expressions.Num(); ExpressionIndex++)
 	{
 		FRigVMExprAST* Expression = Expressions[ExpressionIndex];
@@ -2593,16 +2611,26 @@ void FRigVMParserAST::FoldLiterals()
 			FRigVMLiteralExprAST* const* MappedExpr = ValueToLiteral.Find(Hash);
 			if (MappedExpr)
 			{
-				TArray<FRigVMExprAST*> Parents = Expression->Parents;
-				for (FRigVMExprAST* Parent : Parents)
+				// Skip folding if either this literal or the canonical belongs to an
+				// interactive node — each interactive node must keep its own register.
+				if (!IsInteractive(LiteralExpr) && !IsInteractive(*MappedExpr))
 				{
-					Parent->ReplaceChild(Expression, *MappedExpr);
+					TArray<FRigVMExprAST*> Parents = Expression->Parents;
+					for (FRigVMExprAST* Parent : Parents)
+					{
+						Parent->ReplaceChild(Expression, *MappedExpr);
+					}
+					ExpressionsToRemove.AddUnique(Expression);
 				}
-				ExpressionsToRemove.AddUnique(Expression);
 			}
 			else
 			{
-				ValueToLiteral.Add(Hash, LiteralExpr);
+				// Don't use an interactive node's literal as the canonical for a value —
+				// doing so would allow non-interactive nodes to share its register.
+				if (!IsInteractive(LiteralExpr))
+				{
+					ValueToLiteral.Add(Hash, LiteralExpr);
+				}
 			}
 		}
 	}
diff --git a/Engine/Plugins/Runtime/RigVM/Source/RigVMDeveloper/Private/RigVMCompiler/RigVMCompiler.cpp b/Engine/Plugins/Runtime/RigVM/Source/RigVMDeveloper/Private/RigVMCompiler/RigVMCompiler.cpp
--- a/Engine/Plugins/Runtime/RigVM/Source/RigVMDeveloper/Private/RigVMCompiler/RigVMCompiler.cpp
+++ b/Engine/Plugins/Runtime/RigVM/Source/RigVMDeveloper/Private/RigVMCompiler/RigVMCompiler.cpp
@@ -185,12 +185,18 @@ FRigVMOperand FRigVMCompilerWorkData::FindOrAddProperty(const ERigVMMemoryType I
 	FString NewName = InNewName;
 	int32 Suffix = 0;
 
-	// share literals only when literal folding is enabled (disabled during direct manipulation to keep per-node registers separate)
+	// share literals only when literal folding is globally enabled
 	if(Settings.ASTSettings.bFoldLiterals && InMemoryType == ERigVMMemoryType::Literal && !bMakeUniqueProperty && PropertyDescriptions.Contains(ERigVMMemoryType::Literal))
 	{
-		const TArray<FRigVMPropertyDescription>& LiteralDescriptions = PropertyDescriptions[ERigVMMemoryType::Literal]; 
+		const TArray<FRigVMPropertyDescription>& LiteralDescriptions = PropertyDescriptions[ERigVMMemoryType::Literal];
 		for(int32 Index = 0; Index < LiteralDescriptions.Num(); Index++)
 		{
+			// skip registers owned by interactive nodes — sharing with them would cause DM
+			// writes to bleed into nodes that are not currently being manipulated.
+			if(InteractiveLiteralIndices.Contains(Index))
+			{
+				continue;
+			}
 			const FRigVMPropertyDescription& ExistingProperty = LiteralDescriptions[Index];
 			if(ExistingProperty.CPPType.Equals(InProperty.CPPType, ESearchCase::IgnoreCase))
 			{
@@ -5575,8 +5581,21 @@ FRigVMOperand URigVMCompiler::FindOrAddRegister(const FRigVMVarExprAST* InVarExp
 			}
 		}
 
-		const bool bMakeUniqueProperty = MemoryType == ERigVMMemoryType::Work;
+		// Force a unique literal register for interactive nodes (e.g. nodes with DM active) so
+		// that writing to the register during manipulation does not affect other nodes that
+		// would otherwise share the same literal register due to having the same default value.
+		const bool bIsInteractiveLiteral = bIsLiteral
+			&& !WorkData.Settings.ASTSettings.InteractiveNodes.IsEmpty()
+			&& WorkData.Settings.ASTSettings.InteractiveNodes.Contains(FObjectKey(Pin->GetNode()));
+		const bool bMakeUniqueProperty = (MemoryType == ERigVMMemoryType::Work) || bIsInteractiveLiteral;
 		Operand = WorkData.FindOrAddProperty(MemoryType, RegisterName, CPPType, Pin->GetCPPTypeObject(), JoinedDefaultValue, bMakeUniqueProperty);
+
+		// Record the index so FindOrAddProperty's sharing scan skips it for non-interactive
+		// callers — preventing other nodes from being mapped to this node's register.
+		if(bIsInteractiveLiteral && Operand.IsValid())
+		{
+			WorkData.InteractiveLiteralIndices.Add(Operand.GetRegisterIndex());
+		}
 	}
 	ensure(Operand.IsValid());
 
diff --git a/Engine/Plugins/Runtime/RigVM/Source/RigVMDeveloper/Public/RigVMCompiler/RigVMAST.h b/Engine/Plugins/Runtime/RigVM/Source/RigVMDeveloper/Public/RigVMCompiler/RigVMAST.h
--- a/Engine/Plugins/Runtime/RigVM/Source/RigVMDeveloper/Public/RigVMCompiler/RigVMAST.h
+++ b/Engine/Plugins/Runtime/RigVM/Source/RigVMDeveloper/Public/RigVMCompiler/RigVMAST.h
@@ -1135,6 +1135,13 @@ struct FRigVMParserASTSettings
 	UPROPERTY()
 	TArray<TObjectPtr<URigVMLink>> LinksToSkip;
 
+	// nodes whose input literals must not be shared with other nodes' literals.
+	// used by Direct Manipulation to keep per-node registers during interactive editing
+	// so that writing to one node's literal register does not affect other nodes.
+	// stored as FObjectKey rather than FName to correctly handle duplicate node names
+	// inside collapsed graphs. not a UPROPERTY — session-only, intentionally not serialized.
+	TSet<FObjectKey> InteractiveNodes;
+
 	FRigVMReportDelegate ReportDelegate;
 
 	UPROPERTY(VisibleAnywhere, Category = "AST", transient)
diff --git a/Engine/Plugins/Runtime/RigVM/Source/RigVMDeveloper/Public/RigVMCompiler/RigVMCompiler.h b/Engine/Plugins/Runtime/RigVM/Source/RigVMDeveloper/Public/RigVMCompiler/RigVMCompiler.h
--- a/Engine/Plugins/Runtime/RigVM/Source/RigVMDeveloper/Public/RigVMCompiler/RigVMCompiler.h
+++ b/Engine/Plugins/Runtime/RigVM/Source/RigVMDeveloper/Public/RigVMCompiler/RigVMCompiler.h
@@ -199,6 +199,12 @@ public:
 	TMap<ERigVMMemoryType, TArray<FRigVMPropertyPathDescription>> PropertyPathDescriptions;
 	TMap<ERigVMMemoryType, TArray<FRigVMPropertyDescription>> PropertyDescriptions;
 
+	// literal register indices that belong to interactive nodes (e.g. nodes with DM active).
+	// the FindOrAddProperty sharing scan skips these so that non-interactive nodes cannot
+	// be mapped to an interactive node's register, which would cause DM writes to bleed
+	// into nodes that are not being manipulated.
+	TSet<int32> InteractiveLiteralIndices;
+
 	UE_API FRigVMOperand AddProperty(ERigVMMemoryType InMemoryType, const FName& InName, const FString& InCPPType, UObject* InCPPTypeObject, const FString& InDefaultValue = FString());
 	UE_API FRigVMOperand FindProperty(ERigVMMemoryType InMemoryType, const FName& InName) const;
 	UE_API FRigVMOperand FindOrAddProperty(const ERigVMMemoryType InMemoryType, const FRigVMFunctionCompilationPropertyDescription& InProperty, const FString& InNewName, const bool bMakeUniqueProperty);
