Control Rig: Fix ensure when deleting a module asset that is being used in a modular rig
#rnx

#Synced-CL 52096374

diff --git a/Engine/Plugins/Runtime/RigVM/Source/RigVM/Private/RigVMRuntimeAsset.cpp b/Engine/Plugins/Runtime/RigVM/Source/RigVM/Private/RigVMRuntimeAsset.cpp
--- a/Engine/Plugins/Runtime/RigVM/Source/RigVM/Private/RigVMRuntimeAsset.cpp
+++ b/Engine/Plugins/Runtime/RigVM/Source/RigVM/Private/RigVMRuntimeAsset.cpp
@@ -140,6 +140,20 @@ void URigVMRuntimeAsset::PostDuplicate(bool bDuplicateForPIE)
 #endif
 }
 
+void URigVMRuntimeAsset::MarkInstancesAsGarbage()
+{
+	// Clean up archetype instances before destroying the asset
+	TArray<UObject*> InstancesToDestroy = GetArchetypeInstances(false);
+	for (UObject* Instance : InstancesToDestroy)
+	{
+		if (IsValid(Instance))
+		{
+			Instance->MarkAsGarbage();
+		}
+	}
+	ArchetypeInstances.Empty();
+}
+
 TArray<FRigVMExternalVariable> URigVMRuntimeAsset::GetExternalVariables() 
 {
 	TArray<FRigVMExternalVariable> ExternalVariables;
diff --git a/Engine/Plugins/Runtime/RigVM/Source/RigVM/Public/RigVMRuntimeAsset.h b/Engine/Plugins/Runtime/RigVM/Source/RigVM/Public/RigVMRuntimeAsset.h
--- a/Engine/Plugins/Runtime/RigVM/Source/RigVM/Public/RigVMRuntimeAsset.h
+++ b/Engine/Plugins/Runtime/RigVM/Source/RigVM/Public/RigVMRuntimeAsset.h
@@ -111,6 +111,8 @@ protected:
 	UE_API virtual void PostLoad() override;
 	UE_API virtual void PreDuplicate(FObjectDuplicationParameters& DupParams) override;
 	UE_API virtual void PostDuplicate(bool bDuplicateForPIE) override;
+	
+	UE_API virtual void MarkInstancesAsGarbage();
 
 	UFUNCTION(BlueprintCallable, Category = "VM")
 	virtual UClass* GetRigVMHostClass() const override { return URigVMHost::StaticClass(); }
@@ -295,6 +297,7 @@ protected:
 	friend class FRigVMAssetVariablesChangeScope;
 	friend struct FControlRigAssetStrongReference;
 	friend class URigVMEditorAsset;
+	friend class FRigVMEditorModule;
 };
 
 #if WITH_EDITOR
diff --git a/Engine/Plugins/Runtime/RigVM/Source/RigVMEditor/Private/RigVMEditorModule.cpp b/Engine/Plugins/Runtime/RigVM/Source/RigVMEditor/Private/RigVMEditorModule.cpp
--- a/Engine/Plugins/Runtime/RigVM/Source/RigVMEditor/Private/RigVMEditorModule.cpp
+++ b/Engine/Plugins/Runtime/RigVM/Source/RigVMEditor/Private/RigVMEditorModule.cpp
@@ -199,6 +199,7 @@ void FRigVMEditorModule::StartupModuleCommon()
 	PropertiesToUnregisterOnShutdown.Add(FRigVMPythonSettings::StaticStruct()->GetFName());
 	PropertyEditorModule.RegisterCustomPropertyTypeLayout(PropertiesToUnregisterOnShutdown.Last(), FOnGetPropertyTypeCustomizationInstance::CreateStatic(&FRigVMPythonLogDetails::MakeInstance));
 
+	AssetsPreDeleteDelegate = FEditorDelegates::OnAssetsPreDelete.AddStatic(&FRigVMEditorModule::HandleAssetsPreDelete);
 }
 
 void FRigVMEditorModule::ShutdownModule()
@@ -249,6 +250,8 @@ void FRigVMEditorModule::ShutdownModuleCommon()
 
 	// Unregister to fixup newly created BPs
 	FKismetEditorUtilities::UnregisterAutoBlueprintNodeCreation(this);
+	
+	FEditorDelegates::OnAssetsPreDelete.Remove(AssetsPreDeleteDelegate);
 }
 
 void FRigVMEditorModule::GetTypeActions(FRigVMEditorAssetInterfacePtr RigVMBlueprint, FBlueprintActionDatabaseRegistrar& ActionRegistrar)
@@ -2719,4 +2722,15 @@ bool FRigVMEditorModule::IsRigVMEditorModuleBase() const
 	return false;
 }
 
+void FRigVMEditorModule::HandleAssetsPreDelete(const TArray<UObject*>& InAssetsToDelete)
+{
+	for (UObject* AssetToDelete : InAssetsToDelete)
+	{
+		if (URigVMRuntimeAsset* Asset = Cast<URigVMRuntimeAsset>(AssetToDelete))
+		{
+			Asset->MarkInstancesAsGarbage();
+		}
+	}
+}
+
 #undef LOCTEXT_NAMESPACE
diff --git a/Engine/Plugins/Runtime/RigVM/Source/RigVMEditor/Public/RigVMEditorModule.h b/Engine/Plugins/Runtime/RigVM/Source/RigVMEditor/Public/RigVMEditorModule.h
--- a/Engine/Plugins/Runtime/RigVM/Source/RigVMEditor/Public/RigVMEditorModule.h
+++ b/Engine/Plugins/Runtime/RigVM/Source/RigVMEditor/Public/RigVMEditorModule.h
@@ -216,6 +216,9 @@ private:
 
 	/** Event broadcast when finding references for a node is requested */
 	static UE::RigVMEditor::FRigVMEditorOnRequestFindNodeReferencesDelegate OnRequestFindNodeReferences;
+	
+	FDelegateHandle AssetsPreDeleteDelegate;
+	static void HandleAssetsPreDelete(const TArray<UObject*>& InAssetsToDelete);
 
 	friend class URigVMEdGraphNode;
 };
