Control Rig: Dragging from a node and pressing F now spawns a for each node

#Synced-CL 52510584

diff --git a/Engine/Plugins/Runtime/RigVM/Source/RigVMEditor/Private/Editor/RigVMEditor.cpp b/Engine/Plugins/Runtime/RigVM/Source/RigVMEditor/Private/Editor/RigVMEditor.cpp
--- a/Engine/Plugins/Runtime/RigVM/Source/RigVMEditor/Private/Editor/RigVMEditor.cpp
+++ b/Engine/Plugins/Runtime/RigVM/Source/RigVMEditor/Private/Editor/RigVMEditor.cpp
@@ -25,6 +25,7 @@
 #include "Kismet2/BlueprintEditorUtils.h"
 #include "Kismet2/KismetDebugUtilities.h"
 #include "Kismet2/WatchedPin.h"
+#include "RigVMFunctions/RigVMDispatch_Array.h"
 #include "RigVMFunctions/RigVMFunction_ControlFlow.h"
 #include "RigVMModel/Nodes/RigVMAggregateNode.h"
 #include "Widgets/SRigVMBulkEditDialog.h"
@@ -859,6 +860,10 @@ FReply FRigVMEditorBase::OnSpawnGraphNodeByShortcut(FInputChord InChord, const F
 				{
 					Controller->AddUnitNode(FRigVMFunction_ControlFlowBranch::StaticStruct(), FRigVMStruct::ExecuteName, FDeprecateSlateVector2D(InPosition), FString(), true, true);
 				}
+				else if (InChord.Key == EKeys::F)
+				{
+					URigVMTemplateNode* NewNode = Controller->AddTemplateNode(FRigVMDispatch_ArrayIterator().GetTemplateNotation(), FDeprecateSlateVector2D(InPosition), FString(), true, true);
+				}
 			}
 		}
 	}
