Control Rig: Fix rename variable affecting duplicated variables
#rnx

#Synced-CL 52569456

diff --git a/Engine/Plugins/Runtime/RigVM/Source/RigVMDeveloper/Private/RigVMEditorAsset.cpp b/Engine/Plugins/Runtime/RigVM/Source/RigVMDeveloper/Private/RigVMEditorAsset.cpp
--- a/Engine/Plugins/Runtime/RigVM/Source/RigVMDeveloper/Private/RigVMEditorAsset.cpp
+++ b/Engine/Plugins/Runtime/RigVM/Source/RigVMDeveloper/Private/RigVMEditorAsset.cpp
@@ -4838,6 +4838,7 @@ bool URigVMEditorAsset::OnPasteVariable(const FString& InText)
 				FScopedTransaction Transaction(FText::Format(LOCTEXT("PasteVariable", "Paste Variable: {0}"), FText::FromName(NewVar.Name)));
 
 				NewVar.Name = FRigVMBlueprintUtils::FindUniqueVariableName(this, NewVar.Name.ToString());
+				NewVar.Guid = FGuid::NewGuid();
 				FRigVMExternalVariable ExternalVariable = RigVMVariableUtils::ExternalVariableFromRigVMVariableDescription(NewVar);
 				FName Result = RuntimeAsset->AddHostMemberVariableFromExternal(ExternalVariable, NewVar.DefaultValue);
 
diff --git a/Engine/Plugins/Runtime/RigVM/Source/RigVMEditor/Private/Widgets/SRigVMEditorGraphExplorer.cpp b/Engine/Plugins/Runtime/RigVM/Source/RigVMEditor/Private/Widgets/SRigVMEditorGraphExplorer.cpp
--- a/Engine/Plugins/Runtime/RigVM/Source/RigVMEditor/Private/Widgets/SRigVMEditorGraphExplorer.cpp
+++ b/Engine/Plugins/Runtime/RigVM/Source/RigVMEditor/Private/Widgets/SRigVMEditorGraphExplorer.cpp
@@ -1312,7 +1312,7 @@ TSharedPtr<SWidget> SRigVMEditorGraphExplorer::OnContextMenuOpening()
 		return TSharedPtr<SWidget>();
 	}
 
-	FMenuBuilder MenuBuilder(/*Close after selection*/true, Editor->GetToolkitCommands());
+	FMenuBuilder MenuBuilder(/*Close after selection*/true, UICommandList);
 	
 	TArray<FRigVMExplorerElementKey> Selection = TreeView->GetSelectedKeys();
 
@@ -1706,8 +1706,13 @@ void SRigVMEditorGraphExplorer::OnPasteLocalVariable()
 		FRigVMGraphVariableDescription::StaticStruct()->ImportText(Import, &Description, nullptr, PPF_None, &Errors, FRigVMGraphVariableDescription::StaticStruct()->GetName());
 		if (Errors.IsEmpty())
 		{
-			ensure(false);
-			//Editor->OnPasteNewLocalVariable(Description);
+			if (URigVMGraph* FocusedGraph = Editor->GetFocusedModel())
+			{
+				if (URigVMController* Controller = Blueprint->GetRigVMClient()->GetController(FocusedGraph))
+				{
+					Controller->AddLocalVariable(Description.Name, Description.CPPType, Description.CPPTypeObject, Description.DefaultValue, true, true);
+				}
+			}
 			Refresh();
 		}
 	}
