Control Rig: Fix empty graph name when multiple events exist in the same graph
#rnx

#Synced-CL 52569794

diff --git a/Engine/Plugins/Runtime/RigVM/Source/RigVMDeveloper/Private/EdGraph/RigVMEdGraphSchema.cpp b/Engine/Plugins/Runtime/RigVM/Source/RigVMDeveloper/Private/EdGraph/RigVMEdGraphSchema.cpp
--- a/Engine/Plugins/Runtime/RigVM/Source/RigVMDeveloper/Private/EdGraph/RigVMEdGraphSchema.cpp
+++ b/Engine/Plugins/Runtime/RigVM/Source/RigVMDeveloper/Private/EdGraph/RigVMEdGraphSchema.cpp
@@ -1696,6 +1696,10 @@ void URigVMEdGraphSchema::GetGraphDisplayInformation(const UEdGraph& Graph, /*ou
 						NodePath.RemoveFromEnd(NodePathSuffix);
 						NodePath.RemoveFromStart(FRigVMClient::RigVMModelPrefix);
 						NodePath.TrimStartAndEndInline();
+						if (NodePath.IsEmpty())
+						{
+							NodePath = TEXT("Rig Graph");
+						}
 						DisplayInfo.DisplayName = FText::FromString(NodePath);
 						DisplayInfo.PlainName = DisplayInfo.DisplayName;
 					}
