Control Rig: FIx converted rigs from BP to non-BP opening graph editor tabs that were not opened in the original asset
#rnx

#Synced-CL 52613983

diff --git a/Engine/Plugins/Runtime/RigVM/Source/RigVMDeveloper/Private/RigVMBlueprintLegacy.cpp b/Engine/Plugins/Runtime/RigVM/Source/RigVMDeveloper/Private/RigVMBlueprintLegacy.cpp
--- a/Engine/Plugins/Runtime/RigVM/Source/RigVMDeveloper/Private/RigVMBlueprintLegacy.cpp
+++ b/Engine/Plugins/Runtime/RigVM/Source/RigVMDeveloper/Private/RigVMBlueprintLegacy.cpp
@@ -616,6 +616,24 @@ bool URigVMBlueprint::CopyBlueprintToAsset(URigVMBlueprint* InBlueprint, URigVMR
 			EditorAsset->CreateEdGraph(Graph);
 		}
 
+		// Copy LastEditedDocuments from source blueprint, remapping to new EdGraphs
+		EditorAsset->GetLastEditedDocuments().Reset();
+		for (const FEditedDocumentInfo& SourceDoc : InBlueprint->GetLastEditedDocuments())
+		{
+			if (URigVMEdGraph* SourceEdGraph = Cast<URigVMEdGraph>(SourceDoc.EditedObjectPath.ResolveObject()))
+			{
+				// Find the corresponding EdGraph in the new asset by ModelNodePath
+				if (UEdGraph* NewEdGraph = EditorAsset->GetEdGraph(SourceEdGraph->ModelNodePath))
+				{
+					FEditedDocumentInfo NewDoc;
+					NewDoc.EditedObjectPath = NewEdGraph;
+					NewDoc.SavedViewOffset = SourceDoc.SavedViewOffset;
+					NewDoc.SavedZoomAmount = SourceDoc.SavedZoomAmount;
+					EditorAsset->GetLastEditedDocuments().Add(NewDoc);
+				}
+			}
+		}
+
 		EditorAsset->GraphDisplaySettings = InBlueprint->GetRigGraphDisplaySettings();
 		EditorAsset->CompileSettings = InBlueprint->GetVMCompileSettings();
 		
