Control Rig: Fix SaveOnCompile not working in BP independent rigs
#rnx

#Synced-CL 52613983

diff --git a/Engine/Plugins/Runtime/RigVM/Source/RigVM/Public/RigVMSettings.h b/Engine/Plugins/Runtime/RigVM/Source/RigVM/Public/RigVMSettings.h
--- a/Engine/Plugins/Runtime/RigVM/Source/RigVM/Public/RigVMSettings.h
+++ b/Engine/Plugins/Runtime/RigVM/Source/RigVM/Public/RigVMSettings.h
@@ -14,6 +14,17 @@
 
 class UStaticMesh;
 
+/**
+ * Determines when to save RigVM assets after compilation.
+ */
+UENUM()
+enum class ERigVMSaveOnCompile : uint8
+{
+	Never UMETA(DisplayName="Never"),
+	SuccessOnly UMETA(DisplayName="On Success Only"),
+	Always UMETA(DisplayName = "Always"),
+};
+
 /**
  * Customize RigVM Rig Editor.
  */
@@ -28,6 +39,11 @@ public:
 	
 #if WITH_EDITORONLY_DATA
 
+	// Determines when the asset should be saved after compilation.
+	// This applies to RigVM assets that are not blueprint-based (e.g., Control Rig Runtime Assets).
+	UPROPERTY(EditAnywhere, config, Category = Compiler)
+	ERigVMSaveOnCompile SaveOnCompile = ERigVMSaveOnCompile::Never;
+
 	// with this checked similar nodes in the graph will be highlighted.
 	// this highlighting can then also be used to functionize items
 	UPROPERTY(EditAnywhere, config, Category = Interaction)
diff --git a/Engine/Plugins/Runtime/RigVM/Source/RigVMDeveloper/Private/RigVMEditorAsset.cpp b/Engine/Plugins/Runtime/RigVM/Source/RigVMDeveloper/Private/RigVMEditorAsset.cpp
--- a/Engine/Plugins/Runtime/RigVM/Source/RigVMDeveloper/Private/RigVMEditorAsset.cpp
+++ b/Engine/Plugins/Runtime/RigVM/Source/RigVMDeveloper/Private/RigVMEditorAsset.cpp
@@ -2091,7 +2091,6 @@ void IRigVMEditorAssetInterface::HandleReportFromCompiler(EMessageSeverity::Type
 	if (InSeverity == EMessageSeverity::Error)
 	{
 		SetAssetStatus(ERigVMAssetStatus::RVMA_Error);
-		(void)GetObject()->MarkPackageDirty();
 
 		// see UnitTest "ControlRig.Basics.OrphanedPins" to learn why errors are suppressed this way
 		if (GetVMCompileSettings().SurpressErrors)
diff --git a/Engine/Plugins/Runtime/RigVM/Source/RigVMEditor/Private/Editor/RigVMNewEditor.cpp b/Engine/Plugins/Runtime/RigVM/Source/RigVMEditor/Private/Editor/RigVMNewEditor.cpp
--- a/Engine/Plugins/Runtime/RigVM/Source/RigVMEditor/Private/Editor/RigVMNewEditor.cpp
+++ b/Engine/Plugins/Runtime/RigVM/Source/RigVMEditor/Private/Editor/RigVMNewEditor.cpp
@@ -4,6 +4,7 @@
 
 #include "BlueprintActionDatabase.h"
 #include "BlueprintEditorSettings.h"
+#include "FileHelpers.h"
 #include "BlueprintEditorTabs.h"
 #include "IMessageLogListing.h"
 #include "K2Node_Composite.h"
@@ -672,7 +673,7 @@ TSharedRef<SWidget> FRigVMNewEditor::GenerateCompileOptionsMenu()
 	// @TODO: disable the menu and change up the tooltip when all sub items are disabled
 	MenuBuilder.AddSubMenu(
 		LOCTEXT("SaveOnCompileSubMenu", "Save on Compile"),
-		LOCTEXT("SaveOnCompileSubMenu_ToolTip", "Determines how the Blueprint is saved whenever you compile it."),
+		LOCTEXT("SaveOnCompileSubMenu_ToolTip", "Determines if the asset is saved whenever you compile it."),
 		FNewMenuDelegate::CreateSP(this, &FRigVMNewEditor::MakeSaveOnCompileSubMenu)
 	);
 
@@ -1124,42 +1125,17 @@ UEdGraphPin* FRigVMNewEditor::GetCurrentlySelectedPin() const
 	return nullptr;
 }
 
-void FRigVMNewEditor::SetSaveOnCompileSetting(ESaveOnCompile NewSetting)
+void FRigVMNewEditor::SetSaveOnCompileSetting(ERigVMSaveOnCompile NewSetting)
 {
-	// Copied from BlueprintEditorImpl::SetSaveOnCompileSetting
-	
-	UBlueprintEditorSettings* Settings = GetMutableDefault<UBlueprintEditorSettings>();
+	URigVMEditorSettings* Settings = GetMutableDefault<URigVMEditorSettings>();
 	Settings->SaveOnCompile = NewSetting;
 	Settings->SaveConfig();
 }
 
-bool FRigVMNewEditor::IsSaveOnCompileEnabled() const
+bool FRigVMNewEditor::IsSaveOnCompileOptionSet(ERigVMSaveOnCompile Option)
 {
-	// Copied from BlueprintEditorImpl::IsSaveOnCompileEnabled
-	
-	FRigVMEditorAssetInterfacePtr Blueprint = GetRigVMAssetInterface();
-	bool const bIsLevelScript = (Cast<ULevelScriptBlueprint>(Blueprint->GetObject()) != nullptr);
-
-	return !bIsLevelScript;
-}
-
-bool FRigVMNewEditor::IsSaveOnCompileOptionSet(TWeakPtr<FRigVMNewEditor> Editor, ESaveOnCompile Option)
-{
-	// Copied from BlueprintEditorImpl::IsSaveOnCompileOptionSet
-	
-	const UBlueprintEditorSettings* Settings = GetDefault<UBlueprintEditorSettings>();
-
-	ESaveOnCompile CurrentSetting = Settings->SaveOnCompile;
-	if (!Editor.IsValid() || !Editor.Pin()->IsSaveOnCompileEnabled())
-	{
-		// if save-on-compile is disabled for the blueprint, then we want to 
-		// show "Never" as being selected
-		// 
-		// @TODO: a tooltip explaining why would be nice too
-		CurrentSetting = SoC_Never;
-	}
-
-	return (CurrentSetting == Option);
+	const URigVMEditorSettings* Settings = GetDefault<URigVMEditorSettings>();
+	return (Settings->SaveOnCompile == Option);
 }
 
 void FRigVMNewEditor::ToggleJumpToErrorNodeSetting()
@@ -1243,24 +1219,24 @@ void FRigVMNewEditor::CreateDefaultCommands()
 			FExecuteAction::CreateSP(this, &FRigVMEditorBase::Compile),
 			FCanExecuteAction::CreateSP(this, &FRigVMEditorBase::IsCompilingEnabled));
 
-		TWeakPtr<FRigVMNewEditor> WeakThisPtr = SharedThis(this);
+		// Save-on-compile commands - always available for runtime assets (no can-execute guard needed)
 		ToolkitCommands->MapAction(
 			FRigVMEditorCommands::Get().SaveOnCompile_Never,
-			FExecuteAction::CreateSP(this, &FRigVMNewEditor::SetSaveOnCompileSetting, (ESaveOnCompile)SoC_Never),
-			FCanExecuteAction::CreateSP(this, &FRigVMNewEditor::IsSaveOnCompileEnabled),
-			FIsActionChecked::CreateSP(this, &FRigVMNewEditor::IsSaveOnCompileOptionSet, WeakThisPtr, (ESaveOnCompile)SoC_Never)
+			FExecuteAction::CreateSP(this, &FRigVMNewEditor::SetSaveOnCompileSetting, ERigVMSaveOnCompile::Never),
+			FCanExecuteAction(),
+			FIsActionChecked::CreateSP(this, &FRigVMNewEditor::IsSaveOnCompileOptionSet, ERigVMSaveOnCompile::Never)
 		);
 		ToolkitCommands->MapAction(
 			FRigVMEditorCommands::Get().SaveOnCompile_SuccessOnly,
-			FExecuteAction::CreateSP(this, &FRigVMNewEditor::SetSaveOnCompileSetting, (ESaveOnCompile)SoC_SuccessOnly),
-			FCanExecuteAction::CreateSP(this, &FRigVMNewEditor::IsSaveOnCompileEnabled),
-			FIsActionChecked::CreateSP(this, &FRigVMNewEditor::IsSaveOnCompileOptionSet, WeakThisPtr, (ESaveOnCompile)SoC_SuccessOnly)
+			FExecuteAction::CreateSP(this, &FRigVMNewEditor::SetSaveOnCompileSetting, ERigVMSaveOnCompile::SuccessOnly),
+			FCanExecuteAction(),
+			FIsActionChecked::CreateSP(this, &FRigVMNewEditor::IsSaveOnCompileOptionSet, ERigVMSaveOnCompile::SuccessOnly)
 		);
 		ToolkitCommands->MapAction(
 			FRigVMEditorCommands::Get().SaveOnCompile_Always,
-			FExecuteAction::CreateSP(this, &FRigVMNewEditor::SetSaveOnCompileSetting, (ESaveOnCompile)SoC_Always),
-			FCanExecuteAction::CreateSP(this, &FRigVMNewEditor::IsSaveOnCompileEnabled),
-			FIsActionChecked::CreateSP(this, &FRigVMNewEditor::IsSaveOnCompileOptionSet, WeakThisPtr, (ESaveOnCompile)SoC_Always)
+			FExecuteAction::CreateSP(this, &FRigVMNewEditor::SetSaveOnCompileSetting, ERigVMSaveOnCompile::Always),
+			FCanExecuteAction(),
+			FIsActionChecked::CreateSP(this, &FRigVMNewEditor::IsSaveOnCompileOptionSet, ERigVMSaveOnCompile::Always)
 		);
 
 		ToolkitCommands->MapAction(
@@ -1489,6 +1465,33 @@ void FRigVMNewEditor::CompileImpl()
 		//AnalyticsTrackCompileEvent(BlueprintObj, LogResults.NumErrors, LogResults.NumWarnings);
 
 		RefreshInspector();
+
+		// Handle save-on-compile for runtime assets
+		// This uses RigVM-specific settings rather than blueprint settings since runtime assets are blueprint-independent
+		if (!BlueprintObj->IsRegeneratingOnLoad())
+		{
+			const URigVMEditorSettings* Settings = GetDefault<URigVMEditorSettings>();
+			const ERigVMSaveOnCompile SaveOnCompileSetting = Settings->SaveOnCompile;
+			const ERigVMAssetStatus AssetStatus = BlueprintObj->GetAssetStatus();
+
+			const bool bShouldSave =
+				(SaveOnCompileSetting == ERigVMSaveOnCompile::Always) ||
+				(SaveOnCompileSetting == ERigVMSaveOnCompile::SuccessOnly && AssetStatus == ERigVMAssetStatus::RVMA_UpToDate);
+
+			if (bShouldSave)
+			{
+				// Save the package containing the editing object
+				// The editor asset (what we're editing) is a sub-object of the runtime asset,
+				// both residing in the same package
+				UPackage* Package = BlueprintObj->GetObject()->GetOutermost();
+				if (Package)
+				{
+					TArray<UPackage*> PackagesToSave;
+					PackagesToSave.Add(Package);
+					FEditorFileUtils::PromptForCheckoutAndSave(PackagesToSave, /*bCheckDirty=*/true, /*bPromptToSave=*/false);
+				}
+			}
+		}
 	}
 }
 
diff --git a/Engine/Plugins/Runtime/RigVM/Source/RigVMEditor/Public/Editor/RigVMNewEditor.h b/Engine/Plugins/Runtime/RigVM/Source/RigVMEditor/Public/Editor/RigVMNewEditor.h
--- a/Engine/Plugins/Runtime/RigVM/Source/RigVMEditor/Public/Editor/RigVMNewEditor.h
+++ b/Engine/Plugins/Runtime/RigVM/Source/RigVMEditor/Public/Editor/RigVMNewEditor.h
@@ -4,6 +4,7 @@
 
 #include "BlueprintEditorSettings.h"
 #include "RigVMEditor.h"
+#include "RigVMSettings.h"
 #include "WorkflowOrientedApp/WorkflowCentricApplication.h"
 
 #define UE_API RIGVMEDITOR_API
@@ -165,9 +166,8 @@ public:
 
 	UE_API TSharedRef<SWidget> GenerateCompileOptionsMenu();
 	UE_API void MakeSaveOnCompileSubMenu(FMenuBuilder& InMenu);
-	UE_API void SetSaveOnCompileSetting(ESaveOnCompile NewSetting);
-	UE_API bool IsSaveOnCompileEnabled() const;
-	UE_API bool IsSaveOnCompileOptionSet(TWeakPtr<FRigVMNewEditor> Editor, ESaveOnCompile Option);
+	UE_API void SetSaveOnCompileSetting(ERigVMSaveOnCompile NewSetting);
+	UE_API bool IsSaveOnCompileOptionSet(ERigVMSaveOnCompile Option);
 	UE_API void ToggleJumpToErrorNodeSetting();
 	UE_API bool IsJumpToErrorNodeOptionSet();
 	UE_API UEdGraphNode* FindNodeWithError(UBlueprint* Blueprint, EMessageSeverity::Type Severity);
