RigVM: Skip ensure and just log missing object types when cooking.

#rnx

#Synced-CL 52537381

diff --git a/Engine/Plugins/Runtime/RigVM/Source/RigVM/Private/RigVMCore/RigVMTemplate.cpp b/Engine/Plugins/Runtime/RigVM/Source/RigVM/Private/RigVMCore/RigVMTemplate.cpp
--- a/Engine/Plugins/Runtime/RigVM/Source/RigVM/Private/RigVMCore/RigVMTemplate.cpp
+++ b/Engine/Plugins/Runtime/RigVM/Source/RigVM/Private/RigVMCore/RigVMTemplate.cpp
@@ -1124,11 +1124,18 @@ FRigVMTemplate::FTypeMap FRigVMTemplate::GetArgumentTypesFromString_NoLock(const
 					if (TypeIndex == INDEX_NONE && RigVMTypeUtils::RequiresCPPTypeObject(TypeName))
 					{
 						UObject* CPPTypeObject = RigVMTypeUtils::ObjectFromCPPType(TypeName, true, InTypeResolver);
-						if (ensure(CPPTypeObject != nullptr))
+						if (CPPTypeObject != nullptr)
 						{
 							FRigVMTemplateArgumentType ArgType(TypeFName, CPPTypeObject);
 							TypeIndex = InRegistry->FindOrAddType_NoLock(ArgType, false);
 						}
+						else
+						{
+							UE_LOG(LogRigVM, Display, TEXT("RigVMTemplate '%s' cannot resolve type '%s' when solving permutation for '%s'."),
+									*GetNotation().ToString(),
+									*TypeName,
+									*InTypeString);
+						}
 					}
 					
 					if(TypeIndex != INDEX_NONE)
