Control Rig: Fix make constant nodes breaking link inside pure functions
#rnx

#Synced-CL 52613983

diff --git a/Engine/Plugins/Animation/ControlRig/Source/ControlRig/Public/Units/Debug/RigUnit_DebugBase.h b/Engine/Plugins/Animation/ControlRig/Source/ControlRig/Public/Units/Debug/RigUnit_DebugBase.h
--- a/Engine/Plugins/Animation/ControlRig/Source/ControlRig/Public/Units/Debug/RigUnit_DebugBase.h
+++ b/Engine/Plugins/Animation/ControlRig/Source/ControlRig/Public/Units/Debug/RigUnit_DebugBase.h
@@ -21,7 +21,7 @@ struct FRigUnit_DebugBase : public FRigUnit
 /*
  * Base class for all mutable debug nodes
  */
-USTRUCT(meta=(Abstract, Category="Debug", NodeColor = "0.83077 0.846873 0.049707", DocumentationPolicy = "Strict"))
+USTRUCT(meta=(Abstract, Category="Debug", NodeColor = "0.83077 0.846873 0.049707", DocumentationPolicy = "Strict", Pure))
 struct FRigUnit_DebugBaseMutable : public FRigUnitMutable
 {
 	GENERATED_BODY()
diff --git a/Engine/Plugins/Animation/ControlRig/Source/ControlRig/Public/Units/Debug/RigUnit_VisualDebug.h b/Engine/Plugins/Animation/ControlRig/Source/ControlRig/Public/Units/Debug/RigUnit_VisualDebug.h
--- a/Engine/Plugins/Animation/ControlRig/Source/ControlRig/Public/Units/Debug/RigUnit_VisualDebug.h
+++ b/Engine/Plugins/Animation/ControlRig/Source/ControlRig/Public/Units/Debug/RigUnit_VisualDebug.h
@@ -191,7 +191,7 @@ struct FRigUnit_VisualDebugTransform : public FRigUnit_DebugBase
 	RIGVM_METHOD()
 	UE_API virtual void Execute() override;
 
-	UPROPERTY(meta = (Input, Output))
+	UPROPERTY(meta = (Input, Output, Constant))
 	FTransform Value;
 
 	UPROPERTY(meta = (Input))
@@ -230,7 +230,7 @@ struct FRigUnit_VisualDebugTransformItemSpace : public FRigUnit_DebugBase
 	RIGVM_METHOD()
 	UE_API virtual void Execute() override;
 
-	UPROPERTY(meta = (Input, Output))
+	UPROPERTY(meta = (Input, Output, Constant))
 	FTransform Value;
 
 	UPROPERTY(meta = (Input))
diff --git a/Engine/Plugins/Runtime/RigVM/Source/RigVM/Private/RigVMCore/RigVMTemplate.cpp b/Engine/Plugins/Runtime/RigVM/Source/RigVM/Private/RigVMCore/RigVMTemplate.cpp
--- a/Engine/Plugins/Runtime/RigVM/Source/RigVM/Private/RigVMCore/RigVMTemplate.cpp
+++ b/Engine/Plugins/Runtime/RigVM/Source/RigVM/Private/RigVMCore/RigVMTemplate.cpp
@@ -1437,7 +1437,12 @@ FString FRigVMTemplate::GetArgumentMetaData(const FName& InArgumentName, const F
 			
 			if(const FProperty* Property = ResolvedFunction->Struct->FindPropertyByName(InArgumentName))
 			{
-				const FString MetaData = Property->GetMetaData(InMetaDataKey);
+				// Flag-style metadata (no value) returns empty from GetMetaData; use sentinel.
+				FString MetaData = Property->GetMetaData(InMetaDataKey);
+				if (MetaData.IsEmpty() && Property->HasMetaData(InMetaDataKey))
+				{
+					MetaData = FRigVMDispatchFactory::TrueString;
+				}
 				if (!ResolvedMetaData.IsEmpty())
 				{
 					if(!ResolvedMetaData.Equals(MetaData))
diff --git a/Engine/Plugins/Runtime/RigVM/Source/RigVM/Private/RigVMFunctions/RigVMDispatch_Constant.cpp b/Engine/Plugins/Runtime/RigVM/Source/RigVM/Private/RigVMFunctions/RigVMDispatch_Constant.cpp
--- a/Engine/Plugins/Runtime/RigVM/Source/RigVM/Private/RigVMFunctions/RigVMDispatch_Constant.cpp
+++ b/Engine/Plugins/Runtime/RigVM/Source/RigVM/Private/RigVMFunctions/RigVMDispatch_Constant.cpp
@@ -69,6 +69,16 @@ FText FRigVMDispatch_Constant::GetArgumentTooltip(const FName& InArgumentName, T
 	return FRigVMDispatch_CoreBase::GetArgumentTooltip(InArgumentName, InTypeIndex);
 }
 
+FString FRigVMDispatch_Constant::GetArgumentMetaData(const FName& InArgumentName, const FName& InMetaDataKey) const
+{
+	if (InArgumentName == ValueName && InMetaDataKey == FRigVMStruct::ConstantMetaName)
+	{
+		return TEXT("True");
+	}
+
+	return FRigVMDispatch_CoreBase::GetArgumentMetaData(InArgumentName, InMetaDataKey);
+}
+
 #endif
 
 void FRigVMDispatch_Constant::Execute(FRigVMExtendedExecuteContext& InContext, FRigVMMemoryHandleArray Handles, FRigVMPredicateBranchArray RigVMBranches)
diff --git a/Engine/Plugins/Runtime/RigVM/Source/RigVM/Public/RigVMCore/RigVMStruct.h b/Engine/Plugins/Runtime/RigVM/Source/RigVM/Public/RigVMCore/RigVMStruct.h
--- a/Engine/Plugins/Runtime/RigVM/Source/RigVM/Public/RigVMCore/RigVMStruct.h
+++ b/Engine/Plugins/Runtime/RigVM/Source/RigVM/Public/RigVMCore/RigVMStruct.h
@@ -303,6 +303,25 @@ public:
 	static inline const FLazyName ControlFlowCountName = FLazyName(TEXT("Count"));
 	static inline const FLazyName ControlFlowIndexName = FLazyName(TEXT("Index"));
 
+#if WITH_EDITOR
+	/**
+	 * Checks if any struct in the hierarchy has the specified metadata.
+	 * UScriptStruct metadata is not inherited automatically, so this walks
+	 * up the hierarchy to check each struct in the chain.
+	 */
+	static bool HasMetaData(const UScriptStruct* InStruct, const FName& InMetaDataKey)
+	{
+		for (const UScriptStruct* Current = InStruct; Current; Current = Cast<UScriptStruct>(Current->GetSuperStruct()))
+		{
+			if (Current->HasMetaData(InMetaDataKey))
+			{
+				return true;
+			}
+		}
+		return false;
+	}
+#endif
+
 protected:
 
 	RIGVM_API static float GetRatioFromIndex(int32 InIndex, int32 InCount);
diff --git a/Engine/Plugins/Runtime/RigVM/Source/RigVM/Public/RigVMFunctions/Debug/RigVMFunction_DebugBase.h b/Engine/Plugins/Runtime/RigVM/Source/RigVM/Public/RigVMFunctions/Debug/RigVMFunction_DebugBase.h
--- a/Engine/Plugins/Runtime/RigVM/Source/RigVM/Public/RigVMFunctions/Debug/RigVMFunction_DebugBase.h
+++ b/Engine/Plugins/Runtime/RigVM/Source/RigVM/Public/RigVMFunctions/Debug/RigVMFunction_DebugBase.h
@@ -22,7 +22,7 @@ struct FRigVMFunction_DebugBase : public FRigVMStruct
 /*
  * The base class for all mutable debug draw nodes
  */
-USTRUCT(meta=(Abstract, Category="Debug", NodeColor = "0.83077 0.846873 0.049707", DocumentationPolicy = "Strict"))
+USTRUCT(meta=(Abstract, Category="Debug", NodeColor = "0.83077 0.846873 0.049707", DocumentationPolicy = "Strict", Pure))
 struct FRigVMFunction_DebugBaseMutable : public FRigVMStructMutable
 {
 	GENERATED_BODY()
diff --git a/Engine/Plugins/Runtime/RigVM/Source/RigVM/Public/RigVMFunctions/Debug/RigVMFunction_VisualDebug.h b/Engine/Plugins/Runtime/RigVM/Source/RigVM/Public/RigVMFunctions/Debug/RigVMFunction_VisualDebug.h
--- a/Engine/Plugins/Runtime/RigVM/Source/RigVM/Public/RigVMFunctions/Debug/RigVMFunction_VisualDebug.h
+++ b/Engine/Plugins/Runtime/RigVM/Source/RigVM/Public/RigVMFunctions/Debug/RigVMFunction_VisualDebug.h
@@ -197,7 +197,7 @@ struct FRigVMFunction_VisualDebugTransform : public FRigVMFunction_DebugBase
 	RIGVM_METHOD()
 	RIGVM_API virtual void Execute() override;
 
-	UPROPERTY(meta = (Input, Output))
+	UPROPERTY(meta = (Input, Output, Constant))
 	FTransform Value;
 
 	UPROPERTY(meta = (Input))
@@ -236,7 +236,7 @@ struct FRigVMFunction_VisualDebugTransformNoSpace : public FRigVMFunction_DebugB
 	RIGVM_API virtual void Execute() override;
 
 	// The pass through transform to draw
-	UPROPERTY(meta = (Input, Output))
+	UPROPERTY(meta = (Input, Output, Constant))
 	FTransform Value;
 
 	// If set to False the debug drawing will be skipped
diff --git a/Engine/Plugins/Runtime/RigVM/Source/RigVM/Public/RigVMFunctions/Math/RigVMFunction_MathBool.h b/Engine/Plugins/Runtime/RigVM/Source/RigVM/Public/RigVMFunctions/Math/RigVMFunction_MathBool.h
--- a/Engine/Plugins/Runtime/RigVM/Source/RigVM/Public/RigVMFunctions/Math/RigVMFunction_MathBool.h
+++ b/Engine/Plugins/Runtime/RigVM/Source/RigVM/Public/RigVMFunctions/Math/RigVMFunction_MathBool.h
@@ -122,7 +122,7 @@ struct FRigVMFunction_MathBoolMake : public FRigVMFunction_MathBoolBase
 	RIGVM_METHOD()
 	RIGVM_API virtual void Execute() override;
 
-	UPROPERTY(meta=(Input, Output))
+	UPROPERTY(meta=(Input, Output, Constant))
 	bool Value;
 
 	RIGVM_METHOD()
diff --git a/Engine/Plugins/Runtime/RigVM/Source/RigVM/Public/RigVMFunctions/Math/RigVMFunction_MathDouble.h b/Engine/Plugins/Runtime/RigVM/Source/RigVM/Public/RigVMFunctions/Math/RigVMFunction_MathDouble.h
--- a/Engine/Plugins/Runtime/RigVM/Source/RigVM/Public/RigVMFunctions/Math/RigVMFunction_MathDouble.h
+++ b/Engine/Plugins/Runtime/RigVM/Source/RigVM/Public/RigVMFunctions/Math/RigVMFunction_MathDouble.h
@@ -122,7 +122,7 @@ struct FRigVMFunction_MathDoubleMake : public FRigVMFunction_MathDoubleBase
 	RIGVM_METHOD()
 	RIGVM_API virtual void Execute() override;
 
-	UPROPERTY(meta=(Input, Output))
+	UPROPERTY(meta=(Input, Output, Constant))
 	double Value;
 	
 	RIGVM_METHOD()
diff --git a/Engine/Plugins/Runtime/RigVM/Source/RigVM/Public/RigVMFunctions/Math/RigVMFunction_MathFloat.h b/Engine/Plugins/Runtime/RigVM/Source/RigVM/Public/RigVMFunctions/Math/RigVMFunction_MathFloat.h
--- a/Engine/Plugins/Runtime/RigVM/Source/RigVM/Public/RigVMFunctions/Math/RigVMFunction_MathFloat.h
+++ b/Engine/Plugins/Runtime/RigVM/Source/RigVM/Public/RigVMFunctions/Math/RigVMFunction_MathFloat.h
@@ -124,7 +124,7 @@ struct FRigVMFunction_MathFloatMake : public FRigVMFunction_MathFloatBase
 	RIGVM_METHOD()
 	UE_API virtual void Execute() override;
 
-	UPROPERTY(meta=(Input, Output))
+	UPROPERTY(meta=(Input, Output, Constant))
 	float Value;
 
 	RIGVM_METHOD()
diff --git a/Engine/Plugins/Runtime/RigVM/Source/RigVM/Public/RigVMFunctions/Math/RigVMFunction_MathInt.h b/Engine/Plugins/Runtime/RigVM/Source/RigVM/Public/RigVMFunctions/Math/RigVMFunction_MathInt.h
--- a/Engine/Plugins/Runtime/RigVM/Source/RigVM/Public/RigVMFunctions/Math/RigVMFunction_MathInt.h
+++ b/Engine/Plugins/Runtime/RigVM/Source/RigVM/Public/RigVMFunctions/Math/RigVMFunction_MathInt.h
@@ -104,7 +104,7 @@ struct FRigVMFunction_MathIntMake : public FRigVMFunction_MathIntBase
 	RIGVM_METHOD()
 	RIGVM_API virtual void Execute() override;
 
-	UPROPERTY(meta=(Input, Output))
+	UPROPERTY(meta=(Input, Output, Constant))
 	int32 Value;
 
 	RIGVM_METHOD()
diff --git a/Engine/Plugins/Runtime/RigVM/Source/RigVM/Public/RigVMFunctions/RigVMDispatch_Constant.h b/Engine/Plugins/Runtime/RigVM/Source/RigVM/Public/RigVMFunctions/RigVMDispatch_Constant.h
--- a/Engine/Plugins/Runtime/RigVM/Source/RigVM/Public/RigVMFunctions/RigVMDispatch_Constant.h
+++ b/Engine/Plugins/Runtime/RigVM/Source/RigVM/Public/RigVMFunctions/RigVMDispatch_Constant.h
@@ -24,6 +24,7 @@ public:
 	RIGVM_API virtual FString GetNodeTitle(const FRigVMTemplateTypeMap& InTypes) const override;
 	RIGVM_API virtual FText GetNodeTooltip(const FRigVMTemplateTypeMap& InTypes) const override;
 	RIGVM_API virtual FText GetArgumentTooltip(const FName& InArgumentName, TRigVMTypeIndex InTypeIndex) const override;
+	RIGVM_API virtual FString GetArgumentMetaData(const FName& InArgumentName, const FName& InMetaDataKey) const override;
 #endif
 	virtual bool IsSingleton() const override { return true; }
 	virtual bool SupportsRenaming() const override { return true; }
diff --git a/Engine/Plugins/Runtime/RigVM/Source/RigVM/Public/RigVMFunctions/RigVMDispatch_Print.h b/Engine/Plugins/Runtime/RigVM/Source/RigVM/Public/RigVMFunctions/RigVMDispatch_Print.h
--- a/Engine/Plugins/Runtime/RigVM/Source/RigVM/Public/RigVMFunctions/RigVMDispatch_Print.h
+++ b/Engine/Plugins/Runtime/RigVM/Source/RigVM/Public/RigVMFunctions/RigVMDispatch_Print.h
@@ -8,7 +8,7 @@
 /*
  * Prints any value to the log
  */
-USTRUCT(meta=(DisplayName = "Print", NodeColor = "0.8, 0, 0.2, 1"))
+USTRUCT(meta=(DisplayName = "Print", NodeColor = "0.8, 0, 0.2, 1", Pure))
 struct FRigVMDispatch_Print : public FRigVMDispatchFactory
 {
 	GENERATED_BODY()
diff --git a/Engine/Plugins/Runtime/RigVM/Source/RigVMDeveloper/Private/RigVMCompiler/RigVMCompiler.cpp b/Engine/Plugins/Runtime/RigVM/Source/RigVMDeveloper/Private/RigVMCompiler/RigVMCompiler.cpp
--- a/Engine/Plugins/Runtime/RigVM/Source/RigVMDeveloper/Private/RigVMCompiler/RigVMCompiler.cpp
+++ b/Engine/Plugins/Runtime/RigVM/Source/RigVMDeveloper/Private/RigVMCompiler/RigVMCompiler.cpp
@@ -1595,13 +1595,13 @@ bool URigVMCompiler::Compile(const FRigVMCompileSettings& InSettings, TArray<URi
 							}
 						}
 
-						// Reject non-local sources connected to IO pins.
+						// Reject non-local sources connected to non-constant IO pins.
 						// Only local variable getters can connect to IO pins in pure functions.
 						if (!ModelNode->IsA<URigVMRerouteNode>())
 						{
 							for (const URigVMPin* Pin : ModelNode->GetPins())
 							{
-								if (Pin->GetDirection() != ERigVMPinDirection::IO)
+								if (Pin->GetDirection() != ERigVMPinDirection::IO || Pin->IsDefinedAsConstant())
 								{
 									continue;
 								}
diff --git a/Engine/Plugins/Runtime/RigVM/Source/RigVMDeveloper/Private/RigVMModel/Nodes/RigVMDispatchNode.cpp b/Engine/Plugins/Runtime/RigVM/Source/RigVMDeveloper/Private/RigVMModel/Nodes/RigVMDispatchNode.cpp
--- a/Engine/Plugins/Runtime/RigVM/Source/RigVMDeveloper/Private/RigVMModel/Nodes/RigVMDispatchNode.cpp
+++ b/Engine/Plugins/Runtime/RigVM/Source/RigVMDeveloper/Private/RigVMModel/Nodes/RigVMDispatchNode.cpp
@@ -67,18 +67,11 @@ bool URigVMDispatchNode::IsPure() const
 {
 	if (const UScriptStruct* Struct = GetFactoryStruct())
 	{
-		// Pure metadata explicitly marks dispatch nodes as pure, overriding hidden pin checks.
-		// This allows nodes with internal state (e.g., loop counters in control flow)
-		// to be marked pure as long as they have no external side effects.
-		if (Struct->HasMetaData(FRigVMStruct::PureMetaName))
+		if (FRigVMStruct::HasMetaData(Struct, FRigVMStruct::PureMetaName))
 		{
 			return true;
 		}
 	}
-
-	// For nodes without Pure metadata, fall back to base class logic which checks:
-	// 1. Not mutable (no execute context pins)
-	// 2. No hidden pins (no internal state)
 	return Super::IsPure();
 }
 
@@ -254,6 +247,18 @@ bool URigVMDispatchNode::ShouldInputPinComputeLazily(const URigVMPin* InPin) con
 	return Super::ShouldInputPinComputeLazily(InPin);
 }
 
+bool URigVMDispatchNode::IsPinConstant(const URigVMPin* InPin) const
+{
+	const URigVMPin* RootPin = InPin->GetRootPin();
+	check(RootPin->GetNode() == this);
+
+	if(const FRigVMDispatchFactory* Factory = GetFactory())
+	{
+		return Factory->HasArgumentMetaData(RootPin->GetFName(), FRigVMStruct::ConstantMetaName); 
+	}
+	return Super::IsPinConstant(InPin);
+}
+
 FString URigVMDispatchNode::GetOriginalDefaultValueForRootPin(const URigVMPin* InRootPin) const
 {
 	if(const FRigVMDispatchFactory* Factory = GetFactory())
diff --git a/Engine/Plugins/Runtime/RigVM/Source/RigVMDeveloper/Private/RigVMModel/Nodes/RigVMUnitNode.cpp b/Engine/Plugins/Runtime/RigVM/Source/RigVMDeveloper/Private/RigVMModel/Nodes/RigVMUnitNode.cpp
--- a/Engine/Plugins/Runtime/RigVM/Source/RigVMDeveloper/Private/RigVMModel/Nodes/RigVMUnitNode.cpp
+++ b/Engine/Plugins/Runtime/RigVM/Source/RigVMDeveloper/Private/RigVMModel/Nodes/RigVMUnitNode.cpp
@@ -102,18 +102,11 @@ bool URigVMUnitNode::IsPure() const
 
 	if (UScriptStruct* Struct = GetScriptStruct())
 	{
-		// Pure metadata explicitly marks nodes as pure, overriding hidden pin checks.
-		// This allows nodes with internal state (e.g., loop counters in control flow)
-		// to be marked pure as long as they have no external side effects.
-		if (Struct->HasMetaData(FRigVMStruct::PureMetaName))
+		if (FRigVMStruct::HasMetaData(Struct, FRigVMStruct::PureMetaName))
 		{
 			return true;
 		}
 	}
-
-	// For nodes without Pure metadata, fall back to base class logic which checks:
-	// 1. Not mutable (no execute context pins)
-	// 2. No hidden pins (no internal state)
 	return Super::IsPure();
 }
 
@@ -262,6 +255,21 @@ bool URigVMUnitNode::ShouldInputPinComputeLazily(const URigVMPin* InPin) const
 	return Super::ShouldInputPinComputeLazily(InPin);
 }
 
+bool URigVMUnitNode::IsPinConstant(const URigVMPin* InPin) const
+{
+	const URigVMPin* RootPin = InPin->GetRootPin();
+	check(RootPin->GetNode() == this);
+
+	if(const UScriptStruct* Struct = GetScriptStruct())
+	{
+		if(const FProperty* Property = Struct->FindPropertyByName(RootPin->GetFName()))
+		{
+			return Property->HasMetaData(FRigVMStruct::ConstantMetaName);
+		}
+	}
+	return Super::IsPinConstant(InPin);
+}
+
 FString URigVMUnitNode::GetOriginalDefaultValueForRootPin(const URigVMPin* InRootPin) const
 {
 	if(InRootPin->CanProvideDefaultValue())
diff --git a/Engine/Plugins/Runtime/RigVM/Source/RigVMDeveloper/Private/RigVMModel/RigVMClient.cpp b/Engine/Plugins/Runtime/RigVM/Source/RigVMDeveloper/Private/RigVMModel/RigVMClient.cpp
--- a/Engine/Plugins/Runtime/RigVM/Source/RigVMDeveloper/Private/RigVMModel/RigVMClient.cpp
+++ b/Engine/Plugins/Runtime/RigVM/Source/RigVMDeveloper/Private/RigVMModel/RigVMClient.cpp
@@ -1383,6 +1383,7 @@ FRigVMClientPatchResult FRigVMClient::PatchModelsOnLoad()
 			Result.Merge(Controller->PatchInvalidLinksOnWildcards());
 			Result.Merge(Controller->PatchExecutePins());
 			Result.Merge(Controller->PatchLazyPins());
+			Result.Merge(Controller->PatchConstantPins());
 			Result.Merge(Controller->PatchUserDefinedStructPinNames());
 			Result.Merge(Controller->PatchExternalVariablesToArguments());
 		
diff --git a/Engine/Plugins/Runtime/RigVM/Source/RigVMDeveloper/Private/RigVMModel/RigVMController.cpp b/Engine/Plugins/Runtime/RigVM/Source/RigVMDeveloper/Private/RigVMModel/RigVMController.cpp
--- a/Engine/Plugins/Runtime/RigVM/Source/RigVMDeveloper/Private/RigVMModel/RigVMController.cpp
+++ b/Engine/Plugins/Runtime/RigVM/Source/RigVMDeveloper/Private/RigVMModel/RigVMController.cpp
@@ -18841,6 +18841,12 @@ void URigVMController::AddPinsForTemplate(const FRigVMTemplate* InTemplate, cons
 		{
 			Pin->bIsLazy = true;
 		}
+		
+		if (Pin->Direction != ERigVMPinDirection::Output &&
+			!InTemplate->GetArgumentMetaData(Arg->Name, FRigVMStruct::ConstantMetaName).IsEmpty())
+		{
+			Pin->bIsConstant = true;
+		}
 
 		AddNodePin(InNode, Pin);
 
@@ -18964,6 +18970,11 @@ void URigVMController::ConfigurePinFromProperty(FProperty* InProperty, URigVMPin
 			InOutPin->bIsLazy = true;
 		}
 	}
+	
+	if (InProperty->HasMetaData(FRigVMStruct::ConstantMetaName))
+	{
+		InOutPin->bIsConstant = true;
+	}
 #endif
 
 	if (FStructProperty* StructProperty = CastField<FStructProperty>(PropertyForType))
@@ -21203,6 +21214,13 @@ void URigVMController::ResolveTemplateNodeMetaData(URigVMTemplateNode* InNode, c
 					SetPinDefaultValue(Pin, NewDefaultValue, true, bSetupUndoRedo, false);
 				}
 			}
+
+			// Update constant flag based on the resolved struct's property metadata.
+			const bool bShouldBeConstant = InNode->IsPinConstant(Pin);
+			if(Pin->bIsConstant != bShouldBeConstant)
+			{
+				Pin->bIsConstant = bShouldBeConstant;
+			}
 		}
 	}
 }
@@ -25041,6 +25059,27 @@ FRigVMClientPatchResult URigVMController::PatchLazyPins()
 	return Result;
 }
 
+FRigVMClientPatchResult URigVMController::PatchConstantPins()
+{
+	FRigVMClientPatchResult Result;
+	if (const URigVMGraph* Graph = GetGraph())
+	{
+		for (const URigVMNode* Node : Graph->GetNodes())
+		{
+			for(URigVMPin* Pin : Node->GetPins())
+			{
+				const bool bShouldBeConstant = Node->IsPinConstant(Pin);
+				if(Pin->bIsConstant != bShouldBeConstant)
+				{
+					Pin->bIsConstant = bShouldBeConstant;
+					Result.bChangedContent = true;
+				}
+			}
+		}
+	}
+	return Result;
+}
+
 FRigVMClientPatchResult URigVMController::PatchPinDefaultValues()
 {
 	FRigVMClientPatchResult Result;
diff --git a/Engine/Plugins/Runtime/RigVM/Source/RigVMDeveloper/Private/RigVMModel/RigVMPin.cpp b/Engine/Plugins/Runtime/RigVM/Source/RigVMDeveloper/Private/RigVMModel/RigVMPin.cpp
--- a/Engine/Plugins/Runtime/RigVM/Source/RigVMDeveloper/Private/RigVMModel/RigVMPin.cpp
+++ b/Engine/Plugins/Runtime/RigVM/Source/RigVMDeveloper/Private/RigVMModel/RigVMPin.cpp
@@ -2492,15 +2492,6 @@ bool URigVMPin::CanLink(const URigVMPin* InSourcePin, const URigVMPin* InTargetP
 		return false;
 	}
 
-	if (InTargetPin->IsDefinedAsConstant() && !InSourcePin->IsDefinedAsConstant())
-	{
-		if (OutFailureReason)
-		{
-			*OutFailureReason = TEXT("Cannot connect non-constants to constants.");
-		}
-		return false;
-	}
-
 	if (InSourcePin->CPPType != InTargetPin->CPPType)
 	{
 		bool bCPPTypesDiffer = true;
@@ -2910,10 +2901,17 @@ bool URigVMPin::CanLink(const URigVMPin* InSourcePin, const URigVMPin* InTargetP
 			{
 				TargetNode = TargetPin->GetNode();
 
-				// Whether this target is an IO pin or an input variable pin — both represent
+				// If the target pin is const, the connection is fine.
+				// However, whether this target is an IO pin or an input variable pin — both represent
 				// mutable references that a pure function must not be allowed to write through.
 				const bool bIsIOPin = (TargetPin->GetDirection() == ERigVMPinDirection::IO);
 				const bool bIsInputVariablePin = TargetPin->IsDefinedAsInputVariable();
+				const bool bIsConstPin = TargetPin->IsDefinedAsConstant();
+				
+				if (bIsConstPin)
+				{
+					continue;
+				}
 
 				if (bIsIOPin || bIsInputVariablePin)
 				{
diff --git a/Engine/Plugins/Runtime/RigVM/Source/RigVMDeveloper/Public/RigVMModel/Nodes/RigVMDispatchNode.h b/Engine/Plugins/Runtime/RigVM/Source/RigVMDeveloper/Public/RigVMModel/Nodes/RigVMDispatchNode.h
--- a/Engine/Plugins/Runtime/RigVM/Source/RigVMDeveloper/Public/RigVMModel/Nodes/RigVMDispatchNode.h
+++ b/Engine/Plugins/Runtime/RigVM/Source/RigVMDeveloper/Public/RigVMModel/Nodes/RigVMDispatchNode.h
@@ -84,6 +84,7 @@ protected:
 	UE_API const UScriptStruct* GetFactoryStruct() const;
 	UE_API virtual void InvalidateCache() override;
 	UE_API virtual bool ShouldInputPinComputeLazily(const URigVMPin* InPin) const override;
+	UE_API virtual bool IsPinConstant(const URigVMPin* InPin) const override;
 	UE_API virtual FString GetOriginalDefaultValueForRootPin(const URigVMPin* InRootPin) const override;
 	UE_API FString GetFactoryDefaultValue() const;
 
diff --git a/Engine/Plugins/Runtime/RigVM/Source/RigVMDeveloper/Public/RigVMModel/Nodes/RigVMUnitNode.h b/Engine/Plugins/Runtime/RigVM/Source/RigVMDeveloper/Public/RigVMModel/Nodes/RigVMUnitNode.h
--- a/Engine/Plugins/Runtime/RigVM/Source/RigVMDeveloper/Public/RigVMModel/Nodes/RigVMUnitNode.h
+++ b/Engine/Plugins/Runtime/RigVM/Source/RigVMDeveloper/Public/RigVMModel/Nodes/RigVMUnitNode.h
@@ -105,6 +105,7 @@ protected:
 
 	UE_API virtual FText GetToolTipTextForPin(const URigVMPin* InPin) const override;
 	UE_API virtual bool ShouldInputPinComputeLazily(const URigVMPin* InPin) const override;
+	UE_API virtual bool IsPinConstant(const URigVMPin* InPin) const override;
 	UE_API virtual FString GetOriginalDefaultValueForRootPin(const URigVMPin* InRootPin) const override;
 	UE_API void EnumeratePropertiesOnHostAndStructInstance(
 		URigVMHost* InHost,
diff --git a/Engine/Plugins/Runtime/RigVM/Source/RigVMDeveloper/Public/RigVMModel/RigVMController.h b/Engine/Plugins/Runtime/RigVM/Source/RigVMDeveloper/Public/RigVMModel/RigVMController.h
--- a/Engine/Plugins/Runtime/RigVM/Source/RigVMDeveloper/Public/RigVMModel/RigVMController.h
+++ b/Engine/Plugins/Runtime/RigVM/Source/RigVMDeveloper/Public/RigVMModel/RigVMController.h
@@ -1606,6 +1606,7 @@ protected:
 	UE_API FRigVMClientPatchResult PatchFunctionsWithInvalidReturnPaths();
 	UE_API FRigVMClientPatchResult PatchExecutePins();
 	UE_API FRigVMClientPatchResult PatchLazyPins();
+	UE_API FRigVMClientPatchResult PatchConstantPins();
 	UE_API FRigVMClientPatchResult PatchPinDefaultValues();
 	UE_API FRigVMClientPatchResult PatchUserDefinedStructPinNames();
 	UE_API FRigVMClientPatchResult PatchLocalVariableTypes();
diff --git a/Engine/Plugins/Runtime/RigVM/Source/RigVMDeveloper/Public/RigVMModel/RigVMNode.h b/Engine/Plugins/Runtime/RigVM/Source/RigVMDeveloper/Public/RigVMModel/RigVMNode.h
--- a/Engine/Plugins/Runtime/RigVM/Source/RigVMDeveloper/Public/RigVMModel/RigVMNode.h
+++ b/Engine/Plugins/Runtime/RigVM/Source/RigVMDeveloper/Public/RigVMModel/RigVMNode.h
@@ -470,6 +470,7 @@ protected:
 	UE_API virtual FText GetToolTipTextForPin(const URigVMPin* InPin) const;
 	virtual bool AllowsLinksOn(const URigVMPin* InPin) const { return true; }
 	virtual bool ShouldInputPinComputeLazily(const URigVMPin* InPin) const { return false; }
+	virtual bool IsPinConstant(const URigVMPin* InPin) const { return false; }
 	UE_API virtual FString GetOriginalDefaultValueForRootPin(const URigVMPin* InRootPin) const;
 
 	UE_API void UpdateTraitRootPinNames();
