RigVM compiler: swap ensure condition order in TraverseNoOp to match SetupMemory phase

When VariableOperand is invalid, the SetupMemory phase silently skips operand
creation. The BuildInstructions phase must also skip rather than firing
ensure(RootDefaultValueOperand.IsValid()). Short-circuit evaluation now skips
the ensure when VariableOperand is invalid.

#Synced-CL 52638266

diff --git a/Engine/Plugins/Runtime/RigVM/Source/RigVMDeveloper/Private/RigVMCompiler/RigVMCompiler.cpp b/Engine/Plugins/Runtime/RigVM/Source/RigVMDeveloper/Private/RigVMCompiler/RigVMCompiler.cpp
--- a/Engine/Plugins/Runtime/RigVM/Source/RigVMDeveloper/Private/RigVMCompiler/RigVMCompiler.cpp
+++ b/Engine/Plugins/Runtime/RigVM/Source/RigVMDeveloper/Private/RigVMCompiler/RigVMCompiler.cpp
@@ -4480,7 +4480,7 @@ bool URigVMCompiler::TraverseNoOp(const FRigVMNoOpExprAST* InExpr, FRigVMCompile
 							}
 						}
 					}
-					else if (ensure(RootDefaultValueOperand.IsValid()) && VariableOperand.IsValid())
+					else if (VariableOperand.IsValid() && ensure(RootDefaultValueOperand.IsValid()))
 					{
 						for (const FString& PropertyPath : PropertyPathsRequiringCopy)
 						{
