Control Rig: Fix BP independent rigs not handling renaming of assets correctly in relation to function identifiers
#rnx

#Synced-CL 52619015

diff --git a/Engine/Plugins/Runtime/RigVM/Source/RigVM/Private/RigVMRuntimeAsset.cpp b/Engine/Plugins/Runtime/RigVM/Source/RigVM/Private/RigVMRuntimeAsset.cpp
--- a/Engine/Plugins/Runtime/RigVM/Source/RigVM/Private/RigVMRuntimeAsset.cpp
+++ b/Engine/Plugins/Runtime/RigVM/Source/RigVM/Private/RigVMRuntimeAsset.cpp
@@ -3,6 +3,7 @@
 #include "RigVMRuntimeAsset.h"
 
 #include "RigVMHost.h"
+#include "RigVMCore/RigVMMemoryStorage.h"
 #include "Misc/StringOutputDevice.h"
 #include "UObject/ObjectPtr.h"
 #include "UObject/ObjectSaveContext.h"
@@ -140,6 +141,42 @@ void URigVMRuntimeAsset::PostDuplicate(bool bDuplicateForPIE)
 #endif
 }
 
+void URigVMRuntimeAsset::PostRename(UObject* OldOuter, const FName OldName)
+{
+	Super::PostRename(OldOuter, OldName);
+#if WITH_EDITORONLY_DATA
+	// Whenever the asset is renamed/moved, generated classes parented to the old package
+	// are not moved to the new package automatically (see FAssetRenameManager), so we
+	// have to manually perform the move/rename, to avoid invalid reference to the old package
+	
+	// Note: while asset duplication doesn't duplicate the classes either, it is not a problem there
+	// because we always recompile in post duplicate.
+	TArray<UObject*> Objects;
+	GetObjectsWithOuter(OldOuter->GetPackage(), Objects, EGetObjectsFlags::None);
+	for (UObject* Object : Objects)
+	{
+		if (URigVMMemoryStorageGeneratorClass* MemoryClass = Cast<URigVMMemoryStorageGeneratorClass>(Object))
+		{
+			MemoryClass->Rename(nullptr, GetPackage(), REN_DoNotDirty | REN_DontCreateRedirectors | REN_NonTransactional | REN_AllowPackageLinkerMismatch);
+		}
+	}
+
+	// Store the RuntimeAsset's old/new paths for EditorAsset to use in its PostRename.
+	// This is needed because by the time EditorAsset::PostRename is called, we've lost
+	// access to the old package path (OldOuter for EditorAsset will be this RuntimeAsset,
+	// which has already been moved to the new location).
+	PreRenameOldAssetPath = FString::Printf(TEXT("%s.%s"), *OldOuter->GetPathName(), *OldName.ToString());
+	PreRenameNewAssetPath = GetPathName();
+
+	const FString NewEditorAssetName = FString::Printf(TEXT("%s_EditorOnly"), *GetName());
+	EditorAsset->Rename(*NewEditorAssetName, this, REN_DoNotDirty | REN_DontCreateRedirectors | REN_NonTransactional | REN_AllowPackageLinkerMismatch);
+
+	// Clear after EditorAsset has used them
+	PreRenameOldAssetPath.Empty();
+	PreRenameNewAssetPath.Empty();
+#endif
+}
+
 void URigVMRuntimeAsset::MarkInstancesAsGarbage()
 {
 	// Clean up archetype instances before destroying the asset
diff --git a/Engine/Plugins/Runtime/RigVM/Source/RigVM/Public/RigVMRuntimeAsset.h b/Engine/Plugins/Runtime/RigVM/Source/RigVM/Public/RigVMRuntimeAsset.h
--- a/Engine/Plugins/Runtime/RigVM/Source/RigVM/Public/RigVMRuntimeAsset.h
+++ b/Engine/Plugins/Runtime/RigVM/Source/RigVM/Public/RigVMRuntimeAsset.h
@@ -111,6 +111,7 @@ protected:
 	UE_API virtual void PostLoad() override;
 	UE_API virtual void PreDuplicate(FObjectDuplicationParameters& DupParams) override;
 	UE_API virtual void PostDuplicate(bool bDuplicateForPIE) override;
+	UE_API virtual void PostRename(UObject* OldOuter, const FName OldName) override;
 	
 	UE_API virtual void MarkInstancesAsGarbage();
 
@@ -194,6 +195,11 @@ public:
 #if WITH_EDITORONLY_DATA
 	virtual UObject* GetEditorOnlyData() override { return EditorAsset.Get();}
 	virtual const UObject* GetEditorOnlyData() const override { return EditorAsset.Get();}
+
+	// Transient storage for pre-rename asset paths, used by EditorAsset::PostRename
+	// These are set before EditorAsset->Rename() and cleared after
+	FString PreRenameOldAssetPath;
+	FString PreRenameNewAssetPath;
 #endif
 	
 #if WITH_EDITOR
diff --git a/Engine/Plugins/Runtime/RigVM/Source/RigVMDeveloper/Private/RigVMBlueprintLegacy.cpp b/Engine/Plugins/Runtime/RigVM/Source/RigVMDeveloper/Private/RigVMBlueprintLegacy.cpp
--- a/Engine/Plugins/Runtime/RigVM/Source/RigVMDeveloper/Private/RigVMBlueprintLegacy.cpp
+++ b/Engine/Plugins/Runtime/RigVM/Source/RigVMDeveloper/Private/RigVMBlueprintLegacy.cpp
@@ -661,7 +661,32 @@ void URigVMBlueprint::PostEditChangeChainProperty(struct FPropertyChangedChainEv
 void URigVMBlueprint::PostRename(UObject* OldOuter, const FName OldName)
 {
 	UBlueprint::PostRename(OldOuter, OldName);
-	IRigVMEditorAssetInterface::PostRename(OldOuter, OldName);
+
+	// Whenever the asset is renamed/moved, generated classes parented to the old package
+	// are not moved to the new package automatically (see FAssetRenameManager), so we
+	// have to manually perform the move/rename, to avoid invalid reference to the old package
+	
+	// Note: while asset duplication doesn't duplicate the classes either, it is not a problem there
+	// because we always recompile in post duplicate.
+	TArray<UObject*> Objects;
+	GetObjectsWithOuter(OldOuter->GetPackage(), Objects, EGetObjectsFlags::None);
+	for (UObject* Object : Objects)
+	{
+		if (URigVMMemoryStorageGeneratorClass* MemoryClass = Cast<URigVMMemoryStorageGeneratorClass>(Object))
+		{
+			MemoryClass->Rename(nullptr, GetPackage(), REN_DoNotDirty | REN_DontCreateRedirectors | REN_NonTransactional | REN_AllowPackageLinkerMismatch);
+		}
+	}
+
+	// Update function identifiers to reflect the new asset path
+	const FString NewAssetPath = FString::Printf(TEXT("%s.%s"), *GetOuter()->GetPathName(), *GetFName().ToString());
+	const FString OldAssetPath = FString::Printf(TEXT("%s.%s"), *OldOuter->GetPathName(), *OldName.ToString());
+
+	const FString NewFunctionHostPath = FSoftObjectPath(GetRigVMGraphFunctionHost().GetObject()).ToString();
+	const FString NewFunctionLibraryPath = FSoftObjectPath(GetRigVMClient()->GetFunctionLibrary()).GetSubPathString();
+	const FString OldFunctionHostPath = NewFunctionHostPath.Replace(*NewAssetPath, *OldAssetPath);
+	const FString OldFunctionLibraryPath = NewFunctionLibraryPath.Replace(*NewAssetPath, *OldAssetPath);
+	ReplaceFunctionIdentifiers(OldFunctionHostPath, NewFunctionHostPath, OldFunctionLibraryPath, NewFunctionLibraryPath);
 }
 
 void URigVMBlueprint::GetPreloadDependencies(TArray<UObject*>& OutDeps)
diff --git a/Engine/Plugins/Runtime/RigVM/Source/RigVMDeveloper/Private/RigVMEditorAsset.cpp b/Engine/Plugins/Runtime/RigVM/Source/RigVMDeveloper/Private/RigVMEditorAsset.cpp
--- a/Engine/Plugins/Runtime/RigVM/Source/RigVMDeveloper/Private/RigVMEditorAsset.cpp
+++ b/Engine/Plugins/Runtime/RigVM/Source/RigVMDeveloper/Private/RigVMEditorAsset.cpp
@@ -574,34 +574,6 @@ void IRigVMEditorAssetInterface::PostEditChangeChainProperty(FPropertyChangedCha
 	PostEditChangeChainPropertyEvent.Broadcast(PropertyChangedEvent);
 }
 
-void IRigVMEditorAssetInterface::PostRename(UObject* OldOuter, const FName OldName)
-{
-	// Whenever the asset is renamed/moved, generated classes parented to the old package
-	// are not moved to the new package automatically (see FAssetRenameManager), so we
-	// have to manually perform the move/rename, to avoid invalid reference to the old package
-	
-	// Note: while asset duplication doesn't duplicate the classes either, it is not a problem there
-	// because we always recompile in post duplicate.
-	TArray<UClass*> ClassObjects = GetClassObjectsInPackage(OldOuter->GetPackage());
-
-	for (UClass* ClassObject : ClassObjects)
-	{
-		if (URigVMMemoryStorageGeneratorClass* MemoryClass = Cast<URigVMMemoryStorageGeneratorClass>(ClassObject))
-		{
-			MemoryClass->Rename(nullptr, GetObject()->GetPackage(), REN_DoNotDirty | REN_DontCreateRedirectors | REN_NonTransactional | REN_AllowPackageLinkerMismatch);
-		}
-	}
-
-	const FString NewAssetPath = FString::Printf(TEXT("%s.%s"), *GetObject()->GetOuter()->GetPathName(), *GetObject()->GetFName().ToString());
-	const FString OldAssetPath = FString::Printf(TEXT("%s.%s"), *OldOuter->GetPathName(), *OldName.ToString());
-	
-	const FString NewFunctionHostPath = FSoftObjectPath(GetRigVMClientHost()->GetRigVMGraphFunctionHost().GetObject()).ToString();
-	const FString NewFunctionLibraryPath = FSoftObjectPath(GetRigVMClient()->GetFunctionLibrary()).GetSubPathString();
-	const FString OldFunctionHostPath = NewFunctionHostPath.Replace(*NewAssetPath, *OldAssetPath);
-	const FString OldFunctionLibraryPath = NewFunctionLibraryPath.Replace(*NewAssetPath, *OldAssetPath);
-	ReplaceFunctionIdentifiers(OldFunctionHostPath, NewFunctionHostPath, OldFunctionLibraryPath, NewFunctionLibraryPath);
-}
-
 void IRigVMEditorAssetInterface::GetPreloadDependencies(TArray<UObject*>& OutDeps)
 {
 	TArray<UClass*> ClassObjects = GetClassObjectsInPackage(GetObject()->GetPackage());
@@ -4521,6 +4493,41 @@ void URigVMEditorAsset::PostDuplicate(bool bDuplicateForPIE)
 	return IRigVMEditorAssetInterface::PostDuplicate(bDuplicateForPIE);
 }
 
+void URigVMEditorAsset::PostRename(UObject* OldOuter, const FName OldName)
+{
+	Super::PostRename(OldOuter, OldName);
+
+	// For BP-independent assets, the RuntimeAsset stores the correct old/new paths
+	// before calling EditorAsset->Rename(). We use those paths here because by this point
+	// OldOuter (the RuntimeAsset) has already been moved to its new location, so we can't
+	// reconstruct the old package path from it.
+	URigVMRuntimeAsset* RuntimeAsset = Cast<URigVMRuntimeAsset>(GetOuter());
+	check(RuntimeAsset);
+
+	const FString& OldAssetPath = RuntimeAsset->PreRenameOldAssetPath;
+	const FString& NewAssetPath = RuntimeAsset->PreRenameNewAssetPath;
+
+	if (OldAssetPath.IsEmpty() || NewAssetPath.IsEmpty())
+	{
+		UE_LOG(LogRigVMDeveloper, Warning,
+			TEXT("URigVMEditorAsset::PostRename called without pre-rename paths set. "
+				 "Function identifiers will not be updated. This is expected only if "
+				 "the rename was not initiated through URigVMRuntimeAsset."));
+		return;
+	}
+
+	// Function identifier replacement using full object paths
+	// - Function host path: replace RuntimeAsset path
+	// - Function library path: replace RuntimeAsset path AND EditorAsset name (subobject after the colon)
+	const FString NewFunctionHostPath = FSoftObjectPath(GetRigVMClientHost()->GetRigVMGraphFunctionHost().GetObject()).ToString();
+	const FString NewFunctionLibraryPath = FSoftObjectPath(GetRigVMClient()->GetFunctionLibrary()).ToString();
+	const FString OldFunctionHostPath = NewFunctionHostPath.Replace(*NewAssetPath, *OldAssetPath);
+	const FString OldFunctionLibraryPath = NewFunctionLibraryPath
+		.Replace(*NewAssetPath, *OldAssetPath)
+		.Replace(*GetFName().ToString(), *OldName.ToString());
+	ReplaceFunctionIdentifiers(OldFunctionHostPath, NewFunctionHostPath, OldFunctionLibraryPath, NewFunctionLibraryPath);
+}
+
 void URigVMEditorAsset::BeginDestroy()
 {
 	Super::BeginDestroy();
diff --git a/Engine/Plugins/Runtime/RigVM/Source/RigVMDeveloper/Public/RigVMEditorAsset.h b/Engine/Plugins/Runtime/RigVM/Source/RigVMDeveloper/Public/RigVMEditorAsset.h
--- a/Engine/Plugins/Runtime/RigVM/Source/RigVMDeveloper/Public/RigVMEditorAsset.h
+++ b/Engine/Plugins/Runtime/RigVM/Source/RigVMDeveloper/Public/RigVMEditorAsset.h
@@ -500,7 +500,6 @@ protected:
 
 	// UObject interface
 	UE_API virtual void PostEditChangeChainProperty(struct FPropertyChangedChainEvent& PropertyChangedEvent);
-	UE_API virtual void PostRename(UObject* OldOuter, const FName OldName);
 	/** Called during cooking. Must return all objects that will be Preload()ed when this is serialized at load time. */
 	UE_API virtual void GetPreloadDependencies(TArray<UObject*>& OutDeps);
 
@@ -838,6 +837,7 @@ class URigVMEditorAsset : public UObject, public IRigVMEditorAssetInterface, pub
 	UE_API virtual void PostTransacted(const FTransactionObjectEvent& TransactionEvent) override;
 	UE_API virtual void PreDuplicate(FObjectDuplicationParameters& DupParams) override;
 	UE_API virtual void PostDuplicate(bool bDuplicateForPIE) override;
+	UE_API virtual void PostRename(UObject* OldOuter, const FName OldName) override;
 	UE_API virtual void BeginDestroy() override;
 
 #if WITH_EDITORONLY_DATA
