Control Rig: Redirect open editor for both editor and runtime assets to open the runtime asset
#rnx

#Synced-CL 52681405

diff --git a/Engine/Plugins/Animation/ControlRig/Source/ControlRigEditor/Private/ControlRigEditorAssetActions.cpp b/Engine/Plugins/Animation/ControlRig/Source/ControlRigEditor/Private/ControlRigEditorAssetActions.cpp
new file mode 100644
--- /dev/null
+++ b/Engine/Plugins/Animation/ControlRig/Source/ControlRigEditor/Private/ControlRigEditorAssetActions.cpp
@@ -0,0 +1,33 @@
+// Copyright Epic Games, Inc. All Rights Reserved.
+
+#include "ControlRigEditorAssetActions.h"
+#include "ControlRigEditorAsset.h"
+#include "ControlRigRuntimeAsset.h"
+#include "Editor.h"
+#include "Subsystems/AssetEditorSubsystem.h"
+
+UClass* FControlRigEditorAssetActions::GetSupportedClass() const
+{
+	return UControlRigEditorAsset::StaticClass();
+}
+
+void FControlRigEditorAssetActions::OpenAssetEditor(const TArray<UObject*>& InObjects, TSharedPtr<IToolkitHost> EditWithinLevelEditor)
+{
+	if (IsEngineExitRequested())
+	{
+		return;
+	}
+
+	// Redirect to the runtime asset instead of opening the editor asset directly
+	for (UObject* Obj : InObjects)
+	{
+		if (UControlRigEditorAsset* EditorAsset = Cast<UControlRigEditorAsset>(Obj))
+		{
+			// The runtime asset is the outer of the editor asset
+			if (UControlRigRuntimeAsset* RuntimeAsset = EditorAsset->GetTypedOuter<UControlRigRuntimeAsset>())
+			{
+				GEditor->GetEditorSubsystem<UAssetEditorSubsystem>()->OpenEditorForAsset(RuntimeAsset);
+			}
+		}
+	}
+}
diff --git a/Engine/Plugins/Animation/ControlRig/Source/ControlRigEditor/Private/ControlRigEditorAssetActions.h b/Engine/Plugins/Animation/ControlRig/Source/ControlRigEditor/Private/ControlRigEditorAssetActions.h
new file mode 100644
--- /dev/null
+++ b/Engine/Plugins/Animation/ControlRig/Source/ControlRigEditor/Private/ControlRigEditorAssetActions.h
@@ -0,0 +1,23 @@
+// Copyright Epic Games, Inc. All Rights Reserved.
+
+#pragma once
+
+#include "CoreMinimal.h"
+#include "AssetTypeActions_Base.h"
+
+/**
+ * Asset type actions for UControlRigEditorAsset that redirects to the runtime asset.
+ * When users try to open an EditorAsset directly, we redirect them to open the
+ * RuntimeAsset instead, which is the outer object that contains the editor data.
+ */
+class FControlRigEditorAssetActions : public FAssetTypeActions_Base
+{
+public:
+	// IAssetTypeActions Implementation
+	virtual FText GetName() const override { return NSLOCTEXT("AssetTypeActions", "ControlRigEditorAssetActions", "Control Rig Editor Data"); }
+	virtual FColor GetTypeColor() const override { return FColor(140, 116, 0); }
+	virtual UClass* GetSupportedClass() const override;
+	virtual void OpenAssetEditor(const TArray<UObject*>& InObjects, TSharedPtr<class IToolkitHost> EditWithinLevelEditor = TSharedPtr<IToolkitHost>()) override;
+	virtual uint32 GetCategories() override { return EAssetTypeCategories::Animation; }
+	virtual bool CanFilter() override { return false; }
+};
diff --git a/Engine/Plugins/Animation/ControlRig/Source/ControlRigEditor/Private/ControlRigEditorModule.cpp b/Engine/Plugins/Animation/ControlRig/Source/ControlRigEditor/Private/ControlRigEditorModule.cpp
--- a/Engine/Plugins/Animation/ControlRig/Source/ControlRigEditor/Private/ControlRigEditorModule.cpp
+++ b/Engine/Plugins/Animation/ControlRig/Source/ControlRigEditor/Private/ControlRigEditorModule.cpp
@@ -109,6 +109,7 @@
 #include "Editor/ControlRigLegacyEditor.h"
 #endif
 #include "ControlRigAssetActions.h"
+#include "ControlRigEditorAssetActions.h"
 #include "ControlRigRuntimeAsset.h"
 #include "ControlRigAssetReferenceDetails.h"
 #include "Editor/ControlRigNewEditor.h"
@@ -211,6 +212,7 @@ void FControlRigEditorModule::StartupModule()
 
 	RegisterAssetTypeAction(MakeShareable(new FControlRigBlueprintActions()));
 	RegisterAssetTypeAction(MakeShareable(new FControlRigAssetActions()));
+	RegisterAssetTypeAction(MakeShareable(new FControlRigEditorAssetActions()));
 	RegisterAssetTypeAction(MakeShareable(new FControlRigShapeLibraryActions()));
 
 	// Register sequencer track editor
diff --git a/Engine/Plugins/Animation/ControlRig/Source/ControlRigEditor/Private/Editor/ControlRigEditor.cpp b/Engine/Plugins/Animation/ControlRig/Source/ControlRigEditor/Private/Editor/ControlRigEditor.cpp
--- a/Engine/Plugins/Animation/ControlRig/Source/ControlRigEditor/Private/Editor/ControlRigEditor.cpp
+++ b/Engine/Plugins/Animation/ControlRig/Source/ControlRigEditor/Private/Editor/ControlRigEditor.cpp
@@ -2936,20 +2936,28 @@ void FControlRigBaseEditor::HandleViewportCreated(const TSharedRef<class IPerson
 	FEditorViewportClient& ViewportClient = InViewport->GetViewportClient();
 	if (FAnimationViewportClient* AnimViewportClient = static_cast<FAnimationViewportClient*>(&ViewportClient))
 	{
-		AnimViewportClient->OnSetBoneSize.BindLambda([this](float InBoneSize)
+		
+		TWeakPtr<FControlRigBaseEditor> WeakThis = SharedRef();
+		AnimViewportClient->OnSetBoneSize.BindLambda([WeakThis](float InBoneSize)
 		{
-			if (FControlRigAssetInterfacePtr RigBlueprint = GetControlRigAssetInterface())
+			if (TSharedPtr<FControlRigBaseEditor> PinnedThis = WeakThis.Pin())
 			{
-				RigBlueprint->Modify();
-				RigBlueprint->GetDebugBoneRadius() = InBoneSize;
+				if (FControlRigAssetInterfacePtr RigBlueprint = PinnedThis->GetControlRigAssetInterface())
+				{
+					RigBlueprint->Modify();
+					RigBlueprint->GetDebugBoneRadius() = InBoneSize;
+				}
 			}
 		});
 		
-		AnimViewportClient->OnGetBoneSize.BindLambda([this]() -> float
+		AnimViewportClient->OnGetBoneSize.BindLambda([WeakThis]() -> float
 		{
-			if (FControlRigAssetInterfacePtr RigBlueprint = GetControlRigAssetInterface())
+			if (TSharedPtr<FControlRigBaseEditor> PinnedThis = WeakThis.Pin())
 			{
-				return RigBlueprint->GetDebugBoneRadius();
+				if (FControlRigAssetInterfacePtr RigBlueprint = PinnedThis->GetControlRigAssetInterface())
+				{
+					return RigBlueprint->GetDebugBoneRadius();
+				}
 			}
 
 			return 1.0f;
diff --git a/Engine/Plugins/Runtime/RigVM/Source/RigVMEditor/Private/RigVMEditorAssetActions.cpp b/Engine/Plugins/Runtime/RigVM/Source/RigVMEditor/Private/RigVMEditorAssetActions.cpp
new file mode 100644
--- /dev/null
+++ b/Engine/Plugins/Runtime/RigVM/Source/RigVMEditor/Private/RigVMEditorAssetActions.cpp
@@ -0,0 +1,34 @@
+// Copyright Epic Games, Inc. All Rights Reserved.
+
+#include "RigVMEditorAssetActions.h"
+#include "RigVMEditorAsset.h"
+#include "RigVMRuntimeAsset.h"
+#include "Editor.h"
+#include "Subsystems/AssetEditorSubsystem.h"
+
+UClass* FRigVMEditorAssetActions::GetSupportedClass() const
+{
+	return URigVMEditorAsset::StaticClass();
+}
+
+void FRigVMEditorAssetActions::OpenAssetEditor(const TArray<UObject*>& InObjects, TSharedPtr<IToolkitHost> EditWithinLevelEditor)
+{
+	if (IsEngineExitRequested())
+	{
+		return;
+	}
+
+	// Redirect to the runtime asset instead of opening the editor asset directly
+	for (UObject* Obj : InObjects)
+	{
+		if (URigVMEditorAsset* EditorAsset = Cast<URigVMEditorAsset>(Obj))
+		{
+			// The runtime asset is the outer of the editor asset
+			if (URigVMRuntimeAsset* RuntimeAsset = EditorAsset->GetTypedOuter<URigVMRuntimeAsset>())
+			{
+				GEditor->GetEditorSubsystem<UAssetEditorSubsystem>()->OpenEditorForAsset(RuntimeAsset);
+			}
+		}
+	}
+}
+
diff --git a/Engine/Plugins/Runtime/RigVM/Source/RigVMEditor/Private/RigVMEditorAssetActions.h b/Engine/Plugins/Runtime/RigVM/Source/RigVMEditor/Private/RigVMEditorAssetActions.h
new file mode 100644
--- /dev/null
+++ b/Engine/Plugins/Runtime/RigVM/Source/RigVMEditor/Private/RigVMEditorAssetActions.h
@@ -0,0 +1,23 @@
+// Copyright Epic Games, Inc. All Rights Reserved.
+
+#pragma once
+
+#include "CoreMinimal.h"
+#include "AssetTypeActions_Base.h"
+
+/**
+ * Asset type actions for URigVMEditorAsset that redirects to the runtime asset.
+ * When users try to open an EditorAsset directly, we redirect them to open the
+ * RuntimeAsset instead, which is the outer object that contains the editor data.
+ */
+class FRigVMEditorAssetActions : public FAssetTypeActions_Base
+{
+public:
+	// IAssetTypeActions Implementation
+	virtual FText GetName() const override { return NSLOCTEXT("AssetTypeActions", "RigVMEditorAssetActions", "RigVM Editor Data"); }
+	virtual FColor GetTypeColor() const override { return FColor(201, 29, 85); }
+	virtual UClass* GetSupportedClass() const override;
+	virtual void OpenAssetEditor(const TArray<UObject*>& InObjects, TSharedPtr<class IToolkitHost> EditWithinLevelEditor = TSharedPtr<IToolkitHost>()) override;
+	virtual uint32 GetCategories() override { return EAssetTypeCategories::Animation; }
+	virtual bool CanFilter() override { return false; }
+};
diff --git a/Engine/Plugins/Runtime/RigVM/Source/RigVMEditor/Private/RigVMEditorModule.cpp b/Engine/Plugins/Runtime/RigVM/Source/RigVMEditor/Private/RigVMEditorModule.cpp
--- a/Engine/Plugins/Runtime/RigVM/Source/RigVMEditor/Private/RigVMEditorModule.cpp
+++ b/Engine/Plugins/Runtime/RigVM/Source/RigVMEditor/Private/RigVMEditorModule.cpp
@@ -5,6 +5,7 @@
 =============================================================================*/
 
 #include "RigVMEditorModule.h"
+#include "RigVMEditorAssetActions.h"
 #include "AssetToolsModule.h"
 #include "Modules/ModuleManager.h"
 #include "AssetRegistry/AssetRegistryModule.h"
@@ -169,6 +170,12 @@ void FRigVMEditorModule::StartupModule()
 		IModularFeatures::Get().RegisterModularFeature(TraceServices::ModuleFeatureName, &RigVMTraceModule);
 		IModularFeatures::Get().RegisterModularFeature(RewindDebugger::IRewindDebuggerTrackCreator::ModularFeatureName, &GRewindDebuggerRigVMTrackCreator);
 #endif
+
+		// Register asset type actions for EditorAsset redirection
+		IAssetTools& AssetTools = FModuleManager::LoadModuleChecked<FAssetToolsModule>("AssetTools").Get();
+		TSharedRef<IAssetTypeActions> EditorAssetActions = MakeShareable(new FRigVMEditorAssetActions());
+		RegisteredAssetTypeActions.Add(EditorAssetActions);
+		AssetTools.RegisterAssetTypeActions(EditorAssetActions);
 	}
 
 	StartupModuleCommon();
@@ -227,6 +234,14 @@ void FRigVMEditorModule::ShutdownModule()
 			IModularFeatures::Get().UnregisterModularFeature(TraceServices::ModuleFeatureName, &RigVMTraceModule);
 			IModularFeatures::Get().UnregisterModularFeature(RewindDebugger::IRewindDebuggerTrackCreator::ModularFeatureName, &GRewindDebuggerRigVMTrackCreator);
 #endif
+
+			// Unregister asset type actions
+			IAssetTools& AssetTools = FModuleManager::LoadModuleChecked<FAssetToolsModule>("AssetTools").Get();
+			for (const TSharedRef<IAssetTypeActions>& Action : RegisteredAssetTypeActions)
+			{
+				AssetTools.UnregisterAssetTypeActions(Action);
+			}
+			RegisteredAssetTypeActions.Empty();
 		}
 	}
 
diff --git a/Engine/Plugins/Runtime/RigVM/Source/RigVMEditor/Public/RigVMEditorModule.h b/Engine/Plugins/Runtime/RigVM/Source/RigVMEditor/Public/RigVMEditorModule.h
--- a/Engine/Plugins/Runtime/RigVM/Source/RigVMEditor/Public/RigVMEditorModule.h
+++ b/Engine/Plugins/Runtime/RigVM/Source/RigVMEditor/Public/RigVMEditorModule.h
@@ -7,6 +7,7 @@
 #pragma once
 
 #include "CoreMinimal.h"
+#include "IAssetTypeActions.h"
 #include "RigVMBlueprintLegacy.h"
 #include "RigVMEditorAsset.h"
 #include "Modules/ModuleInterface.h"
@@ -210,6 +211,9 @@ private:
 	/** StaticStruct is not safe on shutdown, so we cache the name, and use this to unregister on shut down */
 	TArray<FName> PropertiesToUnregisterOnShutdown;
 
+	/** Registered asset type actions */
+	TArray<TSharedRef<IAssetTypeActions>> RegisteredAssetTypeActions;
+
 	UE_API void HandleNewBlueprintCreated(UBlueprint* InBlueprint);
 
 	UE_API bool ShowWorkflowOptionsDialog(URigVMUserWorkflowOptions* InOptions) const;
